#UE4:
* removed not needed divide

diff --git a/Engine/Shaders/BasePassPixelShader.usf b/Engine/Shaders/BasePassPixelShader.usf
--- a/Engine/Shaders/BasePassPixelShader.usf
+++ b/Engine/Shaders/BasePassPixelShader.usf
@@ -448,7 +448,7 @@ float ClipForEditorPrimitives(FMaterialPixelParameters MaterialParameters)
 		float PixelDeviceZ = MaterialParameters.SVPosition.z;
 
 		// Soft Bias with DeviceZ for best quality
-		const float DeviceDepthFade = 0.00005f * View.GeneralPurposeTweak;
+		const float DeviceDepthFade = 0.00005f;
 
 		// 0.5f is to bias around the actual value, 1 or 0 are another option
 		Ret = saturate(0.5f - (DeviceZ - PixelDeviceZ) / DeviceDepthFade);
