Fixed tangent space calculations to match xNormal. The tangent space matrix will now no longer normalize itself. The additional upside is 2 less normalizes.

diff --git a/Engine/Shaders/MaterialTemplate.usf b/Engine/Shaders/MaterialTemplate.usf
--- a/Engine/Shaders/MaterialTemplate.usf
+++ b/Engine/Shaders/MaterialTemplate.usf
@@ -1271,17 +1271,13 @@ void CalcMaterialParameters(
 /** Assemble the transform from tangent space into world space */
 half3x3 AssembleTangentToWorld( half3 TangentToWorld0, half4 TangentToWorld2 )
 {
+	// Will not be orthonormal after interpolation. This perfectly matches xNormal.
+	// Any mismatch with xNormal will cause distortions for baked normal maps.
+
 	half3 TangentToWorld2N = TangentToWorld2.xyz;
 	half3 TangentToWorld0N = TangentToWorld0;
 
-	#if FEATURE_LEVEL != FEATURE_LEVEL_ES2
-		// Renormalize basis vectors as they are no longer unit length due to linear interpolation
-		// ES2 will accept some error here for performance
-		TangentToWorld2N = normalize(TangentToWorld2.xyz);
-		TangentToWorld0N = normalize(TangentToWorld0);
-	#endif
-
-	// Derive the third basis vector off of the other two, guaranteed to be unit length as the other two are an orthonormal set.
+	// Derive the third basis vector off of the other two.
 	// Flip based on the determinant sign
 	half3 TangentToWorld1N = cross(TangentToWorld2N,TangentToWorld0N) * TangentToWorld2.w;
 	// Transform from tangent space to world space
