#UE4:
* fixed sRGB on rendertargets for OpenGL
* tested on D3D11 and OpenGL, if this wouldn't work the diffuse/albedo would appear to be darker
* use sRGB flag on GBufferC which is diffuse color to save shader instructions
  doesn't affect ES2

diff --git a/Engine/Shaders/DeferredShadingCommon.usf b/Engine/Shaders/DeferredShadingCommon.usf
--- a/Engine/Shaders/DeferredShadingCommon.usf
+++ b/Engine/Shaders/DeferredShadingCommon.usf
@@ -26,25 +26,27 @@ float3 DecodeSpecularColor(float3 SpecularColor)
 
 float3 EncodeDiffuseColor(float3 DiffuseColor)
 {
-	// Allocate more precision to the darks
-	return sqrt(saturate(DiffuseColor));
+	// we use sRGB on the render target to give more precision to the darks
+	return DiffuseColor;
 }
 
 float3 DecodeDiffuseColor(float3 DiffuseColor)
 {
-	return DiffuseColor * DiffuseColor;
+	// we use sRGB on the render target to give more precision to the darks
+	return DiffuseColor;
 }
 
 // TODO Can sRGB writes do this for us?
 float3 EncodeBaseColor(float3 BaseColor)
 {
-	// Gamma 2.0
-	return sqrt( saturate(BaseColor) );
+	// we use sRGB on the render target to give more precision to the darks
+	return BaseColor;
 }
 
 float3 DecodeBaseColor(float3 BaseColor)
 {
-	return Square( BaseColor );
+	// we use sRGB on the render target to give more precision to the darks
+	return BaseColor;
 }
 
 float3 EncodeSubsurfaceColor(float3 SubsurfaceColor)
diff --git a/Engine/Source/Runtime/OpenGLDrv/Private/OpenGLDrv.cpp b/Engine/Source/Runtime/OpenGLDrv/Private/OpenGLDrv.cpp
--- a/Engine/Source/Runtime/OpenGLDrv/Private/OpenGLDrv.cpp
+++ b/Engine/Source/Runtime/OpenGLDrv/Private/OpenGLDrv.cpp
@@ -494,4 +494,6 @@ void InitDefaultGLContextState(void)
 	// Intel HD4000 under <= 10.8.4 requires GL_DITHER disabled or dithering will occur on any channel < 8bits.
 	// No other driver does this but we don't need GL_DITHER on anyway.
 	glDisable(GL_DITHER);
+	// Render targets with TexCreate_SRGB should do sRGB conversion like in D3D11
+	glEnable(GL_FRAMEBUFFER_SRGB);
 }
diff --git a/Engine/Source/Runtime/Renderer/Private/PostProcess/SceneRenderTargets.cpp b/Engine/Source/Runtime/Renderer/Private/PostProcess/SceneRenderTargets.cpp
--- a/Engine/Source/Runtime/Renderer/Private/PostProcess/SceneRenderTargets.cpp
+++ b/Engine/Source/Runtime/Renderer/Private/PostProcess/SceneRenderTargets.cpp
@@ -851,7 +851,7 @@ void FSceneRenderTargets::AllocateDeferredShadingPathRenderTargets()
 		{
 			const EPixelFormat DiffuseGBufferFormat = bHighPrecisionGBuffers ? PF_FloatRGBA : PF_B8G8R8A8;
 
-			FPooledRenderTargetDesc Desc(FPooledRenderTargetDesc::Create2DDesc(BufferSize, DiffuseGBufferFormat, TexCreate_None, TexCreate_RenderTargetable, false));
+			FPooledRenderTargetDesc Desc(FPooledRenderTargetDesc::Create2DDesc(BufferSize, DiffuseGBufferFormat, TexCreate_SRGB, TexCreate_RenderTargetable, false));
 			GRenderTargetPool.FindFreeElement(Desc, GBufferC, TEXT("GBufferC"));
 		}
 
