Clean up a bit of last change.

diff --git a/Engine/Shaders/MaterialTemplate.usf b/Engine/Shaders/MaterialTemplate.usf
--- a/Engine/Shaders/MaterialTemplate.usf
+++ b/Engine/Shaders/MaterialTemplate.usf
@@ -1274,12 +1274,9 @@ half3x3 AssembleTangentToWorld( half3 TangentToWorld0, half4 TangentToWorld2 )
 	// Will not be orthonormal after interpolation. This perfectly matches xNormal.
 	// Any mismatch with xNormal will cause distortions for baked normal maps.
 
-	half3 TangentToWorld2N = TangentToWorld2.xyz;
-	half3 TangentToWorld0N = TangentToWorld0;
-
 	// Derive the third basis vector off of the other two.
 	// Flip based on the determinant sign
-	half3 TangentToWorld1N = cross(TangentToWorld2N,TangentToWorld0N) * TangentToWorld2.w;
+	half3 TangentToWorld1 = cross(TangentToWorld2.xyz,TangentToWorld0) * TangentToWorld2.w;
 	// Transform from tangent space to world space
-	return half3x3(TangentToWorld0N, TangentToWorld1N, TangentToWorld2N);
+	return half3x3(TangentToWorld0, TangentToWorld1, TangentToWorld2.xyz);
 }
