#UE4:
* minor shader optimization

diff --git a/Engine/Shaders/DeferredLightingCommon.usf b/Engine/Shaders/DeferredLightingCommon.usf
--- a/Engine/Shaders/DeferredLightingCommon.usf
+++ b/Engine/Shaders/DeferredLightingCommon.usf
@@ -199,7 +199,7 @@ float3 PointLightSubsurface( FScreenSpaceData ScreenSpaceData, int LightingModel
 		float BackScatter = InGBufferData.GBufferAO * NormalContribution / (PI * 2);
 
 		// lerp to never exceed 1 (energy conserving)
-		return InGBufferData.SubsurfaceColor * lerp(BackScatter, 1, InScatter) * SubsurfaceExtinction;
+		return InGBufferData.SubsurfaceColor * (lerp(BackScatter, 1, InScatter) * SubsurfaceExtinction);
 	}
 	else if (LightingModelId == LIGHTINGMODELID_PREINTEGRATED_SKIN)
 	{
