#UE4 #hlslcc statics in shaders must now be consts!

diff --git a/Engine/Shaders/HeightFogCommon.usf b/Engine/Shaders/HeightFogCommon.usf
--- a/Engine/Shaders/HeightFogCommon.usf
+++ b/Engine/Shaders/HeightFogCommon.usf
@@ -15,10 +15,10 @@ float4 DirectionalInscatteringColor;
 float DirectionalInscatteringStartDistance;
 
 /** FogDensity * exp2(-FogHeightFalloff * (CameraWorldPosition.z - FogHeight)) in x, FogHeightFalloff in y, CosTerminatorAngle in z, StartDistance in w. */
-static float4 ExponentialFogParameters = SharedFogParameter0;
+static const float4 ExponentialFogParameters = SharedFogParameter0;
 /** Color to use. */
-static float3 ExponentialFogColor = SharedFogParameter1.xyz;
-static float MinFogOpacity = SharedFogParameter1.w;
+static const float3 ExponentialFogColor = SharedFogParameter1.xyz;
+static const float MinFogOpacity = SharedFogParameter1.w;
 
 half4 GetExponentialHeightFog(float3 WorldPosition, float3 InCameraPosition)
 {
diff --git a/Engine/Shaders/PostProcessAmbientOcclusion.usf b/Engine/Shaders/PostProcessAmbientOcclusion.usf
--- a/Engine/Shaders/PostProcessAmbientOcclusion.usf
+++ b/Engine/Shaders/PostProcessAmbientOcclusion.usf
@@ -34,7 +34,7 @@ float3 ReconstructCSPos(float SceneDepth, float2 ScreenPos)
 // could be moved to a more central spot
 float2 ReconstructSSPosFromCS(float3 In)
 {
-    return In.xy / In.z;
+	return In.xy / In.z;
 }
 
 // could be moved to a more central spot
@@ -74,7 +74,7 @@ float4 NoiseScale;
 
 #if USE_SAMPLESET == 0
 	#define SAMPLESET_ARRAY_SIZE 6
-	static float2 OcclusionSamplesOffsets[SAMPLESET_ARRAY_SIZE]=
+	static const float2 OcclusionSamplesOffsets[SAMPLESET_ARRAY_SIZE]=
 	{
 		// 6 points distributed on the unit disc, spiral order and distance
 		float2(0.000, 0.200), 
@@ -86,7 +86,7 @@ float4 NoiseScale;
 	};
 #else // USE_SAMPLESET == 0
 	#define SAMPLESET_ARRAY_SIZE 5
-	static float2 OcclusionSamplesOffsets[SAMPLESET_ARRAY_SIZE]=
+	static const float2 OcclusionSamplesOffsets[SAMPLESET_ARRAY_SIZE]=
 	{
 		// 5 points distributed on a unit circle
 		float2(0.156434, 0.987688),
@@ -383,7 +383,7 @@ void MainPS(in float4 UVAndScreenPos : TEXCOORD0, out float4 OutColor : SV_Targe
 
 			SamplePosition0 = ReconstructCSPos(leftAbs, ScreenSpacePosLeft) * InvFovFix;
 			SamplePosition1 = ReconstructCSPos(rightAbs, ScreenSpacePosRight) * InvFovFix;
-	    
+		
 			float3 DeltaLeft = SamplePosition0 - ViewSpacePosition;
 			float3 DeltaRight = SamplePosition1 - ViewSpacePosition;
 
diff --git a/Engine/Shaders/ShadowProjectionCommon.usf b/Engine/Shaders/ShadowProjectionCommon.usf
--- a/Engine/Shaders/ShadowProjectionCommon.usf
+++ b/Engine/Shaders/ShadowProjectionCommon.usf
@@ -225,8 +225,8 @@ Texture2D TranslucencyShadowTransmission1;
 SamplerState TranslucencyShadowTransmission1Sampler;
 
 /** These have been tweaked to minimize ringing while not losing too much high frequency in the captured opacity function. */
-static float RingingSuppressionFactor = 1;
-static float NumTermsForRingingSuppression = 8;
+static const float RingingSuppressionFactor = 1;
+static const float NumTermsForRingingSuppression = 8;
 
 // Creates a scale for a given frequency that minimizes ringing inherent with a Fourier basis
 // This is done by scaling down the contribution of the higher frequency components
