#UE4:
* renamed function to be more clear
* fixed shader warning about same variable declared 2 times (scope issue)

diff --git a/Engine/Shaders/DeferredShadingCommon.usf b/Engine/Shaders/DeferredShadingCommon.usf
--- a/Engine/Shaders/DeferredShadingCommon.usf
+++ b/Engine/Shaders/DeferredShadingCommon.usf
@@ -432,7 +432,7 @@ SamplerState CustomDepthTextureSampler;
 	Texture2D<float> CustomDepthTextureNonMS;
 
 	// @param PixelPos - integer pixel pos (from left top)
-	FGBufferData GetGBufferData(uint2 PixelPos, bool bGetNormalizedNormal = true)
+	FGBufferData GetGBufferDataUint(uint2 PixelPos, bool bGetNormalizedNormal = true)
 	{
 		float4 GBufferA = GBufferATextureNonMS.Load(int3(PixelPos, 0));
 		float4 GBufferB = GBufferBTextureNonMS.Load(int3(PixelPos, 0));
@@ -461,11 +461,11 @@ SamplerState CustomDepthTextureSampler;
 	}
 
 	// @param PixelPos - integer pixel pos (from left top)
-	FScreenSpaceData GetScreenSpaceData(uint2 PixelPos, bool bGetNormalizedNormal = true)
+	FScreenSpaceData GetScreenSpaceDataUint(uint2 PixelPos, bool bGetNormalizedNormal = true)
 	{
 		FScreenSpaceData Out;
-	
-		Out.GBuffer = GetGBufferData(PixelPos, bGetNormalizedNormal);
+
+		Out.GBuffer = GetGBufferDataUint(PixelPos, bGetNormalizedNormal);
 
 		// todo: optimize
 		// this is what we want but because WhiteDummy (in case AO is disabled) doesn't support this lookup we do the code below
@@ -480,7 +480,7 @@ SamplerState CustomDepthTextureSampler;
 
 			Out.AmbientOcclusion = ScreenSpaceAO.r;
 		}
-	
+
 		return Out;
 	}
 #endif
diff --git a/Engine/Shaders/ReflectionEnvironmentComputeShaders.usf b/Engine/Shaders/ReflectionEnvironmentComputeShaders.usf
--- a/Engine/Shaders/ReflectionEnvironmentComputeShaders.usf
+++ b/Engine/Shaders/ReflectionEnvironmentComputeShaders.usf
@@ -264,10 +264,10 @@ void DoTileCulling(uint3 GroupId, uint ThreadIndex, float MinTileZ, float MaxTil
 
 		// O(N^2) simple parallel sort
 		LOOP
-		for (uint CaptureIndex = ThreadIndex; CaptureIndex < NumCapturesAffectingTile; CaptureIndex += THREADGROUP_TOTALSIZE)
+		for (uint CaptureIndex2 = ThreadIndex; CaptureIndex2 < NumCapturesAffectingTile; CaptureIndex2 += THREADGROUP_TOTALSIZE)
 		{
 			// Sort by original capture index
-			int SortKey = TileReflectionCaptureIndices[CaptureIndex];
+			int SortKey = TileReflectionCaptureIndices[CaptureIndex2];
 			uint NumSmaller = 0;
 
 			// Count how many items have a smaller key, so we can insert ourselves into the correct position, without requiring interaction between threads
@@ -282,7 +282,7 @@ void DoTileCulling(uint3 GroupId, uint ThreadIndex, float MinTileZ, float MaxTil
 			}
 
 			// Move this entry into its sorted position
-			SortedTileReflectionCaptureIndices[NumSmaller] = TileReflectionCaptureIndices[CaptureIndex];
+			SortedTileReflectionCaptureIndices[NumSmaller] = TileReflectionCaptureIndices[CaptureIndex2];
 		}
 
 	#endif
@@ -318,7 +318,7 @@ void ReflectionEnvironmentTiledDeferredMain(
 	DoTileCulling(GroupId, ThreadIndex, MinTileZ, MaxTileZ, MinTileZ2, MaxTileZ2);
 
 	// Lookup GBuffer properties once per pixel
-	FScreenSpaceData ScreenSpaceData = GetScreenSpaceData(PixelPos);
+	FScreenSpaceData ScreenSpaceData = GetScreenSpaceDataUint(PixelPos);
 	FGBufferData InGBufferData = ScreenSpaceData.GBuffer;
 	
 	float4 HomogeneousWorldPosition = mul(float4(ScreenPosition * SceneDepth, SceneDepth, 1), View.ScreenToWorld);
diff --git a/Engine/Shaders/TiledDeferredLightShaders.usf b/Engine/Shaders/TiledDeferredLightShaders.usf
--- a/Engine/Shaders/TiledDeferredLightShaders.usf
+++ b/Engine/Shaders/TiledDeferredLightShaders.usf
@@ -162,7 +162,7 @@ void TiledDeferredLightingMain(
 	uint NumSimpleLightsAffectingTile = TileNumSimpleLights;
 	
 	// Lookup GBuffer properties once per pixel 
-	FScreenSpaceData ScreenSpaceData = GetScreenSpaceData(PixelPos);
+	FScreenSpaceData ScreenSpaceData = GetScreenSpaceDataUint(PixelPos);
 	FGBufferData InGBufferData = ScreenSpaceData.GBuffer; 
 	float3 CameraVector = normalize(WorldPosition - View.ViewOrigin.xyz);
 
