Fixed up gamma correction on iOS for slate.

diff --git a/Engine/Shaders/SlateElementPixelShader.usf b/Engine/Shaders/SlateElementPixelShader.usf
--- a/Engine/Shaders/SlateElementPixelShader.usf
+++ b/Engine/Shaders/SlateElementPixelShader.usf
@@ -249,8 +249,11 @@ float4 Main( VertexOut VIn ) : SV_Target0
 
 	OutColor = ApplyColorVisionDeficiencyFilter(OutColor);
 
+#if !COMPILER_GLSL_ES2
 	// gamma correct
 	OutColor.rgb = pow(OutColor.rgb,1.0f/DisplayGamma);
+#endif
+
 
 #if DEBUG_CLIPPING
 	float4 Diff = float4( VIn.WorldPosition.x - VIn.ClipCoords.x, VIn.ClipCoords.z - VIn.WorldPosition.x, VIn.WorldPosition.y - VIn.ClipCoords.y, VIn.ClipCoords.w - VIn.WorldPosition.y );
