Slate mobile font color shader fix.

diff --git a/Engine/Shaders/SlateVertexShader.usf b/Engine/Shaders/SlateVertexShader.usf
--- a/Engine/Shaders/SlateVertexShader.usf
+++ b/Engine/Shaders/SlateVertexShader.usf
@@ -75,7 +75,11 @@ VertexOut Main(
 
 	// TextureCoordinates contains both the first and second texture coordinates in xy and zw respectively.
 	VOut.TextureCoordinates = InTextureCoordinates;
+#if ES2_PROFILE && COMPILER_GLSL_ES2
+	VOut.Color = InColor FCOLOR_COMPONENT_SWIZZLE;
+#else
 	VOut.Color = InColor;
+#endif
 
 	return VOut;
 }
