Optimization: Shaved off 31% of the high quality, and 12% of the lower quality temporal AA cost.

diff --git a/Engine/Shaders/PostProcessTemporalAA.usf b/Engine/Shaders/PostProcessTemporalAA.usf
--- a/Engine/Shaders/PostProcessTemporalAA.usf
+++ b/Engine/Shaders/PostProcessTemporalAA.usf
@@ -20,39 +20,86 @@ float2 RandomOffset;
 // TODO: This can be removed.
 float ClampFeebackMix;
 
-float Luma(float3 Color) 
-{
-	#if 1
-		// This seems to work better (less same luma ghost trails).
-		// CCIR 601 function for luma.
-		return dot(Color, float3(0.299, 0.587, 0.114));
-	#else
-		// Rec 709 function for luma.
-		return dot(Color, float3(0.2126, 0.7152, 0.0722));
-	#endif
+
+
+// Faster but less accurate luma computation. 
+// Luma includes a scaling by 4.
+float Luma4(float3 Color)
+{
+	return (Color.g * 2.0) + (Color.r + Color.b);
+}
+
+// Optimized HDR weighting function.
+float HdrWeight4(float3 Color, float Exposure) 
+{
+	return rcp(Luma4(Color) * Exposure + 4.0);
+}
+
+float HdrWeightG(float3 Color, float Exposure) 
+{
+	return rcp(Color.g * Exposure + 1.0);
+}
+
+float HdrWeightG_(float Color, float Exposure) 
+{
+	return rcp(Color * Exposure + 1.0);
+}
+
+
+// Optimized HDR weighting function.
+float HdrWeight4_(float Color, float Exposure) 
+{
+	return rcp(Color * Exposure + 4.0);
+}
+
+// Optimized HDR weighting inverse.
+float HdrWeightInv4(float3 Color, float Exposure) 
+{
+	return 4.0 * rcp(Luma4(Color) * (-Exposure) + 1.0);
+}
+
+float HdrWeightInvG(float3 Color, float Exposure) 
+{
+	return rcp(Color.g * (-Exposure) + 1.0);
 }
 
-float HighlightCompression(float Channel) 
+float HdrWeightInv4_(float Color, float Exposure) 
 {
-	return Channel * rcp(1.0 + Channel);
+	return 4.0 * rcp(Color * (-Exposure) + 1.0);
 }
-	
-float HighlightDecompression(float Channel) 
+
+float HdrWeightInvG_(float Color, float Exposure) 
 {
-	return Channel * rcp(1.0 - Channel);
+	return rcp(Color * (-Exposure) + 1.0);
 }
 
-float PerceptualLuma(float3 Color, float Exposure) 
+
+// This returns exposure normalized linear luma from a PerceptualLuma4().
+float LinearLuma4(float Channel, float Exposure) 
+{
+	return Channel * HdrWeightInv4_(Channel, Exposure);
+}
+
+// This returns exposure normalized linear luma from a PerceptualLuma4().
+float LinearLumaG(float Channel, float Exposure) 
+{
+	return Channel * HdrWeightInvG_(Channel, Exposure);
+}
+
+
+float PerceptualLuma4(float3 Color, float Exposure) 
 {
-	return sqrt(HighlightCompression(Luma(Color) * Exposure));
+	float L = Luma4(Color);
+	return L * HdrWeight4_(L, Exposure);
 }
 
-float LinearLuma(float Channel) 
+float PerceptualLumaG(float3 Color, float Exposure) 
 {
-	// This returns exposure normalized linear luma from a PerceptualLuma().
-	return HighlightDecompression(Channel * Channel);
+	return Color.g * HdrWeightG_(Color.g, Exposure);
 }
 
+
+
 // Intersect ray with AABB, knowing there is an intersection.
 //   Dir = Ray direction.
 //   Org = Start of the ray.
@@ -70,7 +117,7 @@ float IntersectAABB(float3 Dir, float3 Org, float3 Box)
 	return max(max(min(TNeg.x, TPos.x), min(TNeg.y, TPos.y)), min(TNeg.z, TPos.z));
 }
 
-float HistoryClamp(float3 History, float3 Filtered, float3 NeighborMin, float3 NeighborMax, float Exposure)
+float HistoryClamp(float3 History, float3 Filtered, float3 NeighborMin, float3 NeighborMax)
 {
 	float3 Min = min(Filtered, min(NeighborMin, NeighborMax));
 	float3 Max = max(Filtered, max(NeighborMin, NeighborMax));	
@@ -81,21 +128,6 @@ float HistoryClamp(float3 History, float3 Filtered, float3 NeighborMin, float3 N
 	return saturate(IntersectAABB(Dir, Org, Scale));	
 }
 
-float HdrWeight(float3 Color, float Exposure) 
-{
-	return rcp(max(Luma(Color) * Exposure, 1.0));
-}
-
-float4 HdrLerp(float4 ColorA, float4 ColorB, float Blend, float Exposure) 
-{
-	float BlendA = (1.0 - Blend) * HdrWeight(ColorA.rgb, Exposure);
-	float BlendB =        Blend  * HdrWeight(ColorB.rgb, Exposure);
-	float RcpBlend = rcp(BlendA + BlendB);
-	BlendA *= RcpBlend;
-	BlendB *= RcpBlend;
-	return ColorA * BlendA + ColorB * BlendB;
-}
-
 
 
 void SSRTemporalAAPS( float4 UVAndScreenPos : TEXCOORD0, float3 InExposureScaleVignette : TEXCOORD1, out float4 OutColor : SV_Target0 )
@@ -151,6 +183,10 @@ void MainTemporalAAPS( float4 UVAndScreenPos : TEXCOORD0, float3 InExposureScale
 	float InExposureScale = InExposureScaleVignette.x;
 	#define AA_BORDER 1
 	#define AA_GRAIN 1
+	#define AA_AABB 1
+	#define AA_GREEN_AS_LUMA 1
+	#define AA_LOWPASS 0
+	#define AA_HDR_HISTORY 0
 	#include "PostProcessTemporalCommon.usf"
 #else 
 	OutColor = PostprocessInput0.SampleLevel(PostprocessInput0Sampler, UVAndScreenPos.xy, 0);
@@ -163,6 +199,7 @@ void MainFastTemporalAAPS( float4 UVAndScreenPos : TEXCOORD0, float3 InExposureS
 	float InExposureScale = InExposureScaleVignette.x;
 	#define AA_BORDER 1
 	#define AA_GRAIN 1
+	#define AA_GREEN_AS_LUMA 0
 	#define AA_HDR 0
 	#define AA_HDR_HISTORY 0
 	#define AA_LOWPASS 0
diff --git a/Engine/Shaders/PostProcessTemporalCommon.usf b/Engine/Shaders/PostProcessTemporalCommon.usf
--- a/Engine/Shaders/PostProcessTemporalCommon.usf
+++ b/Engine/Shaders/PostProcessTemporalCommon.usf
@@ -110,6 +110,12 @@ which can be configured via defines for various temporal AA passes.
 	#define AA_ONE_DYNAMIC_SAMPLE 0
 #endif
 
+// Not always faster path, use green as luma.
+#ifndef AA_GREEN_AS_LUMA
+	#define AA_GREEN_AS_LUMA 0
+#endif
+
+
 	// FIND MOTION OF PIXEL AND NEAREST IN NEIGHBORHOOD
     // ------------------------------------------------
 	float3 PosP; // Position of this pixel.
@@ -254,84 +260,53 @@ which can be configured via defines for various temporal AA passes.
 	float4 Neighbor6 = PostprocessInput0.SampleLevel(PostprocessInput0Sampler, UVAndScreenPos.xy, 0, int2(-1,  1));
 	float4 Neighbor7 = PostprocessInput0.SampleLevel(PostprocessInput0Sampler, UVAndScreenPos.xy, 0, int2( 0,  1));
 	float4 Neighbor8 = PostprocessInput0.SampleLevel(PostprocessInput0Sampler, UVAndScreenPos.xy, 0, int2( 1,  1));
+	#if AA_HDR
+		#if AA_GREEN_AS_LUMA
+			Neighbor0.rgb *= HdrWeightG(Neighbor0.rgb, InExposureScale);
+			Neighbor1.rgb *= HdrWeightG(Neighbor1.rgb, InExposureScale);
+			Neighbor2.rgb *= HdrWeightG(Neighbor2.rgb, InExposureScale);
+			Neighbor3.rgb *= HdrWeightG(Neighbor3.rgb, InExposureScale);
+			Neighbor4.rgb *= HdrWeightG(Neighbor4.rgb, InExposureScale);
+			Neighbor5.rgb *= HdrWeightG(Neighbor5.rgb, InExposureScale);
+			Neighbor6.rgb *= HdrWeightG(Neighbor6.rgb, InExposureScale);
+			Neighbor7.rgb *= HdrWeightG(Neighbor7.rgb, InExposureScale);
+			Neighbor8.rgb *= HdrWeightG(Neighbor8.rgb, InExposureScale);
+		#else
+			Neighbor0.rgb *= HdrWeight4(Neighbor0.rgb, InExposureScale);
+			Neighbor1.rgb *= HdrWeight4(Neighbor1.rgb, InExposureScale);
+			Neighbor2.rgb *= HdrWeight4(Neighbor2.rgb, InExposureScale);
+			Neighbor3.rgb *= HdrWeight4(Neighbor3.rgb, InExposureScale);
+			Neighbor4.rgb *= HdrWeight4(Neighbor4.rgb, InExposureScale);
+			Neighbor5.rgb *= HdrWeight4(Neighbor5.rgb, InExposureScale);
+			Neighbor6.rgb *= HdrWeight4(Neighbor6.rgb, InExposureScale);
+			Neighbor7.rgb *= HdrWeight4(Neighbor7.rgb, InExposureScale);
+			Neighbor8.rgb *= HdrWeight4(Neighbor8.rgb, InExposureScale);
+		#endif
+	#endif
 	#if AA_FILTERED
-		#if AA_HDR
-			float SampleHdrWeight0 = HdrWeight(Neighbor0.rgb, InExposureScale);
-			float SampleHdrWeight1 = HdrWeight(Neighbor1.rgb, InExposureScale);
-			float SampleHdrWeight2 = HdrWeight(Neighbor2.rgb, InExposureScale);
-			float SampleHdrWeight3 = HdrWeight(Neighbor3.rgb, InExposureScale);
-			float SampleHdrWeight4 = HdrWeight(Neighbor4.rgb, InExposureScale);
-			float SampleHdrWeight5 = HdrWeight(Neighbor5.rgb, InExposureScale);
-			float SampleHdrWeight6 = HdrWeight(Neighbor6.rgb, InExposureScale);
-			float SampleHdrWeight7 = HdrWeight(Neighbor7.rgb, InExposureScale);
-			float SampleHdrWeight8 = HdrWeight(Neighbor8.rgb, InExposureScale);
-			float4 Filtered = (
-				Neighbor0 * (SampleWeights[0] * SampleHdrWeight0) +
-				Neighbor1 * (SampleWeights[1] * SampleHdrWeight1) +
-				Neighbor2 * (SampleWeights[2] * SampleHdrWeight2) +
-				Neighbor3 * (SampleWeights[3] * SampleHdrWeight3) +
-				Neighbor4 * (SampleWeights[4] * SampleHdrWeight4) +
-				Neighbor5 * (SampleWeights[5] * SampleHdrWeight5) +
-				Neighbor6 * (SampleWeights[6] * SampleHdrWeight6) +
-				Neighbor7 * (SampleWeights[7] * SampleHdrWeight7) +
-				Neighbor8 * (SampleWeights[8] * SampleHdrWeight8)) * rcp(
-					SampleWeights[0] * SampleHdrWeight0 +
-					SampleWeights[1] * SampleHdrWeight1 +
-					SampleWeights[2] * SampleHdrWeight2 +
-					SampleWeights[3] * SampleHdrWeight3 +
-					SampleWeights[4] * SampleHdrWeight4 +
-					SampleWeights[5] * SampleHdrWeight5 +
-					SampleWeights[6] * SampleHdrWeight6 +
-					SampleWeights[7] * SampleHdrWeight7 +
-					SampleWeights[8] * SampleHdrWeight8);
-			#if AA_LOWPASS
-				float4 FilteredLow = (
-					Neighbor0 * (LowpassWeights[0] * SampleHdrWeight0) +
-					Neighbor1 * (LowpassWeights[1] * SampleHdrWeight1) +
-					Neighbor2 * (LowpassWeights[2] * SampleHdrWeight2) +
-					Neighbor3 * (LowpassWeights[3] * SampleHdrWeight3) +
-					Neighbor4 * (LowpassWeights[4] * SampleHdrWeight4) +
-					Neighbor5 * (LowpassWeights[5] * SampleHdrWeight5) +
-					Neighbor6 * (LowpassWeights[6] * SampleHdrWeight6) +
-					Neighbor7 * (LowpassWeights[7] * SampleHdrWeight7) +
-					Neighbor8 * (LowpassWeights[8] * SampleHdrWeight8)) * rcp(
-						LowpassWeights[0] * SampleHdrWeight0 +
-						LowpassWeights[1] * SampleHdrWeight1 +
-						LowpassWeights[2] * SampleHdrWeight2 +
-						LowpassWeights[3] * SampleHdrWeight3 +
-						LowpassWeights[4] * SampleHdrWeight4 +
-						LowpassWeights[5] * SampleHdrWeight5 +
-						LowpassWeights[6] * SampleHdrWeight6 +
-						LowpassWeights[7] * SampleHdrWeight7 +
-						LowpassWeights[8] * SampleHdrWeight8);
-			#else
-				float4 FilteredLow = Filtered;
-			#endif
+		float4 Filtered = 
+			Neighbor0 * SampleWeights[0] +
+			Neighbor1 * SampleWeights[1] +
+			Neighbor2 * SampleWeights[2] +
+			Neighbor3 * SampleWeights[3] +
+			Neighbor4 * SampleWeights[4] +
+			Neighbor5 * SampleWeights[5] +
+			Neighbor6 * SampleWeights[6] +
+			Neighbor7 * SampleWeights[7] +
+			Neighbor8 * SampleWeights[8];
+		#if AA_LOWPASS
+			float4 FilteredLow = 
+				Neighbor0 * LowpassWeights[0] +
+				Neighbor1 * LowpassWeights[1] +
+				Neighbor2 * LowpassWeights[2] +
+				Neighbor3 * LowpassWeights[3] +
+				Neighbor4 * LowpassWeights[4] +
+				Neighbor5 * LowpassWeights[5] +
+				Neighbor6 * LowpassWeights[6] +
+				Neighbor7 * LowpassWeights[7] +
+				Neighbor8 * LowpassWeights[8];
 		#else
-			float4 Filtered = 
-				Neighbor0 * SampleWeights[0] +
-				Neighbor1 * SampleWeights[1] +
-				Neighbor2 * SampleWeights[2] +
-				Neighbor3 * SampleWeights[3] +
-				Neighbor4 * SampleWeights[4] +
-				Neighbor5 * SampleWeights[5] +
-				Neighbor6 * SampleWeights[6] +
-				Neighbor7 * SampleWeights[7] +
-				Neighbor8 * SampleWeights[8];
-			#if AA_LOWPASS
-				float4 FilteredLow = 
-					Neighbor0 * LowpassWeights[0] +
-					Neighbor1 * LowpassWeights[1] +
-					Neighbor2 * LowpassWeights[2] +
-					Neighbor3 * LowpassWeights[3] +
-					Neighbor4 * LowpassWeights[4] +
-					Neighbor5 * LowpassWeights[5] +
-					Neighbor6 * LowpassWeights[6] +
-					Neighbor7 * LowpassWeights[7] +
-					Neighbor8 * LowpassWeights[8];
-			#else
-				float4 FilteredLow = Filtered;
-			#endif
+			float4 FilteredLow = Filtered;
 		#endif
 		#if AA_BORDER
 			// Use unfiltered for 1 pixel border.
@@ -342,13 +317,13 @@ which can be configured via defines for various temporal AA passes.
 			if(FilteredOffScreen) 
 			{
 				Filtered = Neighbor4;
-				FilteredLow = Filtered;
+				FilteredLow = Neighbor4;
 			}
 		#endif
 	#else
 		// Unfiltered.
 		float4 Filtered = Neighbor4;
-		float4 FilteredLow = Filtered;
+		float4 FilteredLow = Neighbor4;
 	#endif
 	#if AA_ROUND
 		float4 NeighborMin2 = min(min(Neighbor0, Neighbor2), min(Neighbor6, Neighbor8));		
@@ -357,13 +332,8 @@ which can be configured via defines for various temporal AA passes.
 		float4 NeighborMax = max(max(max(Neighbor1, Neighbor3), max(Neighbor4, Neighbor5)), Neighbor7);		
 		NeighborMin2 = min(NeighborMin2, NeighborMin);
 		NeighborMax2 = max(NeighborMax2, NeighborMax);
-		#if AA_HDR
-			NeighborMin = HdrLerp(NeighborMin, NeighborMin2, 0.5, InExposureScale);
-			NeighborMax = HdrLerp(NeighborMax, NeighborMax2, 0.5, InExposureScale);
-		#else
-			NeighborMin = lerp(NeighborMin, NeighborMin2, 0.5);
-			NeighborMax = lerp(NeighborMax, NeighborMax2, 0.5);
-		#endif
+		NeighborMin = NeighborMin * 0.5 + NeighborMin2 * 0.5;
+		NeighborMax = NeighborMax * 0.5 + NeighborMax2 * 0.5;
 	#else
 		float4 NeighborMin = min(min(
 			min(min(Neighbor0, Neighbor1), min(Neighbor2, Neighbor3)),
@@ -380,33 +350,15 @@ which can be configured via defines for various temporal AA passes.
 		float DebugDiffCurrent = Filtered.b;
 	#endif
 
-	// FETCH HISTORY AND MANUALLY INTERPOLATE WITH WEIGHT FILTER FOR PROPER HDR
-	// ------------------------------------------------------------------------
-	#if AA_HDR_HISTORY
-		// a c
-		//  EF
-		// gHI
-		float4 OutColorE = PostprocessInput2.SampleLevel(PostprocessInput2Sampler, BackN.xy, 0, int2(0, 0));
-		float4 OutColorF = PostprocessInput2.SampleLevel(PostprocessInput2Sampler, BackN.xy, 0, int2(1, 0));
-		float4 OutColorH = PostprocessInput2.SampleLevel(PostprocessInput2Sampler, BackN.xy, 0, int2(0, 1));
-		float4 OutColorI = PostprocessInput2.SampleLevel(PostprocessInput2Sampler, BackN.xy, 0, int2(1, 1));
-		float WeightE = HdrWeight(OutColorE.rgb, InExposureScale);
-		float WeightF = HdrWeight(OutColorF.rgb, InExposureScale);
-		float WeightH = HdrWeight(OutColorH.rgb, InExposureScale);
-		float WeightI = HdrWeight(OutColorI.rgb, InExposureScale);
-		float2 Subpix = frac(BackN.xy * PostprocessInput1Size.zw);
-		WeightE *= (1.0 - Subpix.x) * (1.0 - Subpix.y);
-		WeightF *= (      Subpix.x) * (1.0 - Subpix.y);
-		WeightH *= (1.0 - Subpix.x) * (      Subpix.y);
-		WeightI *= (      Subpix.x) * (      Subpix.y);
-		float RcpWeightEFHI = rcp(WeightE + WeightF + WeightH + WeightI);
-		WeightE *= RcpWeightEFHI;
-		WeightF *= RcpWeightEFHI;
-		WeightH *= RcpWeightEFHI;
-		WeightI *= RcpWeightEFHI;
-		OutColor = (OutColorE * WeightE) + (OutColorF * WeightF) + (OutColorH * WeightH) + (OutColorI * WeightI);
-	#else
-		OutColor = PostprocessInput1.SampleLevel(PostprocessInput1Sampler, BackN.xy, 0);
+	// FETCH HISTORY
+	// -------------
+	OutColor = PostprocessInput1.SampleLevel(PostprocessInput1Sampler, BackN.xy, 0);
+	#if AA_HDR
+		#if AA_GREEN_AS_LUMA
+			OutColor.rgb *= HdrWeightG(OutColor.rgb, InExposureScale);
+		#else
+			OutColor.rgb *= HdrWeight4(OutColor.rgb, InExposureScale);
+		#endif
 	#endif
 	#if AA_DILATE
 		// Grab alphas from cross pattern and take maximum, dilate feedback control.
@@ -433,12 +385,10 @@ which can be configured via defines for various temporal AA passes.
 
 	// FIND LUMA OF CLAMPED HISTORY
 	// ----------------------------
-	// Save luma converted into a perceptual space.
-	float LumaHistory = PerceptualLuma(OutColor.rgb, InExposureScale);
 	#if AA_AABB
 		// Clamp history, this uses color AABB intersection for tighter fit.
 		// Clamping works with the low pass (if available) to reduce flicker.
-		float ClampBlend = HistoryClamp(OutColor.rgb, FilteredLow.rgb, NeighborMin.rgb, NeighborMax.rgb, InExposureScale);
+		float ClampBlend = HistoryClamp(OutColor.rgb, FilteredLow.rgb, NeighborMin.rgb, NeighborMax.rgb);
 		float4 Clamped = lerp(OutColor, FilteredLow, ClampBlend);
 	#else
 		float4 Clamped = clamp(OutColor, NeighborMin, NeighborMax);
@@ -446,11 +396,35 @@ which can be configured via defines for various temporal AA passes.
 
 	// FIND PERCEPTUAL LUMAS
 	// ---------------------
-	float LumaClamped = PerceptualLuma(Clamped.rgb, InExposureScale);
-	float LumaFiltered = PerceptualLuma(Filtered.rgb, InExposureScale);
-	float LumaFilteredLow = PerceptualLuma(FilteredLow.rgb, InExposureScale);
-	float LumaMin = PerceptualLuma(NeighborMin.rgb, InExposureScale);
-	float LumaMax = PerceptualLuma(NeighborMax.rgb, InExposureScale);
+	#if AA_HDR
+		#if AA_GREEN_AS_LUMA
+			float LumaHistory = OutColor.g;
+			float LumaClamped = Clamped.g;
+			float LumaFiltered = Filtered.g;
+			float LumaMin = NeighborMin.g;
+			float LumaMax = NeighborMax.g;
+		#else
+			float LumaHistory = Luma4(OutColor.rgb);
+			float LumaClamped = Luma4(Clamped.rgb);
+			float LumaFiltered = Luma4(Filtered.rgb);
+			float LumaMin = Luma4(NeighborMin.rgb);
+			float LumaMax = Luma4(NeighborMax.rgb);
+		#endif
+	#else
+		#if AA_GREEN_AS_LUMA
+			float LumaHistory = PerceptualLumaG(OutColor.rgb, InExposureScale);
+			float LumaClamped = PerceptualLumaG(Clamped.rgb, InExposureScale);
+			float LumaFiltered = PerceptualLumaG(Filtered.rgb, InExposureScale);
+			float LumaMin = PerceptualLumaG(NeighborMin.rgb, InExposureScale);
+			float LumaMax = PerceptualLumaG(NeighborMax.rgb, InExposureScale);
+		#else
+			float LumaHistory = PerceptualLuma4(OutColor.rgb, InExposureScale);
+			float LumaClamped = PerceptualLuma4(Clamped.rgb, InExposureScale);
+			float LumaFiltered = PerceptualLuma4(Filtered.rgb, InExposureScale);
+			float LumaMin = PerceptualLuma4(NeighborMin.rgb, InExposureScale);
+			float LumaMax = PerceptualLuma4(NeighborMax.rgb, InExposureScale);
+		#endif
+	#endif
 
 	// CONVERT MOTION AMOUNTS TO ALLOWED HISTORY TO FILTERED CHANGE AMOUNTS
     // --------------------------------------------------------------------	
@@ -551,9 +525,31 @@ which can be configured via defines for various temporal AA passes.
 	// -------------------------------------------------
 	#if (AA_LERP == 0)
 		// Switch luma back to linear.
-		LumaTarget = LinearLuma(LumaTarget);
-		LumaHistory = LinearLuma(LumaHistory);
-		LumaFiltered = LinearLuma(LumaFiltered);
+		#if AA_HDR
+			#if AA_GREEN_AS_LUMA
+				LumaTarget = LinearLumaG(LumaTarget, InExposureScale);
+				LumaHistory = LinearLumaG(LumaHistory, InExposureScale);
+				LumaFiltered = LinearLumaG(LumaFiltered, InExposureScale);
+				OutColor.rgb *= HdrWeightInvG(OutColor.rgb, InExposureScale);
+				Filtered.rgb *= HdrWeightInvG(Filtered.rgb, InExposureScale);
+			#else
+				LumaTarget = LinearLuma4(LumaTarget, InExposureScale);
+				LumaHistory = LinearLuma4(LumaHistory, InExposureScale);
+				LumaFiltered = LinearLuma4(LumaFiltered, InExposureScale);
+				OutColor.rgb *= HdrWeightInv4(OutColor.rgb, InExposureScale);
+				Filtered.rgb *= HdrWeightInv4(Filtered.rgb, InExposureScale);
+			#endif
+		#else
+			#if AA_GREEN_AS_LUMA
+				LumaTarget = LinearLumaG(LumaTarget, InExposureScale);
+				LumaHistory = LinearLumaG(LumaHistory, InExposureScale);
+				LumaFiltered = LinearLumaG(LumaFiltered, InExposureScale);
+			#else
+				LumaTarget = LinearLuma4(LumaTarget, InExposureScale);
+				LumaHistory = LinearLuma4(LumaHistory, InExposureScale);
+				LumaFiltered = LinearLuma4(LumaFiltered, InExposureScale);
+			#endif
+		#endif
 		float LumaDiff = LumaHistory - LumaFiltered;
 		float RcpLumaDiff = rcp(LumaDiff);
 		if(abs(LumaDiff) < (1.0/1024.0)) 
@@ -575,10 +571,13 @@ which can be configured via defines for various temporal AA passes.
 			#endif
 		#endif
 	#else
+		OutColor = lerp(Clamped, Filtered, FixedLerp);
 		#if AA_HDR
-			OutColor = HdrLerp(Clamped, Filtered, FixedLerp, InExposureScale);
-		#else
-			OutColor = lerp(Clamped, Filtered, FixedLerp);
+			#if AA_GREEN_AS_LUMA
+				OutColor.rgb *= HdrWeightInvG(OutColor.rgb, InExposureScale);
+			#else
+				OutColor.rgb *= HdrWeightInv4(OutColor.rgb, InExposureScale);
+			#endif
 		#endif
 	#endif
 	#if AA_NAN
@@ -606,4 +605,5 @@ which can be configured via defines for various temporal AA passes.
 #undef AA_NAN
 #undef AA_BORDER
 #undef AA_FORCE_ALPHA_CLAMP
-#undef AA_ONE_DYNAMIC_SAMPLE
\ No newline at end of file
+#undef AA_ONE_DYNAMIC_SAMPLE
+#undef AA_GREEN_AS_LUMA
