* improved SSAO quality

diff --git a/Engine/Shaders/PostProcessAmbientOcclusion.usf b/Engine/Shaders/PostProcessAmbientOcclusion.usf
--- a/Engine/Shaders/PostProcessAmbientOcclusion.usf
+++ b/Engine/Shaders/PostProcessAmbientOcclusion.usf
@@ -62,12 +62,21 @@ float2 ScreenPosToUV(float2 ScreenPos, float4 InputSize)
 // needed to prevent AO seam near 16 bit float maximum, this feactor pushed the problem far out and it seems to not have a visual degradion nearby
 const static float Constant_Float16F_Scale =  4096.0f * 32.0f;
 
+
+#if SAMPLESET_QUALITY
+	#define USE_SAMPLESET 2
+#else
+	#define USE_SAMPLESET 3
+#endif
+
 // 0:4 samples, 1:9 samples (only really noticable with dither usage ??)
 //#define AO_UPSAMPLE_QUALITY 
 
 #if USE_AO_SETUP_AS_INPUT == 1
 	// lower resolution
-	#define AO_SAMPLE_QUALITY 2
+	#define AO_SAMPLE_QUALITY 3
+	#undef USE_SAMPLESET
+	#define USE_SAMPLESET 0
 	#define AO_UPSAMPLE_QUALITY 1
 #else
 	// full resolution is expensive, do lower quality
@@ -81,11 +90,6 @@ const static float Constant_Float16F_Scale =  4096.0f * 32.0f;
 // 3: 1 point (for testing)
 //#define USE_SAMPLESET 2
 
-#if SAMPLESET_QUALITY
-	#define USE_SAMPLESET 0
-#else
-	#define USE_SAMPLESET 2
-#endif
 
 
 // only for MainSetupPS()
@@ -380,7 +384,7 @@ void MainPS(in float4 UVAndScreenPos : TEXCOORD0, out float4 OutColor : SV_Targe
 	{
 #elif AO_SAMPLE_QUALITY == 3
 	// extract one of 16 base vectors (rotation and scale) from a texture that repeats 4x4, changing over time if TemporalAA is enabled
-	float2 TemporalOffset = (View.TemporalAAParams.r % 8) * float2(3.5, 7) / 64.0f;
+	float2 TemporalOffset = (View.TemporalAAParams.r % 8) * float2(2, 7) / 64.0f;
 
 	float4 RandomNorm = Texture2DSample(RandomNormalTexture, RandomNormalTextureSampler, TemporalOffset + UV * ViewportUVToRandomUV);
 	float4 RandomBase = (RandomNorm * 2 - 1) * float4(AORadiusInShader, AORadiusInShader);
