* SSAO optimized for XBoxOne
* cleanup SSAO shader, minor polish

diff --git a/Engine/Config/BaseDeviceProfiles.ini b/Engine/Config/BaseDeviceProfiles.ini
--- a/Engine/Config/BaseDeviceProfiles.ini
+++ b/Engine/Config/BaseDeviceProfiles.ini
@@ -197,6 +197,13 @@ MeshLODSettings=
 TextureLODSettings=
 ; we output 10:10:10, not 8:8:8 so we don't need color quantization
 +CVars=r.TonemapperQuality=0
+; For SSAO we rely on TemporalAA (with a randomized sample pattern over time) so we can use less samples
++CVars=r.AmbientOcclusionSampleSetQuality=0
+; less passes, and no upsampling as even upsampling costs some performance
++CVars=r.AmbientOcclusionLevels=1
+; larger radius to compensate for fewer passes
++CVars=r.AmbientOcclusionRadiusScale=2
+
 
 [HTML5 DeviceProfile]
 DeviceType=HTML5
diff --git a/Engine/Shaders/PostProcessAmbientOcclusion.usf b/Engine/Shaders/PostProcessAmbientOcclusion.usf
--- a/Engine/Shaders/PostProcessAmbientOcclusion.usf
+++ b/Engine/Shaders/PostProcessAmbientOcclusion.usf
@@ -19,8 +19,6 @@
 // 0:no / 1:yes
 // USE_UPSAMPLE
 
-
-
 // 0:does not use normals (slightly faster, lower quality, WIP = not finished), 1: uses normals (slower but higher quality)
 #define SSAO_NORMAL_AFFECTED 1
 
@@ -31,41 +29,11 @@
 // AO_SAMPLE_QUALITY = 3 : efficient high frequency 4x4 pattern changing over time
 // AO_SAMPLE_QUALITY = 4 : extreme sample count
 
-// [0]: .x:AmbientOcclusionPower, .y:AmbientOcclusionBias/BiasDistance, .z:1/AmbientOcclusionDistance, .w:AmbientOcclusionIntensity
-// [1]: .xy:ViewportUVToRandomUV, .z:AORadiusInShader, .w:Ratio
-// [2]: .x:ScaleFactor(e.g. 4 if current RT is a quarter in size), .y:InvThreshold, .z:ScaleRadiusInWorldSpace(0:VS/1:WS), .w:MipBlend
-// [3]: .x:SSR StepCount, .y:SSR StepSize, .z:StaticFraction, .w: unused
-// [4]: .x:Multipler for FadeDistance/Radius, .y:Additive for FadeDistance/Radius, .zw: unused
-float4 ScreenSpaceAOandSSRShaderParams[5];
-
-// could be moved to a more central spot
-// @param ScreenPos -1 .. 1
-float3 ReconstructCSPos(float SceneDepth, float2 ScreenPos)
-{
-	return float3(ScreenPos * SceneDepth, SceneDepth);
-}
-
-// could be moved to a more central spot
-float2 ReconstructSSPosFromCS(float3 In)
-{
-	return In.xy / In.z;
-}
-
-// could be moved to a more central spot
-// can be optimized
-// @param InputSize e.g. PostprocessInput0Size
-float2 ScreenPosToUV(float2 ScreenPos, float4 InputSize)
-{
-	return (ScreenPos * ScreenPosToPixel.xy + ScreenPosToPixel.zw + 0.5f) * InputSize.zw;
-}
-
-// needed to prevent AO seam near 16 bit float maximum, this feactor pushed the problem far out and it seems to not have a visual degradion nearby
-const static float Constant_Float16F_Scale =  4096.0f * 32.0f;
-
-
-#if SAMPLESET_QUALITY
-	#define USE_SAMPLESET 2
+#if SAMPLESET_QUALITY == 0
+	// low
+	#define USE_SAMPLESET 1
 #else
+	// high
 	#define USE_SAMPLESET 3
 #endif
 
@@ -76,7 +44,7 @@
 	// lower resolution
 	#define AO_SAMPLE_QUALITY 3
 	#undef USE_SAMPLESET
-	#define USE_SAMPLESET 0
+	#define USE_SAMPLESET 3
 	#define AO_UPSAMPLE_QUALITY 1
 #else
 	// full resolution is expensive, do lower quality
@@ -84,34 +52,28 @@
 	#define AO_UPSAMPLE_QUALITY 0
 #endif
 
-// 0: near the surface very large, softly fading out (6 points)
-// 1: more evenly spread (5 points - slightly faster, stronger effect, better with multiple levels?)
-// 2: 3 points
-// 3: 1 point (for testing)
+// 0: 1 point (for testing)
+// 1: 3 points
+// 2: more evenly spread (5 points - slightly faster, stronger effect, better with multiple levels?)
+// 3: near the surface very large, softly fading out (6 points)
 //#define USE_SAMPLESET 2
-
-
-
-// only for MainSetupPS()
-// .x:ScaleFactor(e.g. 4 if current RT is a quarter in size), .y:InvThreshold, .zw: unused
-float4 AmbientOcclusionSetupParams;
-
-// 
-float4 NoiseScale;
-
 #if USE_SAMPLESET == 0
-	#define SAMPLESET_ARRAY_SIZE 6
+	#define SAMPLESET_ARRAY_SIZE 1
 	static const float2 OcclusionSamplesOffsets[SAMPLESET_ARRAY_SIZE]=
 	{
 		// 6 points distributed on the unit disc, spiral order and distance
-		float2(0.000, 0.200), 
-		float2(0.325, 0.101), 
-		float2(0.272, -0.396), 
-		float2(-0.385, -0.488), 
-		float2(-0.711, 0.274), 
-		float2(0.060, 0.900) 
+		float2(0.500, 0.500), 
 	};
 #elif USE_SAMPLESET == 1
+	#define SAMPLESET_ARRAY_SIZE 3
+	static const float2 OcclusionSamplesOffsets[SAMPLESET_ARRAY_SIZE]=
+	{
+		// 3 points distributed on the unit disc, spiral order and distance
+		float2(0, -1.0f) * 0.3f, 
+		float2(0.58f, 0.814f) * 0.7f, 
+		float2(-0.58f, 0.814f) 
+	};
+#elif USE_SAMPLESET == 2
 	#define SAMPLESET_ARRAY_SIZE 5
 	static const float2 OcclusionSamplesOffsets[SAMPLESET_ARRAY_SIZE]=
 	{
@@ -122,28 +84,66 @@
 		float2(-0.707107, -0.707107),
 		float2(-0.891006, 0.453991),
 	};
-#elif USE_SAMPLESET == 2
-	#define SAMPLESET_ARRAY_SIZE 3
-	static const float2 OcclusionSamplesOffsets[SAMPLESET_ARRAY_SIZE]=
-	{
-		// 3 points distributed on the unit disc, spiral order and distance
-		float2(0, -1.0f) * 0.3f, 
-		float2(0.81f, 0.5) * 0.7f, 
-		float2(-0.81f, 0.5) 
-	};
 #else // USE_SAMPLESET == 3
-	#define SAMPLESET_ARRAY_SIZE 1
+	#define SAMPLESET_ARRAY_SIZE 6
 	static const float2 OcclusionSamplesOffsets[SAMPLESET_ARRAY_SIZE]=
 	{
 		// 6 points distributed on the unit disc, spiral order and distance
-		float2(0.500, 0.500), 
+		float2(0.000, 0.200), 
+		float2(0.325, 0.101), 
+		float2(0.272, -0.396), 
+		float2(-0.385, -0.488), 
+		float2(-0.711, 0.274), 
+		float2(0.060, 0.900) 
 	};
 #endif // USE_SAMPLESET
+	
 
+
+
+// [0]: .x:AmbientOcclusionPower, .y:AmbientOcclusionBias/BiasDistance, .z:1/AmbientOcclusionDistance, .w:AmbientOcclusionIntensity
+// [1]: .xy:ViewportUVToRandomUV, .z:AORadiusInShader, .w:Ratio
+// [2]: .x:ScaleFactor(e.g. 4 if current RT is a quarter in size), .y:InvThreshold, .z:ScaleRadiusInWorldSpace(0:VS/1:WS), .w:MipBlend
+// [3]: .x:SSR StepCount, .y:SSR StepSize, .z:StaticFraction, .w: unused
+// [4]: .x:Multipler for FadeDistance/Radius, .y:Additive for FadeDistance/Radius, .zw: unused
+float4 ScreenSpaceAOandSSRShaderParams[5];
+
+// needed to prevent AO seam near 16 bit float maximum, this feactor pushed the problem far out and it seems to not have a visual degradion nearby
+const static float Constant_Float16F_Scale =  4096.0f * 32.0f;
+
+// only for MainSetupPS()
+// .x:ScaleFactor(e.g. 4 if current RT is a quarter in size), .y:InvThreshold, .zw: unused
+float4 AmbientOcclusionSetupParams;
+
+// 
+float4 NoiseScale;
+
 /** RGBA8 linear texture containing random normals */
 Texture2D RandomNormalTexture;
 SamplerState RandomNormalTextureSampler;
 
+
+// could be moved to a more central spot
+// @param ScreenPos -1 .. 1
+float3 ReconstructCSPos(float SceneDepth, float2 ScreenPos)
+{
+	return float3(ScreenPos * SceneDepth, SceneDepth);
+}
+
+// could be moved to a more central spot
+float2 ReconstructSSPosFromCS(float3 In)
+{
+	return In.xy / In.z;
+}
+
+// could be moved to a more central spot
+// can be optimized
+// @param InputSize e.g. PostprocessInput0Size
+float2 ScreenPosToUV(float2 ScreenPos, float4 InputSize)
+{
+	return (ScreenPos * ScreenPosToPixel.xy + ScreenPosToPixel.zw + 0.5f) * InputSize.zw;
+}
+
 // @return true: inside, false outside
 bool IsScreenPosInsideBounds(float2 ScreenPos)
 {
@@ -151,7 +151,6 @@
 	return all(abs(ScreenPos) < 1);
 }
 
-
 // 0: not similar .. 1:very similar
 float ComputeDepthSimilarity(float DepthA, float DepthB, float TweakScale)
 {
