Monte Carlo functions for shaders.

Clean up of some shading code.

diff --git a/Engine/Shaders/BRDF.usf b/Engine/Shaders/BRDF.usf
--- a/Engine/Shaders/BRDF.usf
+++ b/Engine/Shaders/BRDF.usf
@@ -7,6 +7,9 @@
 // Physically based shading model
 // parameterized with the below options
 
+// Microfacet specular = D*G*F / (4*NoL*NoV) = D*Vis*F
+// Vis = G / (4*NoL*NoV)
+
 // Diffuse model
 // 0: Lambert
 // 1: Burley
@@ -25,7 +28,7 @@
 // 2: Kelemen
 // 3: Schlick
 // 4: Smith (matched to GGX)
-#define PHYSICAL_SPEC_G		3
+#define PHYSICAL_SPEC_G		4
 
 // Fresnel
 // 0: None
@@ -66,7 +69,7 @@ float D_Blinn( float Roughness, float NoH )
 	float m = Roughness * Roughness;
 	float m2 = m * m;
 	float n = 2 / m2 - 2;
-	return (n+2)/ 2 * PhongShadingPow( NoH, n );		// 1 mad, 1 exp, 1 mul, 1 log
+	return (n+2) / (2*PI) * PhongShadingPow( NoH, n );		// 1 mad, 1 exp, 1 mul, 1 log
 }
 
 // [Beckmann 1963, "The scattering of electromagnetic waves from rough surfaces"]
@@ -75,7 +78,7 @@ float D_Beckmann( float Roughness, float NoH )
 	float m = Roughness * Roughness;
 	float m2 = m * m;
 	float NoH2 = NoH * NoH;
-	return exp( (NoH2 - 1) / (m2 * NoH2) ) / ( m2 * NoH2 * NoH2 );
+	return exp( (NoH2 - 1) / (m2 * NoH2) ) / ( PI * m2 * NoH2 * NoH2 );
 }
 
 // GGX / Trowbridge-Reitz
@@ -85,7 +88,7 @@ float D_GGX( float Roughness, float NoH )
 	float m = Roughness * Roughness;
 	float m2 = m * m;
 	float d = ( NoH * m2 - NoH ) * NoH + 1;	// 2 mad
-	return m2 / ( d*d );					// 2 mul, 1 rcp
+	return m2 / ( PI*d*d );					// 3 mul, 1 rcp
 }
 
 // Anisotropic GGX
@@ -97,47 +100,46 @@ float D_GGXaniso( float RoughnessX, float RoughnessY, float NoH, float3 H, float
 	float XoH = dot( X, H );
 	float YoH = dot( Y, H );
 	float d = XoH*XoH / (mx*mx) + YoH*YoH / (my*my) + NoH*NoH;
-	return 1 / ( mx*my * d*d );
+	return 1 / ( PI * mx*my * d*d );
 }
 
-float G_Implicit()
+float Vis_Implicit()
 {
 	return 0.25;
 }
 
 // [Neumann et al. 1999, "Compact metallic reflectance models"]
-float G_Neumann( float NoV, float NoL )
+float Vis_Neumann( float NoV, float NoL )
 {
 	return 1 / ( 4 * max( NoL, NoV ) );
 }
 
 // [Kelemen 2001, "A microfacet based coupled specular-matte brdf model with importance sampling"]
-float G_Kelemen( float3 L, float3 V )
+float Vis_Kelemen( float3 L, float3 V )
 {
 	return 1 / ( 2 + 2 * dot(L, V) );
 }
 
-// Tuned to match behavior of G_Smith
+// Tuned to match behavior of Vis_Smith
 // [Schlick 1994, "An Inexpensive BRDF Model for Physically-Based Rendering"]
-float G_Schlick( float Roughness, float NoV, float NoL )
+float Vis_Schlick( float Roughness, float NoV, float NoL )
 {
 	float k = Square( Roughness ) * 0.5;
-	float G_SchlickV = NoV * (1 - k) + k;
-	float G_SchlickL = NoL * (1 - k) + k;
-	return 0.25 / ( G_SchlickV * G_SchlickL );
+	float Vis_SchlickV = NoV * (1 - k) + k;
+	float Vis_SchlickL = NoL * (1 - k) + k;
+	return 0.25 / ( Vis_SchlickV * Vis_SchlickL );
 }
 
-// Smith term for GGX modified by Disney to be less "hot" for small roughness values
+// Smith term for GGX
 // [Smith 1967, "Geometrical shadowing of a random rough surface"]
-// [Burley 2012, "Physically-Based Shading at Disney"]
-float G_Smith( float Roughness, float NoV, float NoL )
+float Vis_Smith( float Roughness, float NoV, float NoL )
 {
 	float a = Square( Roughness );
 	float a2 = a*a;
 
-	float G_SmithV = NoV + sqrt( NoV * (NoV - NoV * a2) + a2 );
-	float G_SmithL = NoL + sqrt( NoL * (NoL - NoL * a2) + a2 );
-	return rcp( G_SmithV * G_SmithL );
+	float Vis_SmithV = NoV + sqrt( NoV * (NoV - NoV * a2) + a2 );
+	float Vis_SmithL = NoL + sqrt( NoL * (NoL - NoL * a2) + a2 );
+	return rcp( Vis_SmithV * Vis_SmithL );
 }
 
 float3 F_None( float3 SpecularColor )
@@ -187,18 +189,19 @@ float Distribution( float Roughness, float NoH )
 #endif
 }
 
+// Vis = G / (4*NoL*NoV)
 float GeometricVisibility( float Roughness, float NoV, float NoL, float VoH, float3 L, float3 V )
 {
 #if   PHYSICAL_SPEC_G == 0
-	return G_Implicit();
+	return Vis_Implicit();
 #elif PHYSICAL_SPEC_G == 1
-	return G_Neumann( NoV, NoL );
+	return Vis_Neumann( NoV, NoL );
 #elif PHYSICAL_SPEC_G == 2
-	return G_Kelemen( L, V );
+	return Vis_Kelemen( L, V );
 #elif PHYSICAL_SPEC_G == 3
-	return G_Schlick( Roughness, NoV, NoL );
+	return Vis_Schlick( Roughness, NoV, NoL );
 #elif PHYSICAL_SPEC_G == 4
-	return G_Smith( Roughness, NoV, NoL );
+	return Vis_Smith( Roughness, NoV, NoL );
 #endif
 }
 
@@ -213,69 +216,10 @@ float3 Fresnel( float3 SpecularColor, float VoH )
 #endif
 }
 
-uint BRDFReverseBits( uint bits )
-{
-#if SM5_PROFILE
-	return reversebits( bits );
-#else
-	bits = ( bits << 16) | ( bits >> 16);
-	bits = ( (bits & 0x00ff00ff) << 8 ) | ( (bits & 0xff00ff00) >> 8 );
-	bits = ( (bits & 0x0f0f0f0f) << 4 ) | ( (bits & 0xf0f0f0f0) >> 4 );
-	bits = ( (bits & 0x33333333) << 2 ) | ( (bits & 0xcccccccc) >> 2 );
-	bits = ( (bits & 0x55555555) << 1 ) | ( (bits & 0xaaaaaaaa) >> 1 );
-	return bits;
-#endif
-}
-
-float2 Hammersley( uint Index, uint NumSamples, uint2 Random )
-{
-	float E1 = frac( (float)Index / NumSamples + float( Random.x & 0xffff ) / (1<<16) );
-	float E2 = float( BRDFReverseBits(Index) ^ Random.y ) * 2.3283064365386963e-10;
-	return float2( E1, E2 );
-}
-
-float3 ImportanceSampleBlinn( float2 E, float Roughness, float3 N )
-{
-	float m = Roughness * Roughness;
-	float n = 2 / (m*m) - 2;
-
-	float Phi = 2 * PI * E.x;
-	// Under OpenGL a very small x value to pow will produce NAN/INF so to ensure that
-	// the importance sampling works specify a larger minimum value.
-	float CosTheta = pow( max(E.y, 0.0001), 1 / (n + 1) );
-	float SinTheta = sqrt( 1 - CosTheta * CosTheta );
-
-	float3 H;
-	H.x = SinTheta * cos( Phi );
-	H.y = SinTheta * sin( Phi );
-	H.z = CosTheta;
-
-	float3 UpVector = abs(N.z) < 0.999 ? float3(0,0,1) : float3(1,0,0);
-	float3 TangentX = normalize( cross( UpVector, N ) );
-	float3 TangentY = cross( N, TangentX );
-	// tangent to world space
-	return TangentX * H.x + TangentY * H.y + N * H.z;
-}
 
-float3 ImportanceSampleGGX( float2 E, float Roughness, float3 N )
-{
-	float m = Roughness * Roughness;
-
-	float Phi = 2 * PI * E.x;
-	float CosTheta = sqrt( (1 - E.y) / ( 1 + (m*m - 1) * E.y ) );
-	float SinTheta = sqrt( 1 - CosTheta * CosTheta );
-
-	float3 H;
-	H.x = SinTheta * cos( Phi );
-	H.y = SinTheta * sin( Phi );
-	H.z = CosTheta;
-
-	float3 UpVector = abs(N.z) < 0.999 ? float3(0,0,1) : float3(1,0,0);
-	float3 TangentX = normalize( cross( UpVector, N ) );
-	float3 TangentY = cross( N, TangentX );
-	// tangent to world space
-	return TangentX * H.x + TangentY * H.y + N * H.z;
-}
+//---------------
+// EnvBRDF
+//---------------
 
 Texture2D		PreIntegratedGF;
 SamplerState	PreIntegratedGFSampler;
diff --git a/Engine/Shaders/DeferredLightPixelShaders.usf b/Engine/Shaders/DeferredLightPixelShaders.usf
--- a/Engine/Shaders/DeferredLightPixelShaders.usf
+++ b/Engine/Shaders/DeferredLightPixelShaders.usf
@@ -36,228 +36,6 @@ FDeferredLightData SetupLightDataForStandardDeferred()
 	return LightData;
 }
 
-bool RayHitSphere( float3 R, float3 SphereCenter, float SphereRadius )
-{
-	float3 ClosestPointOnRay = max( 0, dot( SphereCenter, R ) ) * R;
-	float3 CenterToRay = ClosestPointOnRay - SphereCenter;
-	return dot( CenterToRay, CenterToRay ) <= Square( SphereRadius );
-}
-
-bool RayHitRect( float3 R, float3 RectCenter, float3 RectX, float3 RectY, float3 RectZ, float RectExtentX, float RectExtentY )
-{
-	// Intersect ray with plane
-	float3 PointOnPlane = R * max( 0, dot( RectZ, RectCenter ) / dot( RectZ, R ) );
-
-	bool InExtentX = abs( dot( RectX, PointOnPlane - RectCenter ) ) <= RectExtentX;
-	bool InExtentY = abs( dot( RectY, PointOnPlane - RectCenter ) ) <= RectExtentY;
-	return InExtentX && InExtentY;
-}
-
-float3 UniformSampleCone( float2 E, float3 Axis, float CosThetaMax )
-{
-	float Phi = 2 * PI * E.x;
-	float CosTheta = lerp( CosThetaMax, 1, E.y );
-	float SinTheta = sqrt( 1 - CosTheta * CosTheta );
-
-	float3 L;
-	L.x = SinTheta * cos( Phi );
-	L.y = SinTheta * sin( Phi );
-	L.z = CosTheta;
-
-	float3 UpVector = abs(Axis.z) < 0.999 ? float3(0,0,1) : float3(1,0,0);
-	float3 TangentX = normalize( cross( UpVector, Axis ) );
-	float3 TangentY = cross( Axis, TangentX );
-	// tangent to world space
-	return TangentX * L.x + TangentY * L.y + Axis * L.z;
-}
-
-float3 SampleLight( FDeferredLightData LightData, uint2 Random, float3 LightCenter, float3 SpecularColor, float Roughness, float3 N, float3 V )
-{
-	float NoV = saturate( dot( N, V ) );
-	NoV = max( 0.001, NoV );
-	
-	const float SourceRadius = LightData.SpotAnglesAndSourceRadius.z;
-
-	float DistanceSqr = dot( LightCenter, LightCenter );
-
-	float3 ConeAxis = normalize( LightCenter );
-	float ConeCos = sqrt( 1 - Square( SourceRadius ) / DistanceSqr );
-
-	float3 SpecularLighting = 0;
-
-	const uint NumSamples = 64;
-	for( uint i = 0; i < NumSamples; i++ )
-	{
-		float2 E = Hammersley( i, NumSamples, Random );
-		float3 L = UniformSampleCone( E, ConeAxis, ConeCos );
-
-		float3 H = normalize(V + L);
-		float NoL = saturate( dot(N, L) );
-		float NoV = saturate( dot(N, V) );
-		float NoH = saturate( dot(N, H) );
-		float VoH = saturate( dot(V, H) );
-		
-		if( NoL > 0 )
-		{
-			// Generalized microfacet specular
-			float  D = Distribution( Roughness, NoH ) / PI;
-			float  G = GeometricVisibility( Roughness, NoV, NoL, VoH, L, V );
-			float3 F = Fresnel( SpecularColor, VoH );
-
-			SpecularLighting += NoL / ( DistanceSqr + 1 )  * D * G * F;
-		}
-	}
-
-	return SpecularLighting / NumSamples;
-}
-
-float3 ImportanceSample( FDeferredLightData LightData, uint2 Random, float3 LightCenter, float3 SpecularColor, float Roughness, float3 N, float3 V, float3 R )
-{
-	float NoV = saturate( dot( N, V ) );
-	NoV = max( 0.001, NoV );
-
-	float3 SpecularLighting = 0;
-
-	const uint NumSamples = 64;
-	for( uint i = 0; i < NumSamples; i++ )
-	{
-		float2 E = Hammersley( i, NumSamples, Random );
-		float3 H = ImportanceSampleGGX( E, Roughness, N );
-		float3 L = 2 * dot( V, H ) * H - V;
-
-		float NoL = saturate( dot( N, L ) );
-		float NoH = saturate( dot( N, H ) );
-		float VoH = saturate( dot( V, H ) );
-		
-		if( NoL > 0 )
-		{
-			float SampleColor = 0;
-			
-			if( LIGHT_SOURCE_SHAPE == 1 )
-			{
-				const float SourceRadius = LightData.SpotAnglesAndSourceRadius.z;
-				bool Hit = RayHitSphere( L, LightCenter, SourceRadius );
-				SampleColor = Hit ? (1.0/PI) / Square(SourceRadius) : 0;
-			}
-			else
-			{
-				const float SourceRadiusX = LightData.SpotAnglesAndSourceRadius.z;
-				const float SourceRadiusY = LightData.SpotAnglesAndSourceRadius.w;
-
-				const float3 LightX = float3(1,0,0);
-				const float3 LightY = float3(0,1,0);
-				const float3 LightZ = float3(0,0,1);
-
-				bool Hit = RayHitRect( L, LightCenter, LightX, LightY, LightZ, SourceRadiusX, SourceRadiusY );
-				SampleColor = Hit ? (1.0/PI) / (SourceRadiusX*SourceRadiusY) : 0;
-			}
-
-			float  G = GeometricVisibility( Roughness, NoV, NoL, VoH, L, V );
-			float3 F = Fresnel( SpecularColor, VoH );
-
-			// Incident light = SampleColor * NoL
-			// Microfacet specular = D*G*F / (4*NoL*NoV)
-			// pdf = D * NoH / (4 * VoH)
-			//SpecularLighting += SampleColor * G * F * VoH / (NoH * NoV);
-			SpecularLighting += SampleColor * F * ( G * 4 * VoH * (NoL / NoH) );
-		}
-	}
-
-	return SpecularLighting / NumSamples;
-}
-
-float3 GetDynamicLightingRef(float3 WorldPosition, float3 CameraVector, float2 InUV, FScreenSpaceData ScreenSpaceData, int LightingModelId, FDeferredLightData LightData, float4 LightAttenuation, uint2 Random)
-{
-	float3 N = ScreenSpaceData.GBuffer.WorldNormal;
-	float3 L = LightData.LightDirection;
-	float NoL = BiasedNDotL( dot(N, L) );
-	float DistanceAttenuation = 1;
-	float LightRadiusMask = 1;
-	float SpotFalloff = 1;
-
-	{
-		L = LightData.LightPositionAndInvRadius.xyz - WorldPosition;
-
-		if (LightData.bInverseSquared)
-		{
-			float DistanceSqr = dot( L, L );
-
-			const float SourceLength = LightData.SpotAnglesAndSourceRadius.w;
-
-			BRANCH
-			if( SourceLength > 0 )
-			{
-				// Line segment irradiance
-				float3 Ld = LightData.LightDirection * SourceLength;
-				float3 L0 = L - 0.5 * Ld;
-				float3 L1 = L + 0.5 * Ld;
-				float LengthL0 = length( L0 );
-				float LengthL1 = length( L1 );
-
-				DistanceAttenuation = 2 / ( LengthL0 * LengthL1 + dot( L0, L1 ) + 2 );
-				NoL = saturate( 0.5 * ( dot(N, L0) / LengthL0 + dot(N, L1) / LengthL1 ) );
-			}
-			else
-			{
-				// Sphere irradiance (technically just 1/d^2 but this avoids inf)
-				DistanceAttenuation = 1 / ( DistanceSqr + 1 );
-				NoL = saturate( dot( N, normalize(L) ) );
-				//NoL = BiasedNDotL( dot( N, normalize(L) ) );
-			}
-
-			// TODO scale LightColor
-			// Correction for lumen units
-			DistanceAttenuation *= 16;
-
-			// TODO optimize
-			LightRadiusMask = Square( saturate( 1 - Square( DistanceSqr * Square(LightData.LightPositionAndInvRadius.w) ) ) );
-		}
-
-		if (LightData.bSpotLight)
-		{
-			SpotFalloff = SpotAttenuation(L, -LightData.LightDirection, LightData.SpotAnglesAndSourceRadius.xy);
-		}
-	}
-
-	float3 OutLighting = 0;
-
-	BRANCH
-	if (DistanceAttenuation > 0 && SpotFalloff > 0)
-	{
-		float OpaqueShadowTerm = 1;
-		float SSSShadowTerm = 1;
-
-		if (LightData.bShadowed)
-		{
-			GetShadowTerms(ScreenSpaceData, LightData, WorldPosition, LightAttenuation, OpaqueShadowTerm, SSSShadowTerm);
-		}
-
-		const float3 LightColor = LightData.LightColorAndFalloffExponent.rgb;
-
-		float3 DiffuseLighting = PointLightDiffuse( ScreenSpaceData, L, -CameraVector, N );
-		float3 SpecularLighting = PointLightSpecular( ScreenSpaceData, LightData, L, -CameraVector, N );
-		float3 SubsurfaceLighting = PointLightSubsurface( ScreenSpaceData, LightingModelId, L, -CameraVector, N, SSSShadowTerm );
-
-		float NonShadowedAttenuation = DistanceAttenuation * LightRadiusMask * SpotFalloff;
-		float ShadowedAttenuation = NonShadowedAttenuation * OpaqueShadowTerm;
-		
-		OutLighting += LightColor * ( (NoL * ShadowedAttenuation) * DiffuseLighting + SubsurfaceLighting * NonShadowedAttenuation );
-
-		Random.x ^= View.Random;
-		Random.y ^= View.Random;
-			
-		float3 R = reflect( CameraVector, N );
-		if( ScreenSpaceData.GBuffer.Roughness < 0.2 )
-			SpecularLighting = ImportanceSample( LightData, Random, L, ScreenSpaceData.GBuffer.SpecularColor, ScreenSpaceData.GBuffer.Roughness, N, -CameraVector, R );
-		else
-			SpecularLighting = SampleLight( LightData, Random, L, ScreenSpaceData.GBuffer.SpecularColor, ScreenSpaceData.GBuffer.Roughness, N, -CameraVector );
-
-		OutLighting += 16 * LightColor * SpecularLighting * (LightRadiusMask * SpotFalloff * OpaqueShadowTerm);
-	}
-
-	return OutLighting;
-}
-
 /** 
  * Pixel shader for rendering a directional light using a full screen quad. 
  */
@@ -311,11 +89,11 @@ void RadialPixelMain(
 
 		FDeferredLightData LightData = SetupLightDataForStandardDeferred();
 
-#if 0
-		OutColor.rgb = GetDynamicLightingRef(WorldPosition, CameraVector, ScreenUV, ScreenSpaceData, ScreenSpaceData.GBuffer.LightingModelId, LightData, GetPerPixelLightAttenuation(ScreenUV), ScrambleTEA( uint2( SVPos.xy ) ));
-#else
-		OutColor = GetDynamicLighting(WorldPosition, CameraVector, ScreenUV, ScreenSpaceData, ScreenSpaceData.GBuffer.LightingModelId, LightData, GetPerPixelLightAttenuation(ScreenUV), uint2(0,0));
-#endif
+		uint2 Random = ScrambleTEA( uint2( SVPos.xy ) );
+		Random.x ^= View.Random;
+		Random.y ^= View.Random;
+
+		OutColor = GetDynamicLighting(WorldPosition, CameraVector, ScreenUV, ScreenSpaceData, ScreenSpaceData.GBuffer.LightingModelId, LightData, GetPerPixelLightAttenuation(ScreenUV), Random);
 		OutColor *= ComputeLightProfileMultiplier(WorldPosition, DeferredLightUniforms.LightPosition, DeferredLightUniforms.NormalizedLightDirection);
 	}
 }
\ No newline at end of file
diff --git a/Engine/Shaders/DeferredLightingCommon.usf b/Engine/Shaders/DeferredLightingCommon.usf
--- a/Engine/Shaders/DeferredLightingCommon.usf
+++ b/Engine/Shaders/DeferredLightingCommon.usf
@@ -10,6 +10,7 @@
 #include "DeferredShadingCommon.usf"
 #include "DynamicLightingCommon.usf"
 #include "BRDF.usf"
+#include "MonteCarlo.usf"
 #include "IESLightProfilesCommon.usf"
 
 // 0/1, default is 0 as on NV670 there was a slowdown in average scene, keep for further experiments
@@ -154,11 +155,11 @@ float3 PointLightSpecular( FScreenSpaceData ScreenSpaceData, FDeferredLightData
 	float VoH = saturate( dot(V, H) );
 	
 	// Generalized microfacet specular
-	float  D = Distribution( Roughness, NoH ) / PI;
-	float  G = GeometricVisibility( Roughness, NoV, NoL, VoH, L, V );
+	float D = Distribution( Roughness, NoH );
+	float Vis = GeometricVisibility( Roughness, NoV, NoL, VoH, L, V );
 	float3 F = Fresnel( InGBufferData.SpecularColor, VoH );
 
-	return (Energy * D * G) * F;
+	return (Energy * D * Vis) * F;
 }
 
 float3 SimplePointLightSpecular( FScreenSpaceData ScreenSpaceData, float3 UnitL, float3 V, half3 N )
@@ -173,11 +174,101 @@ float3 SimplePointLightSpecular( FScreenSpaceData ScreenSpaceData, float3 UnitL,
 	float NoH = saturate( dot(N, H) );
 	
 	// Generalized microfacet specular
-	float  D = D_GGX( Roughness, NoH ) / PI;
-	float  G = G_Implicit();
+	float  D = D_GGX( Roughness, NoH );
+	float  Vis = Vis_Implicit();
 	float3 F = F_None( InGBufferData.SpecularColor );
 
-	return (D * G) * F;
+	return (D * Vis) * F;
+}
+
+bool RayHitSphere( float3 R, float3 SphereCenter, float SphereRadius )
+{
+	float3 ClosestPointOnRay = max( 0, dot( SphereCenter, R ) ) * R;
+	float3 CenterToRay = ClosestPointOnRay - SphereCenter;
+	return dot( CenterToRay, CenterToRay ) <= Square( SphereRadius );
+}
+
+bool RayHitRect( float3 R, float3 RectCenter, float3 RectX, float3 RectY, float3 RectZ, float RectExtentX, float RectExtentY )
+{
+	// Intersect ray with plane
+	float3 PointOnPlane = R * max( 0, dot( RectZ, RectCenter ) / dot( RectZ, R ) );
+
+	bool InExtentX = abs( dot( RectX, PointOnPlane - RectCenter ) ) <= RectExtentX;
+	bool InExtentY = abs( dot( RectY, PointOnPlane - RectCenter ) ) <= RectExtentY;
+	return InExtentX && InExtentY;
+}
+
+float3 PointLightSpecularMIS( FScreenSpaceData ScreenSpaceData, FDeferredLightData LightData, float3 LightCenter, float3 V, float3 N, uint2 Random )
+{
+	FGBufferData GBuffer = ScreenSpaceData.GBuffer;
+	float Roughness = GBuffer.Roughness;
+
+	float NoV = saturate( dot( N, V ) );
+	NoV = max( 0.001, NoV );
+	
+	const float SourceRadius = max( 1, LightData.SpotAnglesAndSourceRadius.z );
+
+	const float DistanceSqr = dot( LightCenter, LightCenter );
+	const float3 ConeAxis = normalize( LightCenter );
+	const float ConeCos = sqrt( 1 - Square( SourceRadius ) / DistanceSqr );
+
+	const float SampleColor = (16.0/PI) / Square(SourceRadius);
+
+	float3 SpecularLighting = 0;
+
+	const uint NumSamplesGGX = 16;
+	const uint NumSamplesCone = 16;
+	for( uint i = 0; i < NumSamplesGGX + NumSamplesCone; i++ )
+	{
+		bool bSampleBRDF = i < NumSamplesGGX;
+
+		float3 L, H;
+		if( bSampleBRDF )
+		{
+			float2 E = Hammersley( i, NumSamplesCone, Random );
+			H = TangentToWorld( ImportanceSampleGGX( E, Roughness ).xyz, N );
+			L = 2 * dot( V, H ) * H - V;
+		}
+		else
+		{
+			float2 E = Hammersley( i, NumSamplesGGX, Random );
+			L = TangentToWorld( UniformSampleCone( E, ConeCos ).xyz, ConeAxis );
+			H = normalize(V + L);
+		}
+
+		float NoL = saturate( dot(N, L) );
+		float NoH = saturate( dot(N, H) );
+		float VoH = saturate( dot(V, H) );
+		
+		if( NoL > 0 )
+		{
+			if( bSampleBRDF && !RayHitSphere( L, LightCenter, SourceRadius ) )
+			{
+				continue;
+			}
+			
+			// Generalized microfacet specular
+			float D = D_GGX( Roughness, NoH );
+			float Vis = GeometricVisibility( Roughness, NoV, NoL, VoH, L, V );
+			float3 F = Fresnel( GBuffer.SpecularColor, VoH );
+
+			float ConePDF = 1.0 / ( 2 * PI * (1 - ConeCos) );
+			float GGXPDF = D * NoH / (4 * VoH);
+
+			if( bSampleBRDF )
+			{
+				float Weight = MISWeight( NumSamplesGGX, GGXPDF, NumSamplesCone, ConePDF );
+				SpecularLighting += F * ( SampleColor * NoL * Vis * (4 * VoH / NoH) * Weight );
+			}
+			else
+			{
+				float Weight = MISWeight( NumSamplesCone, ConePDF, NumSamplesGGX, GGXPDF );
+				SpecularLighting += F * ( SampleColor * NoL * Vis * D / ConePDF * Weight );
+			}
+		}
+	}
+
+	return SpecularLighting / (NumSamplesGGX + NumSamplesCone);
 }
 
 float3 PointLightSubsurface( FScreenSpaceData ScreenSpaceData, int LightingModelId, float3 L, float3 V, half3 N, float SubsurfaceExtinction )
@@ -356,13 +447,24 @@ float4 GetDynamicLighting(float3 WorldPosition, float3 CameraVector, float2 InUV
 			float3 DiffuseLighting = PointLightDiffuse( ScreenSpaceData, L, -CameraVector, N );
 			float3 SpecularLighting = PointLightSpecular( ScreenSpaceData, LightData, L, -CameraVector, N );
 			float3 SubsurfaceLighting = PointLightSubsurface( ScreenSpaceData, LightingModelId, L, -CameraVector, N, SSSShadowTerm );
+			
+#define RAY_TRACE 0
+#if RAY_TRACE
+			SpecularLighting = PointLightSpecularMIS( ScreenSpaceData, LightData, L, -CameraVector, N, Random );
+#endif
 
 			float4 LightColor4 = float4(LightColor, Luminance(LightColor));
 			float4 DiffuseLighting4 = float4(DiffuseLighting, 0);
 			float4 SpecularLighting4 = float4(SpecularLighting, Luminance(SpecularLighting));
 			float4 SubsurfaceLighting4 = float4(SubsurfaceLighting, 0);
 
+#if RAY_TRACE
+			OutLighting += LightColor4 * ( (NoL * ShadowedAttenuation) * DiffuseLighting4 + SubsurfaceLighting4 * NonShadowedAttenuation );
+			OutLighting += LightColor4 * SpecularLighting4 * (LightRadiusMask * SpotFalloff * OpaqueShadowTerm);
+#else
 			OutLighting += LightColor4 * ( (NoL * ShadowedAttenuation) * (DiffuseLighting4 + SpecularLighting4) + SubsurfaceLighting4 * NonShadowedAttenuation );
+#endif
+
 			EstimatedCost += 0.4f;	// add the cost of the lighting computations (should sum up to sum to 1)
 		}
 	}
diff --git a/Engine/Shaders/MonteCarlo.usf b/Engine/Shaders/MonteCarlo.usf
new file mode 100644
--- /dev/null
+++ b/Engine/Shaders/MonteCarlo.usf
@@ -0,0 +1,148 @@
+// Copyright 1998-2014 Epic Games, Inc. All Rights Reserved.
+
+/*=============================================================================
+	MonteCarlo.usf: Monte Carlo integration of distributions
+=============================================================================*/
+
+float3 TangentToWorld( float3 Vec, float3 TangentZ )
+{
+	float3 UpVector = abs(TangentZ.z) < 0.999 ? float3(0,0,1) : float3(1,0,0);
+	float3 TangentX = normalize( cross( UpVector, TangentZ ) );
+	float3 TangentY = cross( TangentZ, TangentX );
+	return TangentX * Vec.x + TangentY * Vec.y + TangentZ * Vec.z;
+}
+
+uint ReverseBits32( uint bits )
+{
+#if SM5_PROFILE
+	return reversebits( bits );
+#else
+	bits = ( bits << 16) | ( bits >> 16);
+	bits = ( (bits & 0x00ff00ff) << 8 ) | ( (bits & 0xff00ff00) >> 8 );
+	bits = ( (bits & 0x0f0f0f0f) << 4 ) | ( (bits & 0xf0f0f0f0) >> 4 );
+	bits = ( (bits & 0x33333333) << 2 ) | ( (bits & 0xcccccccc) >> 2 );
+	bits = ( (bits & 0x55555555) << 1 ) | ( (bits & 0xaaaaaaaa) >> 1 );
+	return bits;
+#endif
+}
+
+float2 Hammersley( uint Index, uint NumSamples, uint2 Random )
+{
+	float E1 = frac( (float)Index / NumSamples + float( Random.x & 0xffff ) / (1<<16) );
+	float E2 = float( ReverseBits32(Index) ^ Random.y ) * 2.3283064365386963e-10;
+	return float2( E1, E2 );
+}
+
+float4 UniformSampleSphere( float2 E )
+{
+	float Phi = 2 * PI * E.x;
+	float CosTheta = 1 - 2 * E.y;
+	float SinTheta = sqrt( 1 - CosTheta * CosTheta );
+
+	float3 H;
+	H.x = SinTheta * cos( Phi );
+	H.y = SinTheta * sin( Phi );
+	H.z = CosTheta;
+
+	float PDF = 1.0 / (4 * PI);
+
+	return float4( H, PDF );
+}
+
+float4 UniformSampleHemisphere( float2 E )
+{
+	float Phi = 2 * PI * E.x;
+	float CosTheta = E.y;
+	float SinTheta = sqrt( 1 - CosTheta * CosTheta );
+
+	float3 H;
+	H.x = SinTheta * cos( Phi );
+	H.y = SinTheta * sin( Phi );
+	H.z = CosTheta;
+
+	float PDF = 1.0 / (2 * PI);
+
+	return float4( H, PDF );
+}
+
+float4 CosineSampleHemisphere( float2 E )
+{
+	float Phi = 2 * PI * E.x;
+	float CosTheta = sqrt( E.y );
+	float SinTheta = sqrt( 1 - CosTheta * CosTheta );
+
+	float3 H;
+	H.x = SinTheta * cos( Phi );
+	H.y = SinTheta * sin( Phi );
+	H.z = CosTheta;
+
+	float PDF = CosTheta / PI;
+
+	return float4( H, PDF );
+}
+
+float4 UniformSampleCone( float2 E, float CosThetaMax )
+{
+	float Phi = 2 * PI * E.x;
+	float CosTheta = lerp( CosThetaMax, 1, E.y );
+	float SinTheta = sqrt( 1 - CosTheta * CosTheta );
+
+	float3 L;
+	L.x = SinTheta * cos( Phi );
+	L.y = SinTheta * sin( Phi );
+	L.z = CosTheta;
+
+	float PDF = 1.0 / ( 2 * PI * (1 - CosThetaMax) );
+
+	return float4( L, PDF );
+}
+
+float4 ImportanceSampleBlinn( float2 E, float Roughness )
+{
+	float m = Roughness * Roughness;
+	float n = 2 / (m*m) - 2;
+
+	float Phi = 2 * PI * E.x;
+	float CosTheta = ClampedPow( E.y, 1 / (n + 1) );
+	float SinTheta = sqrt( 1 - CosTheta * CosTheta );
+
+	float3 H;
+	H.x = SinTheta * cos( Phi );
+	H.y = SinTheta * sin( Phi );
+	H.z = CosTheta;
+
+	float D = (n+2)/ (2*PI) * ClampedPow( CosTheta, n );
+	float PDF = D * CosTheta;
+
+	return float4( H, PDF );
+}
+
+float4 ImportanceSampleGGX( float2 E, float Roughness )
+{
+	float m = Roughness * Roughness;
+	float m2 = m * m;
+
+	float Phi = 2 * PI * E.x;
+	float CosTheta = sqrt( (1 - E.y) / ( 1 + (m2 - 1) * E.y ) );
+	float SinTheta = sqrt( 1 - CosTheta * CosTheta );
+
+	float3 H;
+	H.x = SinTheta * cos( Phi );
+	H.y = SinTheta * sin( Phi );
+	H.z = CosTheta;
+	
+	float d = ( CosTheta * m2 - CosTheta ) * CosTheta + 1;
+	float D = m2 / ( PI*d*d );
+	float PDF = D * CosTheta;
+
+	return float4( H, PDF );
+}
+
+// Multiple importance sampling power heuristic of two functions with a power of two. 
+// [Veach 1997, "Robust Monte Carlo Methods for Light Transport Simulation"]
+float MISWeight( uint Num, float PDF, uint OtherNum, float OtherPDF )
+{
+	float Weight = Num * PDF;
+	float OtherWeight = OtherNum * OtherPDF;
+	return Weight * Weight / (Weight * Weight + OtherWeight * OtherWeight);
+}
\ No newline at end of file
diff --git a/Engine/Shaders/PostProcessAmbient.usf b/Engine/Shaders/PostProcessAmbient.usf
--- a/Engine/Shaders/PostProcessAmbient.usf
+++ b/Engine/Shaders/PostProcessAmbient.usf
@@ -10,10 +10,35 @@
 #include "CubemapCommon.usf"
 #include "Random.usf"
 #include "BRDF.usf"
+#include "MonteCarlo.usf"
 
 #define IMPORTANCE_SAMPLE		0
 
 #if 1
+float3 DiffuseIBL( uint2 Random, float3 DiffuseColor, float3 N )
+{
+	float3 DiffuseLighting = 0;
+
+	const uint NumSamples = 32;
+	for( uint i = 0; i < NumSamples; i++ )
+	{
+		float2 E = Hammersley( i, NumSamples, Random );
+		float3 L = TangentToWorld( CosineSampleHemisphere( E ).xyz, N );
+
+		float NoL = saturate( dot( N, L ) );
+		if( NoL > 0 )
+		{
+			float3 SampleColor = AmbientCubemap.SampleLevel( AmbientCubemapSampler, L, 0 ).rgb;
+
+			// lambert = DiffuseColor * NoL / PI
+			// pdf = NoL / PI
+			DiffuseLighting += SampleColor * DiffuseColor;
+		}
+	}
+
+	return DiffuseLighting / NumSamples;
+}
+
 float3 SpecularIBL( uint2 Random, float3 SpecularColor, float Roughness, float3 N, float3 V )
 {
 	float3 SpecularLighting = 0;
@@ -22,7 +47,7 @@ float3 SpecularIBL( uint2 Random, float3 SpecularColor, float Roughness, float3
 	for( uint i = 0; i < NumSamples; i++ )
 	{
 		float2 E = Hammersley( i, NumSamples, Random );
-		float3 H = ImportanceSampleGGX( E, Roughness, N );
+		float3 H = TangentToWorld( ImportanceSampleGGX( E, Roughness ).xyz, N );
 		float3 L = 2 * dot( V, H ) * H - V;
 
 		float NoV = saturate( dot( N, V ) );
@@ -34,16 +59,14 @@ float3 SpecularIBL( uint2 Random, float3 SpecularColor, float Roughness, float3
 		{
 			float3 SampleColor = AmbientCubemap.SampleLevel( AmbientCubemapSampler, L, 0 ).rgb;
 
-			float G = G_Smith( Roughness, NoV, NoL );
+			float Vis = Vis_Smith( Roughness, NoV, NoL );
 			float Fc = pow( 1 - VoH, 5 );
 			float3 F = (1 - Fc) * SpecularColor + Fc;
 
 			// Incident light = SampleColor * NoL
-			// Microfacet specular = D*G*F / (4*NoL*NoV)
+			// Microfacet specular = D*G*F / (4*NoL*NoV) = D*Vis*F
 			// pdf = D * NoH / (4 * VoH)
-			//SpecularLighting += SampleColor * G * F * VoH / (NoH * NoV);
-
-			SpecularLighting += SampleColor * F * ( G * 4 * VoH * (NoL / NoH) );
+			SpecularLighting += SampleColor * F * ( NoL * Vis * (4 * VoH / NoH) );
 		}
 	}
 
@@ -59,7 +82,7 @@ float3 FilterEnvMap( uint2 Random, float Roughness, float3 N, float3 V )
 	for( uint i = 0; i < NumSamples; i++ )
 	{
 		float2 E = Hammersley( i, NumSamples, Random );
-		float3 H = ImportanceSampleGGX( E, Roughness, N );
+		float3 H = TangentToWorld( ImportanceSampleGGX( E, Roughness ).xyz, N );
 		float3 L = 2 * dot( V, H ) * H - V;
 
 		float NoL = saturate( dot( N, L ) );
@@ -82,7 +105,7 @@ float3 PrefilterEnvMap( uint2 Random, float Roughness, float3 R )
 	for( uint i = 0; i < NumSamples; i++ )
 	{
 		float2 E = Hammersley( i, NumSamples, Random );
-		float3 H = ImportanceSampleGGX( E, Roughness, R );
+		float3 H = TangentToWorld( ImportanceSampleGGX( E, Roughness ).xyz, R );
 		float3 L = 2 * dot( R, H ) * H - R;
 
 		float NoL = saturate( dot( R, L ) );
@@ -110,7 +133,7 @@ float2 IntegrateBRDF( uint2 Random, float Roughness, float NoV )
 	for( uint i = 0; i < NumSamples; i++ )
 	{
 		float2 E = Hammersley( i, NumSamples, Random );
-		float3 H = ImportanceSampleGGX( E, Roughness, float3(0,0,1) );
+		float3 H = ImportanceSampleGGX( E, Roughness ).xyz;
 		float3 L = 2 * dot( V, H ) * H - V;
 
 		float NoL = saturate( L.z );
@@ -119,13 +142,16 @@ float2 IntegrateBRDF( uint2 Random, float Roughness, float NoV )
 
 		if( NoL > 0 )
 		{
-			float G = G_Smith( Roughness, NoV, NoL );
-			//float G_Vis = G * VoH / (NoH * NoV);
-			float G_Vis = 4 * G * VoH * (NoL / NoH);
+			float Vis = Vis_Smith( Roughness, NoV, NoL );
+
+			// Incident light = NoL
+			// pdf = D * NoH / (4 * VoH)
+			// NoL * Vis / pdf
+			float NoL_Vis_PDF = NoL * Vis * (4 * VoH / NoH);
 
 			float Fc = pow( 1 - VoH, 5 );
-			A += (1 - Fc) * G_Vis;
-			B += Fc * G_Vis;
+			A += (1 - Fc) * NoL_Vis_PDF;
+			B += Fc * NoL_Vis_PDF;
 		}
 	}
 
@@ -160,6 +186,10 @@ void MainPS(in float4 UVAndScreenPos : TEXCOORD0, out float4 OutColor : SV_Targe
 	int2 PixelPos = int2(UVAndScreenPos.zw * ScreenPosToPixel.xy + ScreenPosToPixel.zw + 0.5f);
 
 	float3 ScreenVector = normalize(mul(float4(ScreenSpacePos, 1, 0), View.ScreenToWorld).xyz);
+	
+	uint2 Random = ScrambleTEA( PixelPos );
+	Random.x ^= uint( View.RealTime * 21389 );
+	Random.y ^= uint( View.RealTime * 49233 );
 
 	float3 DiffuseAndSpecular = 0;
 
@@ -167,23 +197,22 @@ void MainPS(in float4 UVAndScreenPos : TEXCOORD0, out float4 OutColor : SV_Targe
 
 	// Diffuse
 	{
+	#if IMPORTANCE_SAMPLE
+		DiffuseAndSpecular += DiffuseIBL( Random, GBuffer.DiffuseColor, GBuffer.WorldNormal );
+	#else
 		// we want to access the mip with the preconvolved diffuse lighting (coneangle=90 degree)
-
 		DiffuseAndSpecular += GBuffer.DiffuseColor * DiffuseLookup;
+	#endif
 	}
 
 	// Specular
 	{
-		float3 R = reflect( ScreenVector, GBuffer.WorldNormal );
-
 	#if IMPORTANCE_SAMPLE
-		uint2 Random = ScrambleTEA( PixelPos );
-		Random.x ^= uint( View.RealTime * 21389 );
-		Random.y ^= uint( View.RealTime * 49233 );
-
 		DiffuseAndSpecular += SpecularIBL( Random, GBuffer.SpecularColor, GBuffer.Roughness, GBuffer.WorldNormal, -ScreenVector );
 		//DiffuseAndSpecular += ApproximateSpecularIBL( Random, GBuffer.SpecularColor, GBuffer.Roughness, GBuffer.WorldNormal, -ScreenVector );
 	#else
+		float3 R = reflect( ScreenVector, GBuffer.WorldNormal );
+		
 		float AbsoluteSpecularMip = ComputeCubemapMipFromRoughness( GBuffer.Roughness, AmbientCubemapMipAdjust.w );
 		float3 SampleColor = TextureCubeSampleLevel( AmbientCubemap, AmbientCubemapSampler, R, AbsoluteSpecularMip ).rgb;
 
