* fix potential shader error

diff --git a/Engine/Shaders/BasePassPixelShader.usf b/Engine/Shaders/BasePassPixelShader.usf
--- a/Engine/Shaders/BasePassPixelShader.usf
+++ b/Engine/Shaders/BasePassPixelShader.usf
@@ -81,7 +81,7 @@ half3 GetImageBasedReflectionLighting(FMaterialPixelParameters MaterialParameter
 #endif
 	
 	BRANCH
-	if( View.UseLightmaps )
+	if( View.UseLightmaps != 0 )
 	{
 		// Note: make sure this matches the lightmap mixing done on opaque (ReflectionEnvironmentTiledDeferredMain)
 		SpecularIBL.rgb *= Luminance( LowGlossIndirectSpecular );
