Bumped ShaderVersion.usf for CL #2064553

diff --git a/Engine/Shaders/ShaderVersion.usf b/Engine/Shaders/ShaderVersion.usf
--- a/Engine/Shaders/ShaderVersion.usf
+++ b/Engine/Shaders/ShaderVersion.usf
@@ -2,4 +2,4 @@
 // Each time the console command is executed it generates a new GUID. As this file is included
 // in common.usf (which should be included in any shader) it allows to invalidate the shader DDC.
 // 
-// GUID = {0C0ADEAF-7D84-473A-8614-FDA1FD55B435}
+// GUID = {000C0006-0014-0005-0035-02555806A618}
\ No newline at end of file
