Fixed reflections on SSS materials. Improved reflection renormalization.

diff --git a/Engine/Shaders/BasePassForForwardShadingPixelShader.usf b/Engine/Shaders/BasePassForForwardShadingPixelShader.usf
--- a/Engine/Shaders/BasePassForForwardShadingPixelShader.usf
+++ b/Engine/Shaders/BasePassForForwardShadingPixelShader.usf
@@ -133,12 +133,12 @@ void Main(
 #endif
 	
 	half3 Color = 0;
-	half LowGlossIndirectSpecularBrightness = 0;
+	half IndirectIrradiance = 0;
 
 	#if LQ_TEXTURE_LIGHTMAP
 		half4 LightmapColor = GetLightMapColorLQ( GetLightMapCoordinate(Interpolants), MaterialParameters.WorldNormal );
 		Color += LightmapColor.rgb * DiffuseColor;
-		LowGlossIndirectSpecularBrightness = LightmapColor.a;
+		IndirectIrradiance = LightmapColor.a;
 	#elif CACHED_POINT_INDIRECT_LIGHTING
 		#if MATERIALBLENDING_MASKED || MATERIALBLENDING_SOLID
 
@@ -153,7 +153,7 @@ void Main(
 			// Compute diffuse lighting which takes the normal into account
 			half3 DiffuseGI = max(half3(0,0,0), DotSH(PointIndirectLighting, DiffuseTransferSH));
 
-			LowGlossIndirectSpecularBrightness = Luminance(DiffuseGI);
+			IndirectIrradiance = Luminance(DiffuseGI);
 			Color += DiffuseColor * DiffuseGI;
 
 		#else 
@@ -164,7 +164,7 @@ void Main(
 			half3 PointIndirectLighting = IndirectLightingSHCoefficients[0];
 			half3 DiffuseGI = PointIndirectLighting;
 
-			LowGlossIndirectSpecularBrightness = Luminance(DiffuseGI);
+			IndirectIrradiance = Luminance(DiffuseGI);
 			Color += DiffuseColor * DiffuseGI;
 
 		#endif
@@ -185,7 +185,7 @@ void Main(
 		#if !MATERIAL_FULLY_ROUGH
 			half3 SpecularIBL = GetImageBasedReflectionLighting(MaterialParameters, Roughness);
 			// Environment map has been prenormalized, scale by lightmap luminance
-			Color += SpecularIBL * LowGlossIndirectSpecularBrightness * SpecularColor;
+			Color += SpecularIBL * IndirectIrradiance * SpecularColor;
 		#endif
 	#endif
 			 
diff --git a/Engine/Shaders/BasePassPixelShader.usf b/Engine/Shaders/BasePassPixelShader.usf
--- a/Engine/Shaders/BasePassPixelShader.usf
+++ b/Engine/Shaders/BasePassPixelShader.usf
@@ -12,6 +12,7 @@
 #include "LightmapCommon.usf"
 #include "ReflectionEnvironmentShared.usf"
 #include "BRDF.usf"
+#include "Random.usf"
 
 #include "DeferredShadingCommon.usf"
 
@@ -65,7 +66,7 @@ SamplerState ReflectionCubemapSampler;
 int CubemapArrayIndex;
 #endif
 
-half3 GetImageBasedReflectionLighting(FMaterialPixelParameters MaterialParameters, half Roughness, half3 SpecularColor, half3 LowGlossIndirectSpecular)
+half3 GetImageBasedReflectionLighting(FMaterialPixelParameters MaterialParameters, half Roughness, half3 SpecularColor, half IndirectIrradiance)
 {
 #if FEATURE_LEVEL >= FEATURE_LEVEL_SM5
 	half AbsoluteSpecularMip = ComputeReflectionCaptureMipFromRoughness(Roughness);
@@ -85,7 +86,7 @@ half3 GetImageBasedReflectionLighting(FMaterialPixelParameters MaterialParameter
 	if( View.UseLightmaps > 0 )
 	{
 		// Note: make sure this matches the lightmap mixing done on opaque (ReflectionEnvironmentTiledDeferredMain)
-		SpecularIBL.rgb *= Luminance( LowGlossIndirectSpecular );
+		SpecularIBL.rgb *= IndirectIrradiance;
 	}
 	
 	half NoV = saturate(dot(MaterialParameters.WorldNormal, MaterialParameters.CameraVector));
@@ -100,7 +101,7 @@ half3 GetImageBasedReflectionLighting(FMaterialPixelParameters MaterialParameter
 #endif
 
 /** Calculates lighting for translucency. */
-float3 GetTranslucencyLighting(FMaterialPixelParameters MaterialParameters, float3 DiffuseColor, half Roughness, half3 SpecularColor, float3 LowGlossIndirectSpecular, float MaterialAO)
+float3 GetTranslucencyLighting(FMaterialPixelParameters MaterialParameters, float3 DiffuseColor, half Roughness, half3 SpecularColor, float IndirectIrradiance, float MaterialAO)
 {
 	float3 InterpolatedLighting = 0;
 
@@ -232,7 +233,7 @@ float3 GetTranslucencyLighting(FMaterialPixelParameters MaterialParameters, floa
 
 			// Compute diffuse lighting which takes the normal into account
 			float3 DiffuseGI = max(half3(0,0,0), DotSH(PointIndirectLighting, DiffuseTransferSH));
-			LowGlossIndirectSpecular += DiffuseGI;
+			IndirectIrradiance += Luminance(DiffuseGI);
 
 			InterpolatedLighting += DiffuseColor * DiffuseGI * MaterialAO;
 
@@ -241,7 +242,7 @@ float3 GetTranslucencyLighting(FMaterialPixelParameters MaterialParameters, floa
 			float3 PointIndirectLighting = float3(IndirectLightingSHCoefficients[0].x, IndirectLightingSHCoefficients[1].x, IndirectLightingSHCoefficients[2].x);
 			//@todo - why is /2 needed to match directional?
 			float3 DiffuseGI = PointIndirectLighting / (SHAmbientFunction() * 2);
-			LowGlossIndirectSpecular += DiffuseGI;
+			IndirectIrradiance += Luminance(DiffuseGI);
 
 			InterpolatedLighting += DiffuseColor * DiffuseGI * MaterialAO;
 
@@ -249,7 +250,7 @@ float3 GetTranslucencyLighting(FMaterialPixelParameters MaterialParameters, floa
 	#endif
 
 	#if TRANSLUCENCY_LIGHTING_SURFACE
-		InterpolatedLighting += GetImageBasedReflectionLighting(MaterialParameters, Roughness, SpecularColor, LowGlossIndirectSpecular);
+		InterpolatedLighting += GetImageBasedReflectionLighting(MaterialParameters, Roughness, SpecularColor, IndirectIrradiance);
 	#endif
 
 	return InterpolatedLighting;
@@ -361,9 +362,9 @@ float3 GetPrecomputedIndirectLighting(
 	FMaterialPixelParameters MaterialParameters, 
 	FVertexFactoryInterpolantsVSToPS Interpolants,
 	FBasePassInterpolantsVSToPS BasePassInterpolants,
-	out float3 LowGlossIndirectSpecular)
+	out float IndirectIrradiance)
 {
-	LowGlossIndirectSpecular = 0;
+	IndirectIrradiance = 0;
 	float3 Lighting = 0;
 	float2 SkyOcclusionUV = 0;
 
@@ -411,7 +412,7 @@ float3 GetPrecomputedIndirectLighting(
 
 	#endif
 
-	// Sky lighting must contribute to LowGlossIndirectSpecular for ReflectionEnvironment lightmap mixing
+	// Sky lighting must contribute to IndirectIrradiance for ReflectionEnvironment lightmap mixing
 	float3 SkyLighting = GetSkyLighting(MaterialParameters.WorldNormal, SkyOcclusionUV);
 
 	BRANCH
@@ -431,7 +432,7 @@ float3 GetPrecomputedIndirectLighting(
 
 	#if HQ_TEXTURE_LIGHTMAP || LQ_TEXTURE_LIGHTMAP || CACHED_VOLUME_INDIRECT_LIGHTING || CACHED_POINT_INDIRECT_LIGHTING
 		Lighting *= View.IndirectLightingColorScale;
-		LowGlossIndirectSpecular = Lighting;
+		IndirectIrradiance = Luminance(Lighting);
 	#endif
 
 	return Lighting;
@@ -619,11 +620,11 @@ void Main(
 
 
 	half3 Color = 0;
-	float3 LowGlossIndirectSpecular = 0;
+	float IndirectIrradiance = 0;
 
 	#if !MATERIAL_LIGHTINGMODEL_UNLIT
 
-		Color += GetPrecomputedIndirectLighting(MaterialParameters, Interpolants, BasePassInterpolants, LowGlossIndirectSpecular) * DiffuseColor * MaterialAO;
+		Color += GetPrecomputedIndirectLighting(MaterialParameters, Interpolants, BasePassInterpolants, IndirectIrradiance) * DiffuseColor * MaterialAO;
 
 		#if SIMPLE_DYNAMIC_LIGHTING
 			// always unshadowed so BiasedNDotL is not needed
@@ -649,7 +650,7 @@ void Main(
 		
 	// Volume lighting for lit translucency
 	#if (MATERIAL_LIGHTINGMODEL_DEFAULT_LIT || MATERIAL_LIGHTINGMODEL_SUBSURFACE) && (MATERIALBLENDING_TRANSLUCENT || MATERIALBLENDING_ADDITIVE)
-		Color += GetTranslucencyLighting(MaterialParameters, DiffuseColor, Roughness, SpecularColor, LowGlossIndirectSpecular, MaterialAO);
+		Color += GetTranslucencyLighting(MaterialParameters, DiffuseColor, Roughness, SpecularColor, IndirectIrradiance, MaterialAO);
 	#endif
 
 	#if !MATERIAL_LIGHTINGMODEL_UNLIT
@@ -710,7 +711,7 @@ void Main(
 #endif
 		Data.Roughness = Roughness;
 		Data.SubsurfaceColor = SubsurfaceColor;
-		Data.LowGlossIndirectSpecular = LowGlossIndirectSpecular;
+		Data.IndirectIrradiance = IndirectIrradiance;
 		Data.PrecomputedShadowFactors = PrecomputedShadowFactors;
 		Data.GBufferAO = MaterialAO;
 		Data.Opacity = Opacity;
@@ -732,7 +733,8 @@ void Main(
 			float4 OutGBufferE = 0;
 		#endif
 
-		EncodeGBuffer(Data,OutGBufferA,OutGBufferB,OutGBufferC,OutGBufferD,OutGBufferE);
+		float QuantizationBias = PseudoRandom( MaterialParameters.SVPosition ) * 0.5 - 0.5;
+		EncodeGBuffer(Data, OutGBufferA, OutGBufferB, OutGBufferC, OutGBufferD, OutGBufferE, QuantizationBias);
 	#endif 
 	
 
diff --git a/Engine/Shaders/DeferredDecal.usf b/Engine/Shaders/DeferredDecal.usf
--- a/Engine/Shaders/DeferredDecal.usf
+++ b/Engine/Shaders/DeferredDecal.usf
@@ -182,7 +182,7 @@ void MainPS(
 #endif
 	Data.Roughness = GetMaterialRoughness(MaterialParameters);
 	Data.SubsurfaceColor = GetMaterialSubsurfaceColor(MaterialParameters);
-	Data.LowGlossIndirectSpecular = 0;
+	Data.IndirectIrradiance = 0;
 	Data.PrecomputedShadowFactors = 0;
 	Data.GBufferAO = GetMaterialAmbientOcclusion(MaterialParameters);
 	Data.Opacity = Opacity;
diff --git a/Engine/Shaders/DeferredShadingCommon.usf b/Engine/Shaders/DeferredShadingCommon.usf
--- a/Engine/Shaders/DeferredShadingCommon.usf
+++ b/Engine/Shaders/DeferredShadingCommon.usf
@@ -13,6 +13,105 @@
 	#define MSAA_SAMPLE_COUNT 2
 #endif
 
+float3 RGBToYCoCg( float3 RGB )
+{
+	float Y  = dot( RGB, float3(  1, 2,  1 ) ) * 0.25;
+	float Co = dot( RGB, float3(  2, 0, -2 ) ) * 0.25 + ( 0.5 * 256.0 / 255.0 );
+	float Cg = dot( RGB, float3( -1, 2, -1 ) ) * 0.25 + ( 0.5 * 256.0 / 255.0 );
+	
+	float3 YCoCg = float3( Y, Co, Cg );
+	return YCoCg;
+}
+
+float3 YCoCgToRGB( float3 YCoCg )
+{
+	float Y  = YCoCg.x;
+	float Co = YCoCg.y - ( 0.5 * 256.0 / 255.0 );
+	float Cg = YCoCg.z - ( 0.5 * 256.0 / 255.0 );
+
+	float R = Y + Co - Cg;
+	float G = Y + Cg;
+	float B = Y - Co - Cg;
+
+	float3 RGB = float3( R, G, B );
+	return RGB;
+}
+
+// Octahedron Normal Vectors
+// [Cigolle 2014, "A Survey of Efficient Representations for Independent Unit Vectors"]
+//						Mean	Max
+// oct		8:8			0.33709 0.94424
+// snorm	8:8:8		0.17015 0.38588
+// oct		10:10		0.08380 0.23467
+// snorm	10:10:10	0.04228 0.09598
+// oct		12:12		0.02091 0.05874
+
+float2 UnitVectorToOctahedron( float3 N )
+{
+	N.xy /= dot( 1, abs(N) );
+	if( N.z <= 0 )
+	{
+		N.xy = ( 1 - abs(N.yx) ) * ( N.xy >= 0 ? 1 : -1 );
+	}
+	return N.xy;
+}
+
+float3 OctahedronToUnitVector( float2 Oct )
+{
+	float3 N = float3( Oct, 1 - dot( 1, abs(Oct) ) );
+	if( N.z < 0 )
+	{
+		N.xy = ( 1 - abs(N.yx) ) * ( N.xy >= 0 ? 1 : -1 );
+	}
+	return normalize(N);
+}
+
+float3 Pack1212To888( float2 x )
+{
+	// Pack 12:12 to 8:8:8
+#if 0
+	uint2 x1212 = (uint2)( x * 4095.0 );
+	uint2 High = x1212 >> 8;
+	uint2 Low = x1212 & 255;
+	uint3 x888 = uint3( Low, High.x | (High.y << 4) );
+	return x888 / 255.0;
+#else
+	float2 x1212 = floor( x * 4095 );
+	float2 High = floor( x1212 / 256 );	// x1212 >> 8
+	float2 Low = x1212 - High * 256;	// x1212 & 255
+	float3 x888 = float3( Low, High.x + High.y * 16 );
+	return saturate( x888 / 255 );
+#endif
+}
+
+float2 Pack888To1212( float3 x )
+{
+	// Pack 8:8:8 to 12:12
+#if 0
+	uint3 x888 = (uint3)( x * 255.0 );
+	uint High = x888.z >> 4;
+	uint Low = x888.z & 15;
+	uint2 x1212 = x888.xy | uint2( Low << 8, High << 8 );
+	return x1212 / 4095.0;
+#else
+	float3 x888 = floor( x * 255 );
+	float High = floor( x888.z / 16 );	// x888.z >> 4
+	float Low = x888.z - High * 16;		// x888.z & 15
+	float2 x1212 = x888.xy + float2( Low, High ) * 256;
+	return saturate( x1212 / 4095 );
+#endif
+}
+
+float3 EncodeNormal( float3 N )
+{
+	return Pack1212To888( UnitVectorToOctahedron( N ) * 0.5 + 0.5 );
+}
+
+float3 DecodeNormal( float3 N )
+{
+	return OctahedronToUnitVector( Pack888To1212( N )* 2 - 1 );
+}
+
 float3 EncodeSpecularColor(float3 SpecularColor)
 {
 	// Allocate more precision to the darks, which is necessary with bright specular lighting and strong Fresnel
@@ -37,7 +136,6 @@ float3 DecodeDiffuseColor(float3 DiffuseColor)
 	return DiffuseColor;
 }
 
-// TODO Can sRGB writes do this for us?
 float3 EncodeBaseColor(float3 BaseColor)
 {
 	// we use sRGB on the render target to give more precision to the darks
@@ -83,16 +181,19 @@ float3 DecodeSubsurfaceColor(float3 SubsurfaceColor)
 	return Square(SubsurfaceColor);
 }
 
-float3 EncodeIndirectSpecularColor(float3 IndirectSpecularColor)
+float EncodeIndirectIrradiance(float IndirectIrradiance)
 {
-	// Allocate more precision to the darks
-	//@todo - Any indirect lighting higher than 1 is currently clamped
-	return sqrt(saturate(IndirectSpecularColor));
+	float L = IndirectIrradiance;
+	const float LogBlackPoint = 0.00390625;	// exp2(-8);
+	return log2( L + LogBlackPoint ) / 16 + 0.5;
 }
 
-float3 DecodeIndirectSpecularColor(float3 EncodedValue)
+float DecodeIndirectIrradiance(float IndirectIrradiance)
 {
-	return Square(EncodedValue);
+	// LogL -> L
+	float LogL = IndirectIrradiance;
+	const float LogBlackPoint = 0.00390625;	// exp2(-8);
+	return exp2( LogL * 16 - 8 ) - LogBlackPoint;		// 1 exp2, 1 smad, 1 ssub
 }
 
 float ComputeAngleFromRoughness( float Roughness, const float Threshold = 0.04f )
@@ -168,8 +269,8 @@ struct FGBufferData
 	float Specular;
 	// 0..1
 	float3 SubsurfaceColor;
-	// Indirect specular from the lightmap
-	float3 LowGlossIndirectSpecular;
+	// Indirect irradiance luma
+	float IndirectIrradiance;
 	// Static shadow factors for channels assigned by Lightmass
 	// Lights using static shadowing will pick up the appropriate channel in their deferred pass
 	float4 PrecomputedShadowFactors;
@@ -225,7 +326,8 @@ void EncodeGBuffer(
 	out float4 OutGBufferB,
 	out float4 OutGBufferC,
 	out float4 OutGBufferD,
-	out float4 OutGBufferE
+	out float4 OutGBufferE,
+	float QuantizationBias = 0		// -0.5 to 0.5 random float. Used to bias quantization.
 	)
 {
 	OutGBufferA.rgb = float3(Data.WorldNormal * 0.5f + 0.5f);
@@ -241,23 +343,22 @@ void EncodeGBuffer(
 	}
 	else
 	{
+		bool bSubsurface = (Data.LightingModelId == LIGHTINGMODELID_SUBSURFACE || Data.LightingModelId == LIGHTINGMODELID_PREINTEGRATED_SKIN);
+
 #if DIFFUSE_SPEC_INPUTS
-		OutGBufferB = float4(EncodeSpecularColor(Data.SpecularColor), Data.GBufferAO);
+		OutGBufferB.rgb = EncodeSpecularColor(Data.SpecularColor);
 		OutGBufferC.rgb = EncodeDiffuseColor(Data.DiffuseColor);
 #else
 		// NOTE OutGBufferB.b is currently unused!
-		OutGBufferB.r = Data.Metallic;
+		OutGBufferB.r = bSubsurface ? Data.Opacity : Data.Metallic;
 		OutGBufferB.g = Data.Specular;
 		OutGBufferB.b = 0;
-		OutGBufferB.a = Data.GBufferAO;
 		OutGBufferC.rgb = EncodeBaseColor(Data.BaseColor);
 #endif
+		OutGBufferB.a = Encode71(Data.GBufferAO, Data.DecalMask);;
+		OutGBufferC.a = EncodeIndirectIrradiance(Data.IndirectIrradiance) + QuantizationBias * (1.0 / 255.0);
 
-		bool bSubsurface = (Data.LightingModelId == LIGHTINGMODELID_SUBSURFACE || Data.LightingModelId == LIGHTINGMODELID_PREINTEGRATED_SKIN);
-
-		OutGBufferC.a = Encode71(Data.Opacity, Data.DecalMask);
-
-		float3 BufferDColor = bSubsurface ? EncodeSubsurfaceColor(Data.SubsurfaceColor) : EncodeIndirectSpecularColor(Data.LowGlossIndirectSpecular);
+		float3 BufferDColor = bSubsurface ? EncodeSubsurfaceColor(Data.SubsurfaceColor) : float3(0, 0, 0);
 
 		// Roughness in OutGBufferD.r instead of OutGBufferB.a so that deferred decals can blend in roughness while using MRT
 		OutGBufferD = float4(Data.Roughness, BufferDColor.r, BufferDColor.g, BufferDColor.b);
@@ -292,6 +393,10 @@ FGBufferData DecodeGBufferData(
 		ret.WorldNormal = normalize(ret.WorldNormal);
 	}
 
+	ret.LightingModelId = DecodeLightingModelId(InGBufferA);
+	
+	bool bSubsurface = (ret.LightingModelId == LIGHTINGMODELID_SUBSURFACE || ret.LightingModelId == LIGHTINGMODELID_PREINTEGRATED_SKIN);
+
 #if DIFFUSE_SPEC_INPUTS
 	ret.BaseColor = 0;
 	ret.Metallic = 0;
@@ -300,15 +405,15 @@ FGBufferData DecodeGBufferData(
 	ret.DiffuseColor = DecodeDiffuseColor(InGBufferC.rgb);
 	ret.SpecularColor = DecodeSpecularColor(InGBufferB.rgb);
 #else
-	ret.BaseColor = DecodeDiffuseColor(InGBufferC.rgb);
-	ret.Metallic = InGBufferB.r;
+	ret.BaseColor = DecodeBaseColor(InGBufferC.rgb);
+	ret.Metallic = bSubsurface ? 0 : InGBufferB.r;
 	ret.Specular = InGBufferB.g;
 
-	ret.DiffuseColor = lerp( ret.BaseColor, float3(0, 0, 0), ret.Metallic.xxx );
+	ret.DiffuseColor = ret.BaseColor - ret.BaseColor * ret.Metallic;
 	ret.SpecularColor = lerp( 0.08 * ret.Specular.xxx, ret.BaseColor, ret.Metallic );
 
-// todo: COMPILE_SHADERS_FOR_DEVELOPMENT is unfinished feature, using XBOXONE_PROFILE as workaround
-#if COMPILE_SHADERS_FOR_DEVELOPMENT == 1 && !XBOXONE_PROFILE
+	// todo: COMPILE_SHADERS_FOR_DEVELOPMENT is unfinished feature, using XBOXONE_PROFILE as workaround
+	#if COMPILE_SHADERS_FOR_DEVELOPMENT == 1 && !XBOXONE_PROFILE
 	{
 		// this feature is only needed for development/editor - we can compile it out for a shipping build (see r.CompileShadersForDevelopment cvar help)
 		ret.DiffuseColor = ret.DiffuseColor * View.DiffuseOverrideParameter.www + View.DiffuseOverrideParameter.xyz;
@@ -318,13 +423,20 @@ FGBufferData DecodeGBufferData(
 
 #endif
 
+	if( bSubsurface )
+	{
+		ret.SubsurfaceColor = DecodeSubsurfaceColor(InGBufferD.gba);
+		ret.Opacity = InGBufferB.r;
+	}
+	else
+	{
+		ret.SubsurfaceColor = 0;
+		ret.Opacity = 0;
+	}
+	
 	ret.Roughness = InGBufferD.r;
-	ret.LightingModelId = DecodeLightingModelId(InGBufferA);
-	bool bSubsurface = (ret.LightingModelId == LIGHTINGMODELID_SUBSURFACE || ret.LightingModelId == LIGHTINGMODELID_PREINTEGRATED_SKIN);
-	ret.Opacity = Decode71(InGBufferC.a, ret.DecalMask);
-	ret.GBufferAO = InGBufferB.a;
-	ret.SubsurfaceColor = bSubsurface ? DecodeSubsurfaceColor(InGBufferD.gba) : float3(0, 0, 0);
-	ret.LowGlossIndirectSpecular = bSubsurface ? float3(0, 0, 0) : DecodeIndirectSpecularColor(InGBufferD.gba);
+	ret.GBufferAO = Decode71(InGBufferB.a, ret.DecalMask);
+	ret.IndirectIrradiance = DecodeIndirectIrradiance(InGBufferC.a);
 	ret.PrecomputedShadowFactors = InGBufferE;
 	ret.CustomDepth = ConvertFromDeviceZ(CustomNativeDepth);
 	ret.Depth = SceneDepth;
diff --git a/Engine/Shaders/ReflectionEnvironmentComputeShaders.usf b/Engine/Shaders/ReflectionEnvironmentComputeShaders.usf
--- a/Engine/Shaders/ReflectionEnvironmentComputeShaders.usf
+++ b/Engine/Shaders/ReflectionEnvironmentComputeShaders.usf
@@ -8,7 +8,6 @@
 #include "DeferredShadingCommon.usf"  
 #include "BRDF.usf"
 #include "ReflectionEnvironmentShared.usf"
-#include "SkyLightShared.usf"
 
 #if TILED_DEFERRED_CULL_SHADER
 
@@ -578,7 +577,7 @@ void ReflectionEnvironmentTiledDeferredMain(
 		// So, we combine the two for the best of both. This is done by removing the low spatial frequencies from the envmap and replacing them with the lightmap data.
 		// This is only done with luma so as to not get odd color shifting.
 		// Note: make sure this matches the lightmap mixing done for translucency (BasePassPixelShader.usf)
-		SpecularLighting.rgb *= Luminance( GBuffer.LowGlossIndirectSpecular );
+		SpecularLighting.rgb *= GBuffer.IndirectIrradiance;
 #else
 		OutColor.rgb += DiffuseLighting.rgb * GBuffer.DiffuseColor * AO;
 #endif
diff --git a/Engine/Shaders/ReflectionEnvironmentShaders.usf b/Engine/Shaders/ReflectionEnvironmentShaders.usf
--- a/Engine/Shaders/ReflectionEnvironmentShaders.usf
+++ b/Engine/Shaders/ReflectionEnvironmentShaders.usf
@@ -9,7 +9,6 @@
 #include "BRDF.usf"
 #include "SHCommon.usf"
 #include "ReflectionEnvironmentShared.usf"
-#include "SkyLightShared.usf"
 #include "MonteCarlo.usf"
 
 struct FCopyToCubeFaceVSOutput
@@ -367,7 +366,7 @@ void SpecularBouncePS(in float4 UVAndScreenPos : TEXCOORD0, out float4 OutColor
 	FGBufferData GBuffer = ScreenSpaceData.GBuffer;
 
 	// Factors derived from EnvBRDFApprox( SpecularColor, 1, 1 ) == SpecularColor * 0.4524 - 0.0024
-	OutColor = float4( 0.45 * GBuffer.SpecularColor * GBuffer.LowGlossIndirectSpecular, 0 );
+	OutColor = float4( 0.45 * GBuffer.SpecularColor * GBuffer.IndirectIrradiance, 0 );
 }
 
 float4 CapturePositionAndRadius;
@@ -526,7 +525,7 @@ void ReflectionApplyPS(
 			// So, we combine the two for the best of both. This is done by removing the low spatial frequencies from the envmap and replacing them with the lightmap data.
 			// This is only done with luma so as to not get odd color shifting.
 			// Note: make sure this matches the lightmap mixing done for translucency (BasePassPixelShader.usf)
-			SpecularLighting.rgb *= Luminance( GBuffer.LowGlossIndirectSpecular );
+			SpecularLighting.rgb *= GBuffer.IndirectIrradiance;
 		}
 
 #if APPLY_SSR
diff --git a/Engine/Shaders/SkyLightShared.usf b/Engine/Shaders/SkyLightShared.usf
deleted file mode 100644
