Fix for crosscompiler

diff --git a/Engine/Shaders/DeferredShadingCommon.usf b/Engine/Shaders/DeferredShadingCommon.usf
--- a/Engine/Shaders/DeferredShadingCommon.usf
+++ b/Engine/Shaders/DeferredShadingCommon.usf
@@ -51,7 +51,9 @@ float2 UnitVectorToOctahedron( float3 N )
 	N.xy /= dot( 1, abs(N) );
 	if( N.z <= 0 )
 	{
-		N.xy = ( 1 - abs(N.yx) ) * ( N.xy >= 0 ? 1 : -1 );
+		bool2 GreaterThanZero = N.xy >= 0;
+		float2 Value = ( GreaterThanZero ? (float2)1 : (float2)-1 );
+		N.xy = ( 1 - abs(N.yx) ) * Value;
 	}
 	return N.xy;
 }
@@ -61,7 +63,9 @@ float3 OctahedronToUnitVector( float2 Oct )
 	float3 N = float3( Oct, 1 - dot( 1, abs(Oct) ) );
 	if( N.z < 0 )
 	{
-		N.xy = ( 1 - abs(N.yx) ) * ( N.xy >= 0 ? 1 : -1 );
+		bool2 GreaterThanZero = N.xy >= 0;
+		float2 Value = ( GreaterThanZero ? (float2)1 : (float2)-1 );
+		N.xy = ( 1 - abs(N.yx) ) * Value;
 	}
 	return normalize(N);
 }
