* reverted accidential PHYSICAL_SPEC_G change
* added include guard

diff --git a/Engine/Shaders/BRDF.usf b/Engine/Shaders/BRDF.usf
--- a/Engine/Shaders/BRDF.usf
+++ b/Engine/Shaders/BRDF.usf
@@ -4,6 +4,10 @@
 	BRDF.usf: Bidirectional reflectance distribution functions.
 =============================================================================*/
 
+#ifndef __BRDF_COMMON__
+#define __BRDF_COMMON__
+
+
 // Physically based shading model
 // parameterized with the below options
 
@@ -28,7 +32,7 @@
 // 2: Kelemen
 // 3: Schlick
 // 4: Smith (matched to GGX)
-#define PHYSICAL_SPEC_G		4
+#define PHYSICAL_SPEC_G		3
 
 // Fresnel
 // 0: None
@@ -261,4 +265,7 @@ half EnvBRDFApproxNonmetal( half Roughness, half NoV )
 	const half2 c1 = { 1, 0.0425 };
 	half2 r = Roughness * c0 + c1;
 	return min( r.x * r.x, exp2( -9.28 * NoV ) ) * r.x + r.y;
-}
\ No newline at end of file
+}
+
+
+#endif // __BRDF_COMMON__
\ No newline at end of file
