Fix lightmap UVs getting converted from float -> half -> float before being used to sample the lightmap textures.

diff --git a/Engine/Shaders/BasePassForForwardShadingPixelShader.usf b/Engine/Shaders/BasePassForForwardShadingPixelShader.usf
--- a/Engine/Shaders/BasePassForForwardShadingPixelShader.usf
+++ b/Engine/Shaders/BasePassForForwardShadingPixelShader.usf
@@ -136,7 +136,10 @@ void Main(
 	half IndirectIrradiance = 0;
 
 	#if LQ_TEXTURE_LIGHTMAP
-		half4 LightmapColor = GetLightMapColorLQ( GetLightMapCoordinate(Interpolants), MaterialParameters.WorldNormal );
+		float2 LightmapUV0, LightmapUV1;
+		GetLightMapCoordinates(Interpolants, LightmapUV0, LightmapUV1);
+
+		half4 LightmapColor = GetLightMapColorLQ( LightmapUV0, LightmapUV1, MaterialParameters.WorldNormal );
 		Color += LightmapColor.rgb * DiffuseColor;
 		IndirectIrradiance = LightmapColor.a;
 	#elif CACHED_POINT_INDIRECT_LIGHTING
diff --git a/Engine/Shaders/BasePassPixelShader.usf b/Engine/Shaders/BasePassPixelShader.usf
--- a/Engine/Shaders/BasePassPixelShader.usf
+++ b/Engine/Shaders/BasePassPixelShader.usf
@@ -402,13 +402,16 @@ float3 GetPrecomputedIndirectLighting(
 
 	#elif HQ_TEXTURE_LIGHTMAP
 
-		float4 LightmapUV = GetLightMapCoordinate(Interpolants);
-		SkyOcclusionUV = LightmapUV.xy;
-		Lighting = GetLightMapColorHQ(LightmapUV, MaterialParameters.WorldNormal).rgb;
+		float2 LightmapUV0, LightmapUV1;
+		GetLightMapCoordinates(Interpolants, LightmapUV0, LightmapUV1);
+		SkyOcclusionUV = LightmapUV0;
+		Lighting = GetLightMapColorHQ(LightmapUV0, LightmapUV1, MaterialParameters.WorldNormal).rgb;
 
 	#elif LQ_TEXTURE_LIGHTMAP
 
-		Lighting = GetLightMapColorLQ( GetLightMapCoordinate(Interpolants), MaterialParameters.WorldNormal ).rgb;
+		float2 LightmapUV0, LightmapUV1;
+		GetLightMapCoordinates(Interpolants, LightmapUV0, LightmapUV1);
+		Lighting = GetLightMapColorLQ( LightmapUV0, LightmapUV1, MaterialParameters.WorldNormal ).rgb;
 
 	#endif
 
diff --git a/Engine/Shaders/GpuSkinVertexFactory.usf b/Engine/Shaders/GpuSkinVertexFactory.usf
--- a/Engine/Shaders/GpuSkinVertexFactory.usf
+++ b/Engine/Shaders/GpuSkinVertexFactory.usf
@@ -144,9 +144,9 @@ FMaterialPixelParameters GetMaterialPixelParameters(FVertexFactoryInterpolantsVS
 }
 
 #if NEEDS_LIGHTMAP_COORDINATE
-half2 GetLightMapCoordinate(FVertexFactoryInterpolantsVSToPS Interpolants)
+void GetLightMapCoordinates(FVertexFactoryInterpolantsVSToPS Interpolants, out float2 LightmapUV0, out float2 LightmapUV1)
 {
-	return 0;
+	LightmapUV0 = LightmapUV1 = 0;
 }
 
 half2 GetShadowMapCoordinate(FVertexFactoryInterpolantsVSToPS Interpolants)
diff --git a/Engine/Shaders/LandscapeVertexFactory.usf b/Engine/Shaders/LandscapeVertexFactory.usf
--- a/Engine/Shaders/LandscapeVertexFactory.usf
+++ b/Engine/Shaders/LandscapeVertexFactory.usf
@@ -396,12 +396,14 @@ FMaterialVertexParameters GetMaterialVertexParameters( FVertexFactoryInput Input
 }
 
 #if NEEDS_LIGHTMAP_COORDINATE
-float4 GetLightMapCoordinate(FVertexFactoryInterpolantsVSToPS Interpolants)
+void GetLightMapCoordinates(FVertexFactoryInterpolantsVSToPS Interpolants, out float2 LightmapUV0, out float2 LightmapUV1)
 {
 #if ES2_PROFILE
-	return float4( Interpolants.LightMapCoordinate[0].xy, Interpolants.LightMapCoordinate[1].xy );
+	LightmapUV0 = Interpolants.LightMapCoordinate[0].xy;
+	LightmapUV1 = Interpolants.LightMapCoordinate[1].xy;
 #else
-	return Interpolants.LightMapCoordinate.xyxy * half4( 1, 0.5, 1, 0.5 ) + half4( 0, 0, 0, 0.5 );
+	LightmapUV0 = Interpolants.LightMapCoordinate.xy * float2( 1, 0.5 );
+	LightmapUV1 = LightmapUV0 + float2( 0, 0.5 );
 #endif
 }
 
diff --git a/Engine/Shaders/LightMapDensityShader.usf b/Engine/Shaders/LightMapDensityShader.usf
--- a/Engine/Shaders/LightMapDensityShader.usf
+++ b/Engine/Shaders/LightMapDensityShader.usf
@@ -161,11 +161,11 @@ void MainPixelShader(
 	CalcMaterialParameters(MaterialParameters, bIsFrontFace, float4(0,0,.00001f,1));
 	GetMaterialCoverageAndClipping(MaterialParameters);
 
-	float2 LightMapUV;
+	float2 LightMapUV,LightMapUV1;
 #if HQ_TEXTURE_LIGHTMAP || LQ_TEXTURE_LIGHTMAP
-	LightMapUV = GetLightMapCoordinate(FactoryInterpolants).xy;
+	GetLightMapCoordinates(FactoryInterpolants,LightMapUV,LightMapUV1);
 #else
-	LightMapUV = float2(0.1,0.1);
+	LightMapUV = LightMapUV1 = float2(0.1,0.1);
 #endif
 
 	// Area of parallelogram, in world space units.
diff --git a/Engine/Shaders/LightmapCommon.usf b/Engine/Shaders/LightmapCommon.usf
--- a/Engine/Shaders/LightmapCommon.usf
+++ b/Engine/Shaders/LightmapCommon.usf
@@ -15,10 +15,10 @@ SamplerState SkyOcclusionSampler;
 half4 LightMapScale[NUM_LIGHTMAP_COEFFICIENTS];
 half4 LightMapAdd[NUM_LIGHTMAP_COEFFICIENTS];
 
-half4 GetLightMapColorLQ( half4 LightmapUV, half3 WorldNormal )
+half4 GetLightMapColorLQ( float2 LightmapUV0, float2 LightmapUV1, half3 WorldNormal )
 {
-	half4 Lightmap0 = Texture2DSample( LightMapTexture, LightMapSampler, LightmapUV.xy );
-	half4 Lightmap1 = Texture2DSample( LightMapTexture, LightMapSampler, LightmapUV.zw );
+	half4 Lightmap0 = Texture2DSample( LightMapTexture, LightMapSampler, LightmapUV0 );
+	half4 Lightmap1 = Texture2DSample( LightMapTexture, LightMapSampler, LightmapUV1 );
 		
 	// Range scale
 	half3 LogRGB = Lightmap0.rgb * LightMapScale[0].xyz + LightMapAdd[0].xyz;	// 1 vmad
@@ -45,10 +45,10 @@ half4 GetLightMapColorLQ( half4 LightmapUV, half3 WorldNormal )
 	return half4( Color, Luma );
 }
 
-half4 GetLightMapColorHQ( float4 LightmapUV, half3 WorldNormal )
+half4 GetLightMapColorHQ( float2 LightmapUV0, float2 LightmapUV1, half3 WorldNormal )
 {
-	half4 Lightmap0 = Texture2DSample( LightMapTexture, LightMapSampler, LightmapUV.xy );
-	half4 Lightmap1 = Texture2DSample( LightMapTexture, LightMapSampler, LightmapUV.zw );
+	half4 Lightmap0 = Texture2DSample( LightMapTexture, LightMapSampler, LightmapUV0 );
+	half4 Lightmap1 = Texture2DSample( LightMapTexture, LightMapSampler, LightmapUV1 );
 
 	half LogL = Lightmap0.w;
 
diff --git a/Engine/Shaders/LocalVertexFactoryCommon.usf b/Engine/Shaders/LocalVertexFactoryCommon.usf
--- a/Engine/Shaders/LocalVertexFactoryCommon.usf
+++ b/Engine/Shaders/LocalVertexFactoryCommon.usf
@@ -70,9 +70,10 @@ void SetColor(inout FVertexFactoryInterpolantsVSToPS Interpolants, float4 InValu
 }
 
 #if NEEDS_LIGHTMAP_COORDINATE
-float4 GetLightMapCoordinate(FVertexFactoryInterpolantsVSToPS Interpolants)
+void GetLightMapCoordinates(FVertexFactoryInterpolantsVSToPS Interpolants, out float2 LightmapUV0, out float2 LightmapUV1)
 {
-	return Interpolants.LightMapCoordinate.xyxy * half4( 1, 0.5, 1, 0.5 ) + half4( 0, 0, 0, 0.5 );
+	LightmapUV0 = Interpolants.LightMapCoordinate.xy * float2( 1, 0.5 );
+	LightmapUV1 = LightmapUV0 + float2( 0, 0.5 );
 }
 
 float2 GetShadowMapCoordinate(FVertexFactoryInterpolantsVSToPS Interpolants)
@@ -176,12 +177,13 @@ void SetColor(inout FVertexFactoryInterpolantsVSToPS Interpolants, float4 InValu
 	#endif
 }
 
-float4 GetLightMapCoordinate(FVertexFactoryInterpolantsVSToPS Interpolants)
+void GetLightMapCoordinates(FVertexFactoryInterpolantsVSToPS Interpolants, out float2 LightmapUV0, out float2 LightmapUV1)
 {
 	#if USE_INTERPOLATOR_345 
-		return float4(Interpolants.PackedInterpolator3.xy, Interpolants.PackedInterpolator4.xy);
+		LightmapUV0 = Interpolants.PackedInterpolator3.xy;
+		LightmapUV1 = Interpolants.PackedInterpolator4.xy;
 	#else
-		return float4(0, 0, 0, 0);
+		LightmapUV0 = LightmapUV1 = 0;
 	#endif
 }
 
diff --git a/Engine/Shaders/MeshParticleVertexFactory.usf b/Engine/Shaders/MeshParticleVertexFactory.usf
--- a/Engine/Shaders/MeshParticleVertexFactory.usf
+++ b/Engine/Shaders/MeshParticleVertexFactory.usf
@@ -243,9 +243,9 @@ FMaterialVertexParameters GetMaterialVertexParameters( FVertexFactoryInput Input
 }
 
 #if NEEDS_LIGHTMAP_COORDINATE
-half4 GetLightMapCoordinate(FVertexFactoryInterpolantsVSToPS Interpolants)
+void GetLightMapCoordinates(FVertexFactoryInterpolantsVSToPS Interpolants, out float2 LightmapUV0, out float2 LightmapUV1)
 {
-	return 0;
+	LightmapUV0 = LightmapUV1 = 0;
 }
 
 half2 GetShadowMapCoordinate(FVertexFactoryInterpolantsVSToPS Interpolants)
diff --git a/Engine/Shaders/ParticleBeamTrailVertexFactory.usf b/Engine/Shaders/ParticleBeamTrailVertexFactory.usf
--- a/Engine/Shaders/ParticleBeamTrailVertexFactory.usf
+++ b/Engine/Shaders/ParticleBeamTrailVertexFactory.usf
@@ -105,9 +105,9 @@ FMaterialVertexParameters GetMaterialVertexParameters( FVertexFactoryInput Input
 	return Result;
 }
 
-half4 GetLightMapCoordinate(FVertexFactoryInterpolantsVSToPS Interpolants)
+void GetLightMapCoordinates(FVertexFactoryInterpolantsVSToPS Interpolants, out float2 LightmapUV0, out float2 LightmapUV1)
 {
-	return Interpolants.TexCoord.xyxy;
+	LightmapUV0 = LightmapUV1 = 0;
 }
 
 half2 GetShadowMapCoordinate(FVertexFactoryInterpolantsVSToPS Interpolants)
diff --git a/Engine/Shaders/ParticleSpriteVertexFactory.usf b/Engine/Shaders/ParticleSpriteVertexFactory.usf
--- a/Engine/Shaders/ParticleSpriteVertexFactory.usf
+++ b/Engine/Shaders/ParticleSpriteVertexFactory.usf
@@ -234,9 +234,9 @@ FMaterialVertexParameters GetMaterialVertexParameters(FVertexFactoryInput Input,
 	return Result;
 }
 
-half4 GetLightMapCoordinate(FVertexFactoryInterpolantsVSToPS Interpolants)
+void GetLightMapCoordinates(FVertexFactoryInterpolantsVSToPS Interpolants, out float2 LightmapUV0, out float2 LightmapUV1)
 {
-	return Interpolants.TexCoord.xyxy;
+	LightmapUV0 = LightmapUV1 = 0;
 }
 
 half2 GetShadowMapCoordinate(FVertexFactoryInterpolantsVSToPS Interpolants)
