Prevent a half -> float conversion when sampling the IBL cubemap in ES2

diff --git a/Engine/Shaders/BasePassForForwardShadingPixelShader.usf b/Engine/Shaders/BasePassForForwardShadingPixelShader.usf
--- a/Engine/Shaders/BasePassForForwardShadingPixelShader.usf
+++ b/Engine/Shaders/BasePassForForwardShadingPixelShader.usf
@@ -44,7 +44,7 @@ half3 GetImageBasedReflectionLighting(FMaterialPixelParameters MaterialParameter
 	// Compute fractional mip from roughness
 	half AbsoluteSpecularMip = ComputeReflectionCaptureMipFromRoughness(Roughness);
 	// Fetch from cubemap and convert to linear HDR
-	half3 SpecularIBL = RGBMDecode(TextureCubeSampleLevel(ReflectionCubemap, ReflectionCubemapSampler, ProjectedCaptureVector, AbsoluteSpecularMip), 16.0f);
+	half3 SpecularIBL = RGBMDecode(ReflectionCubemap.SampleLevel(ReflectionCubemapSampler, ProjectedCaptureVector, AbsoluteSpecularMip), 16.0f);
 	SpecularIBL *= SpecularIBL;
 #if WEBGL
 	// need a rgb swizzle instead of the existing rgba swizzle, we should add it if another use case comes up. 
