Minor optimization

diff --git a/Engine/Shaders/HeightFogCommon.usf b/Engine/Shaders/HeightFogCommon.usf
--- a/Engine/Shaders/HeightFogCommon.usf
+++ b/Engine/Shaders/HeightFogCommon.usf
@@ -10,21 +10,23 @@ static const float FLT_EPSILON2 = 0.01f;
 /** Shader parameters shared between exponential and multi-layer height fog, to reduce constant register usage. */
 float4 SharedFogParameter0;
 float4 SharedFogParameter1;
-float4 InscatteringLightDirection;
-float4 DirectionalInscatteringColor;
+half4 InscatteringLightDirection;
+half4 DirectionalInscatteringColor;
 float DirectionalInscatteringStartDistance;
 
 /** FogDensity * exp2(-FogHeightFalloff * (CameraWorldPosition.z - FogHeight)) in x, FogHeightFalloff in y, CosTerminatorAngle in z, StartDistance in w. */
 static const float4 ExponentialFogParameters = SharedFogParameter0;
 /** Color to use. */
-static const float3 ExponentialFogColor = SharedFogParameter1.xyz;
-static const float MinFogOpacity = SharedFogParameter1.w;
+static const half3 ExponentialFogColor = SharedFogParameter1.xyz;
+static const half MinFogOpacity = SharedFogParameter1.w;
 
 half4 GetExponentialHeightFog(float3 WorldPosition, float3 InCameraPosition)
 {
 	float3 CameraToReceiver = WorldPosition - InCameraPosition;
-	float CameraToReceiverLength = length(CameraToReceiver);
-	float3 CameraToReceiverNormalized = CameraToReceiver / CameraToReceiverLength;
+	float CameraToReceiverLengthSqr = dot( CameraToReceiver, CameraToReceiver );
+	float CameraToReceiverLengthInv = rsqrt( CameraToReceiverLengthSqr );
+	float CameraToReceiverLength = CameraToReceiverLengthSqr * CameraToReceiverLengthInv;
+	half3 CameraToReceiverNormalized = CameraToReceiver * CameraToReceiverLengthInv;
 
 	// Calculate the line integral of the ray from the camera to the receiver position through the fog density function
 	// The exponential fog density function is d = GlobalDensity * exp(-HeightFalloff * z)
