A workaround for a compile error in PostProcessMotionBlur shader on Mac

diff --git a/Engine/Shaders/PostProcessMotionBlur.usf b/Engine/Shaders/PostProcessMotionBlur.usf
--- a/Engine/Shaders/PostProcessMotionBlur.usf
+++ b/Engine/Shaders/PostProcessMotionBlur.usf
@@ -267,7 +267,7 @@ void SetupPS(float4 UVAndScreenPos : TEXCOORD0,
 // @return 0/1
 float Compute3DCheckerBoard(float3 Pos)
 {
-	int3 TiledWorldPos = frac(Pos) > 0.5f;
+	float3 TiledWorldPos = frac(Pos) > 0.5f;
 	return (float)((uint)dot(float3(1,1,1), TiledWorldPos) % 2);
 }
 
