* more proper shader declaration

diff --git a/Engine/Shaders/PostProcessSelectionOutline.usf b/Engine/Shaders/PostProcessSelectionOutline.usf
--- a/Engine/Shaders/PostProcessSelectionOutline.usf
+++ b/Engine/Shaders/PostProcessSelectionOutline.usf
@@ -9,11 +9,11 @@
 
 #if MSAA_SAMPLE_COUNT > 1
 Texture2DMS<float, MSAA_SAMPLE_COUNT> PostprocessInput1MS;
-Texture2DMS<uint4, MSAA_SAMPLE_COUNT> EditorPrimitivesStencil;
+Texture2DMS<uint2, MSAA_SAMPLE_COUNT> EditorPrimitivesStencil;
 #else
 // note: opengl compiler doesn't like Texture2D<float> so we don't do it
 Texture2D PostprocessInput1MS;
-Texture2D<uint4> EditorPrimitivesStencil;
+Texture2D<uint2> EditorPrimitivesStencil;
 #endif
 
 // RGB:selection color, A:SelectionHighlightIntensity
