Enable DX10 style NaN clamping for PS4.  Remove SafeSaturate for PS4.

diff --git a/Engine/Shaders/Common.usf b/Engine/Shaders/Common.usf
--- a/Engine/Shaders/Common.usf
+++ b/Engine/Shaders/Common.usf
@@ -951,20 +951,12 @@ void DrawRectangle( in float4 InPosition, out float4 OutPosition)
 //SafeSaturate function will remove nan/inf.    
 //Can be expensive, only call when there's a good reason to expect Nans.
 //D3D saturate actually turns Nans -> 1  since it does the min(x, 1) first, and D3D nan rules specify the non-nand operand wins in such a case.  
-//See: http://msdn.microsoft.com/en-us/library/windows/desktop/jj218760(v=vs.85).aspx
-#if PS4_PROFILE
-#define SafeSaturate_Def(type)\
-type SafeSaturate(type In) \
-{\
-	return (isnan(In) || isinf(In)) ? 1 : saturate(In); \
-}
-#else
+//See: http://msdn.microsoft.com/en-us/library/windows/desktop/jj218760(v=vs.85).aspx  
 #define SafeSaturate_Def(type)\
 type SafeSaturate(type In) \
 {\
 	return saturate(In);\
 }
-#endif
 
 SafeSaturate_Def(float)
 SafeSaturate_Def(float2)
