Fixed shader error with subuv

diff --git a/Engine/Shaders/LocalVertexFactory.usf b/Engine/Shaders/LocalVertexFactory.usf
--- a/Engine/Shaders/LocalVertexFactory.usf
+++ b/Engine/Shaders/LocalVertexFactory.usf
@@ -432,9 +432,7 @@ FVertexFactoryInterpolantsVSToPS VertexFactoryGetInterpolantsVSToPS(FVertexFacto
 		SetUV(Interpolants, CoordinateIndex, CustomizedUVs[CoordinateIndex]);
 	}
 
-#endif
-
-#if USE_PARTICLE_SUBUVS
+#elif USE_PARTICLE_SUBUVS
 	SetUV(Interpolants, 0, Input.TexCoords[0]);
 #endif
 
