PS4 - Optimization to remove usage of scratch buffer.

diff --git a/Engine/Shaders/LPVPropagate.usf b/Engine/Shaders/LPVPropagate.usf
--- a/Engine/Shaders/LPVPropagate.usf
+++ b/Engine/Shaders/LPVPropagate.usf
@@ -72,6 +72,7 @@ void GetGeometryVolumeCellWeights( int index,
 	// 8 1-tap weights (corners)
 	{
 		const int destIndex[8] = { 0, 17, 6, 23, 2, 19, 8, 25 };
+		[unroll]
 		for ( int i=0; i<8; i++ )
 		{
 			const float3 direction = directions[ destIndex[i] ];
@@ -86,6 +87,7 @@ void GetGeometryVolumeCellWeights( int index,
 	{
 		const int destIndex[12]      = {  1,	 3,	    5,	   7,	  9,	 11,	14,    16,    18,    20,    22,    24 };
 		const int sourceIndex[12][2] = { {0,4}, {0,2}, {4,6}, {2,6}, {0,1}, {4,5}, {2,3}, {6,7}, {1,5}, {1,3}, {5,7}, {3,7} };
+		[unroll]
 		for ( int i=0; i<12; i++ )
 		{
 			const float3 direction = directions[ destIndex[i] ];
@@ -103,6 +105,7 @@ void GetGeometryVolumeCellWeights( int index,
 	{
 		const int destIndex[6]      = {  4,         10,        12,        13,        15,        21 };
 		const int sourceIndex[6][4] = { {0,2,4,6}, {0,1,4,5}, {0,1,2,3}, {4,5,6,7}, {2,3,6,7}, {1,3,5,7} };
+		[unroll]
 		for ( int i=0; i<6; i++ )
 		{
 			const float3 direction = directions[ destIndex[i] ];
@@ -162,7 +165,9 @@ void CSPropogate(uint3 DTid : SV_DispatchThreadID, uint3 GTid : SV_GroupThreadID
 
 #endif
 
+#if !PS4_PROFILE
 	[loop]
+#endif
 	for ( int i=0; i<n; i++ )
 	{
 		int3 pos = ( DTid+directions[i] ) % 32;
