#UE4:
* fixed shadow artifacs with very small objects
  TTP 335863 UE4: RENDERING: BLOCKER: Directional light preshadows broken

diff --git a/Engine/Shaders/ShadowProjectionPixelShader.usf b/Engine/Shaders/ShadowProjectionPixelShader.usf
--- a/Engine/Shaders/ShadowProjectionPixelShader.usf
+++ b/Engine/Shaders/ShadowProjectionPixelShader.usf
@@ -420,7 +420,7 @@ void Main(
 	ShadowPosition.xy /= ShadowPosition.w;
 	// Clamp pixel depth in light space for shadowing opaque, because areas of the shadow depth buffer that weren't rendered to will have been cleared to 1
 	// We want to force the shadow comparison to result in 'unshadowed' in that case, regardless of whether the pixel being shaded is in front or behind that plane
-	float LightSpacePixelDepthForOpaque = min(ShadowPosition.z, 0.999);
+	float LightSpacePixelDepthForOpaque = min(ShadowPosition.z, 0.99999f);
 	// Must not clamp for SSS shadowing, the subsurface gradient must continue past the far plane
 	float LightSpacePixelDepthForSSS = ShadowPosition.z;
 
