Fixed shader warning

diff --git a/Engine/Shaders/BasePassPixelShader.usf b/Engine/Shaders/BasePassPixelShader.usf
--- a/Engine/Shaders/BasePassPixelShader.usf
+++ b/Engine/Shaders/BasePassPixelShader.usf
@@ -736,7 +736,7 @@ void Main(
 			float4 OutGBufferE = 0;
 		#endif
 
-		float QuantizationBias = PseudoRandom( MaterialParameters.SVPosition ) * 0.5 - 0.5;
+		float QuantizationBias = PseudoRandom( MaterialParameters.SVPosition.xy ) * 0.5 - 0.5;
 		EncodeGBuffer(Data, OutGBufferA, OutGBufferB, OutGBufferC, OutGBufferD, OutGBufferE, QuantizationBias);
 	#endif 
 	
