blending between cloth simulated positions and skinned positions

diff --git a/Engine/Shaders/GpuSkinVertexFactory.usf b/Engine/Shaders/GpuSkinVertexFactory.usf
--- a/Engine/Shaders/GpuSkinVertexFactory.usf
+++ b/Engine/Shaders/GpuSkinVertexFactory.usf
@@ -62,6 +62,8 @@ Texture2D ClothSimulVertsPositions;
 /** Texture from which to read simulated normal of clothing. */
 Texture2D ClothSimulVertsNormals;
 #endif // GPUSKIN_USE_DATA_BUFFERS
+/** blend weight between simulated positions and original key-framed animation */
+float ClothBlendWeight;
 #endif// #if GPUSKIN_APEX_CLOTH
 
 struct FVertexFactoryInput
@@ -207,6 +209,11 @@ struct FVertexFactoryIntermediates
 #if GPUSKIN_PASS_THROUGH
 	uint VertexID;
 #endif
+
+#if GPUSKIN_APEX_CLOTH
+	float3 SimulatedPosition;
+#endif
+
 };
 
 /** Converts from vertex factory specific input to a FMaterialVertexParameters, which is used by vertex shader material inputs. */
@@ -238,8 +245,8 @@ float3 GetClothSimulPosition(int Index)
 {
 #if FEATURE_LEVEL >= FEATURE_LEVEL_SM4
 	#if GPUSKIN_USE_DATA_BUFFERS == 1
-		return ClothSimulVertsPositions[Index].xyz;
-	#else
+	return ClothSimulVertsPositions[Index].xyz;
+#else
 		return ClothSimulVertsPositions.Load(int3(Index % GPUSKIN_CLOTH_TEXTURE_STRIDE, Index / GPUSKIN_CLOTH_TEXTURE_STRIDE, 0));
 	#endif // GPUSKIN_USE_DATA_BUFFERS
 #else
@@ -251,8 +258,8 @@ float3 GetClothSimulNormal(int Index)
 {
 #if FEATURE_LEVEL >= FEATURE_LEVEL_SM4
 	#if GPUSKIN_USE_DATA_BUFFERS == 1
-		return ClothSimulVertsNormals[Index].xyz;
-	#else
+	return ClothSimulVertsNormals[Index].xyz;
+#else
 		return ClothSimulVertsNormals.Load(int3(Index % GPUSKIN_CLOTH_TEXTURE_STRIDE, Index / GPUSKIN_CLOTH_TEXTURE_STRIDE, 0)).xyz;
 	#endif // GPUSKIN_USE_DATA_BUFFERS
 #else
@@ -277,16 +284,7 @@ float3 ClothingPosition( FVertexFactoryInput Input )
 */
 float3 UnpackedPosition( FVertexFactoryInput Input )
 {
-#if GPUSKIN_APEX_CLOTH
-	if( IsSimulatedVertex(Input) )
-	{
-		return ClothingPosition(Input);
-	}
-	else
-#endif //GPUSKIN_APEX_CLOTH
-	{
-		return float3(Input.Position.xyz*MeshExtension + MeshOrigin);
-	}
+	return float3(Input.Position.xyz*MeshExtension + MeshOrigin);
 }
 
 #if GPUSKIN_MORPH_BLEND
@@ -374,13 +372,6 @@ float3 SkinPosition( FVertexFactoryInput Input, FVertexFactoryIntermediates Inte
 #if GPUSKIN_MORPH_BLEND
 	float4 Position = float4(MorphPosition(Input, Intermediates),1);
 #else
-	#if GPUSKIN_APEX_CLOTH
-	if( IsSimulatedVertex(Input) )
-	{
-		return Intermediates.UnpackedPosition;
-	}
-	#endif
-
 	float4 Position = float4(Intermediates.UnpackedPosition,1);
 #endif
 
@@ -435,8 +426,13 @@ float3x3 SkinTangents( FVertexFactoryInput Input, FVertexFactoryIntermediates In
 	half4 TangentZ = TangentBias(Input.TangentZ);
 	
     #if GPUSKIN_APEX_CLOTH
+
+	float3x3 ClothTangent;
+	float TempClothBlendWeight = 0.0f;
 	if( IsSimulatedVertex(Input) )
 	{ 
+		TempClothBlendWeight = ClothBlendWeight;
+
 		float3 A = GetClothSimulPosition(Input.SimulIndices.x);
 		float3 B = GetClothSimulPosition(Input.SimulIndices.y);
 		float3 C = GetClothSimulPosition(Input.SimulIndices.z);
@@ -453,8 +449,8 @@ float3x3 SkinTangents( FVertexFactoryInput Input, FVertexFactoryIntermediates In
 								+ Input.BaryCoordTangent.y*(B+NB*Input.BaryCoordTangent.w)
 								+ Input.BaryCoordTangent.z*(C+NC*Input.BaryCoordTangent.w);
 
-		TangentToLocal[0] = normalize(TangentPosition*MeshExtension + MeshOrigin - Intermediates.UnpackedPosition);
-		TangentToLocal[2] = normalize(NormalPosition*MeshExtension + MeshOrigin - Intermediates.UnpackedPosition);
+		TangentToLocal[0] = normalize(TangentPosition*MeshExtension + MeshOrigin - Intermediates.SimulatedPosition);
+		TangentToLocal[2] = normalize(NormalPosition*MeshExtension + MeshOrigin - Intermediates.SimulatedPosition);
 
 		// cloth data are all in world space so need to change into local space
 		TangentToLocal[0] = mul(TangentToLocal[0], (half3x3)Primitive.WorldToLocal);
@@ -464,7 +460,7 @@ float3x3 SkinTangents( FVertexFactoryInput Input, FVertexFactoryIntermediates In
 	    // and flip vector based on sign of tangent basis determinant
 	    TangentToLocal[1] = cross(TangentToLocal[2], TangentToLocal[0]) * TangentZ.w;
     
-	    return TangentToLocal;
+		ClothTangent = TangentToLocal;
 	}
 	#endif // GPUSKIN_APEX_CLOTH
 
@@ -487,7 +483,12 @@ float3x3 SkinTangents( FVertexFactoryInput Input, FVertexFactoryIntermediates In
 	// derive the new binormal by getting the cross product of the normal and tangent
 	// and flip vector based on sign of tangent basis determinant
 	TangentToLocal[1] = cross(TangentToLocal[2], TangentToLocal[0]) * TangentZ.w;
+
+#if GPUSKIN_APEX_CLOTH
+	return ClothTangent*TempClothBlendWeight + TangentToLocal*(1.0f - TempClothBlendWeight);
+#else
 	return TangentToLocal;
+#endif	// GPUSKIN_APEX_CLOTH
 }
 
 FVertexFactoryIntermediates GetVertexFactoryIntermediates(FVertexFactoryInput Input)
@@ -499,6 +500,17 @@ FVertexFactoryIntermediates GetVertexFactoryIntermediates(FVertexFactoryInput In
 #endif
 	
 	Intermediates.UnpackedPosition = UnpackedPosition(Input);
+#if GPUSKIN_APEX_CLOTH
+	if( IsSimulatedVertex(Input) )
+	{
+		Intermediates.SimulatedPosition = ClothingPosition(Input);
+	}
+	else
+	{
+		Intermediates.SimulatedPosition = Intermediates.UnpackedPosition;
+	}
+	
+#endif
 
 #if GPUSKIN_PASS_THROUGH
 	// Blend Matrix not used
@@ -521,7 +533,13 @@ float4 CalcWorldPosition(FVertexFactoryInput Input, FVertexFactoryIntermediates
 #if GPUSKIN_APEX_CLOTH
 	if( IsSimulatedVertex(Input) )
 	{
-		return float4(Intermediates.UnpackedPosition + View.PreViewTranslation, 1);
+		// skinned positions
+		float4 SkinnedPosition = float4(SkinPosition(Input, Intermediates), 1);
+		float4 TransformedSkinPos = mul(SkinnedPosition, Primitive.LocalToWorld);
+
+		// Intermediates.SimulatedPosition is a clothing position in world coord
+		float3 BlendedPos = Intermediates.SimulatedPosition * ClothBlendWeight + TransformedSkinPos.xyz * (1.0f - ClothBlendWeight);
+		return float4(BlendedPos + View.PreViewTranslation, 1);
 	}
 	else
 #endif
@@ -582,25 +600,25 @@ FVertexFactoryInterpolantsVSToPS VertexFactoryGetInterpolantsVSToPS(FVertexFacto
 // @return previous translated world position
 float4 VertexFactoryGetPreviousWorldPosition(FVertexFactoryInput Input, FVertexFactoryIntermediates Intermediates)
 {
+#if GPUSKIN_PASS_THROUGH
+	uint StreamOffset = GPUSkinCacheStreamFloatOffset + (Intermediates.VertexID * GPUSkinCacheStreamStride);
+	return mul(float4(	GPUSkinCacheStreamBuffer[StreamOffset + GPUSKIN_VB_OFFSET_POSITION + 0],
+						GPUSkinCacheStreamBuffer[StreamOffset + GPUSKIN_VB_OFFSET_POSITION + 1],
+						GPUSkinCacheStreamBuffer[StreamOffset + GPUSKIN_VB_OFFSET_POSITION + 2],
+						1),
+				PreviousLocalToWorld);
+#else
+	float4 PrevSkinPosInWolrd = mul(float4(SkinPreviousPosition(Input, Intermediates),1), PreviousLocalToWorld);
 #if GPUSKIN_APEX_CLOTH
 	if( IsSimulatedVertex(Input) )
 	{
-		return float4(Intermediates.UnpackedPosition + View.PrevPreViewTranslation, 1);
+		float3 BlendedPrevPos = (Intermediates.SimulatedPosition + View.PrevPreViewTranslation)*ClothBlendWeight + PrevSkinPosInWolrd.xyz*(1.0f - ClothBlendWeight);		 
+		return float4(BlendedPrevPos, 1);
 	}
-	else
-#endif
-	{
-#if GPUSKIN_PASS_THROUGH
-		uint StreamOffset = GPUSkinCacheStreamFloatOffset + (Intermediates.VertexID * GPUSkinCacheStreamStride);
-		return mul(float4(	GPUSkinCacheStreamBuffer[StreamOffset + GPUSKIN_VB_OFFSET_POSITION + 0],
-							GPUSkinCacheStreamBuffer[StreamOffset + GPUSKIN_VB_OFFSET_POSITION + 1],
-							GPUSkinCacheStreamBuffer[StreamOffset + GPUSKIN_VB_OFFSET_POSITION + 2],
-							1),
-					PreviousLocalToWorld);
-#else
-		return mul(float4(SkinPreviousPosition(Input, Intermediates),1), PreviousLocalToWorld);
+#endif // GPUSKIN_APEX_CLOTH
+
+	return PrevSkinPosInWolrd;
 #endif
-	}
 }
 
 
diff --git a/Engine/Source/Runtime/Engine/Classes/Components/SkeletalMeshComponent.h b/Engine/Source/Runtime/Engine/Classes/Components/SkeletalMeshComponent.h
--- a/Engine/Source/Runtime/Engine/Classes/Components/SkeletalMeshComponent.h
+++ b/Engine/Source/Runtime/Engine/Classes/Components/SkeletalMeshComponent.h
@@ -326,6 +326,13 @@ public:
 	UPROPERTY(EditAnywhere, BlueprintReadWrite, Category=Clothing)
 	float TeleportRotationThreshold;
 
+	/**
+	 * weight to blend between simulated results and key-framed positions
+	 * if weight is 1.0, shows only cloth simulation results and 0.0 will show only skinned results
+	 */
+	UPROPERTY(EditAnywhere, BlueprintReadWrite, Category = Clothing)
+	float ClothBlendWeight;
+
 	/** Draw the APEX Clothing Normals on clothing sections. */
 	uint32 bDisplayClothingNormals:1;
 
diff --git a/Engine/Source/Runtime/Engine/Private/GPUSkinVertexFactory.cpp b/Engine/Source/Runtime/Engine/Private/GPUSkinVertexFactory.cpp
--- a/Engine/Source/Runtime/Engine/Private/GPUSkinVertexFactory.cpp
+++ b/Engine/Source/Runtime/Engine/Private/GPUSkinVertexFactory.cpp
@@ -690,6 +690,7 @@ public:
 		FGPUSkinVertexFactoryShaderParameters::Bind(ParameterMap);
 		ClothSimulPositionsParameter.Bind(ParameterMap,TEXT("ClothSimulVertsPositions"));
 		ClothSimulNormalsParameter.Bind(ParameterMap,TEXT("ClothSimulVertsNormals"));
+		ClothBlendWeightParameter.Bind(ParameterMap, TEXT("ClothBlendWeight"));
 	}
 	/**
 	* Serialize shader params to an archive
@@ -700,6 +701,7 @@ public:
 		FGPUSkinVertexFactoryShaderParameters::Serialize(Ar);
 		Ar << ClothSimulPositionsParameter;
 		Ar << ClothSimulNormalsParameter;
+		Ar << ClothBlendWeightParameter;
 	}
 
 	virtual void SetMesh(FShader* Shader,const FVertexFactory* VertexFactory,const FSceneView& View,const FMeshBatchElement& BatchElement,uint32 DataFlags) const OVERRIDE
@@ -726,7 +728,7 @@ public:
 				RHISetShaderTexture(Shader->GetVertexShader(), ClothSimulPositionsParameter.GetBaseIndex(), ClothShaderData.GetClothSimulPositionBuffer());
 #endif
 			}
-			
+
 			if(ClothSimulNormalsParameter.IsBound())
 			{
 #if GPUSKIN_USE_DATA_BUFFERS
@@ -735,12 +737,19 @@ public:
 				RHISetShaderTexture(Shader->GetVertexShader(), ClothSimulNormalsParameter.GetBaseIndex(), ClothShaderData.GetClothSimulNormalBuffer());
 #endif
 			}
+
+			SetShaderValue(
+				Shader->GetVertexShader(),
+				ClothBlendWeightParameter,
+				ClothShaderData.ClothBlendWeight
+				);
 		}
 	}
 
 protected:
 	FShaderResourceParameter ClothSimulPositionsParameter;
 	FShaderResourceParameter ClothSimulNormalsParameter;
+	FShaderParameter ClothBlendWeightParameter;
 };
 
 /*-----------------------------------------------------------------------------
diff --git a/Engine/Source/Runtime/Engine/Private/SkeletalMeshComponent.cpp b/Engine/Source/Runtime/Engine/Private/SkeletalMeshComponent.cpp
--- a/Engine/Source/Runtime/Engine/Private/SkeletalMeshComponent.cpp
+++ b/Engine/Source/Runtime/Engine/Private/SkeletalMeshComponent.cpp
@@ -51,6 +51,7 @@ USkeletalMeshComponent::USkeletalMeshComponent(const class FPostConstructInitial
 	ClothTeleportCosineThresholdInRad = FMath::Cos(FMath::DegreesToRadians(TeleportRotationThreshold));
 	ClothTeleportDistThresholdSquared = TeleportDistanceThreshold * TeleportDistanceThreshold;
 	bNeedTeleportAndResetOnceMore = false;
+	ClothBlendWeight = 1.0f;
 #if WITH_CLOTH_COLLISION_DETECTION
 	ClothingCollisionRevision = 0;
 #endif// #if WITH_CLOTH_COLLISION_DETECTION
diff --git a/Engine/Source/Runtime/Engine/Private/SkeletalRenderGPUSkin.cpp b/Engine/Source/Runtime/Engine/Private/SkeletalRenderGPUSkin.cpp
--- a/Engine/Source/Runtime/Engine/Private/SkeletalRenderGPUSkin.cpp
+++ b/Engine/Source/Runtime/Engine/Private/SkeletalRenderGPUSkin.cpp
@@ -298,6 +298,7 @@ void FSkeletalMeshObjectGPUSkin::UpdateDynamicData_RenderThread(FDynamicSkelMesh
 			if( bClothFactory )
 			{				
 				FGPUBaseSkinAPEXClothVertexFactory::ClothShaderType& ClothShaderData = VertexFactoryData.ClothVertexFactories[ChunkIdx]->GetClothShaderData();
+				ClothShaderData.ClothBlendWeight = DynamicData->ClothBlendWeight;
 				int16 ActorIdx = Chunk.CorrespondClothAssetIndex;
 				if( DynamicData->ClothSimulUpdateData.IsValidIndex(ActorIdx) )
 				{
@@ -1040,6 +1041,7 @@ bool FDynamicSkelMeshObjectDataGPUSkin::UpdateClothSimulationData(USkinnedMeshCo
 	USkeletalMeshComponent * SkelMeshComponent = Cast<USkeletalMeshComponent>(InMeshComponent);
 	if(SkelMeshComponent)
 	{
+		ClothBlendWeight = SkelMeshComponent->ClothBlendWeight;
 		SkelMeshComponent->GetUpdateClothSimulationData(ClothSimulUpdateData);
 		return true;
 	}
diff --git a/Engine/Source/Runtime/Engine/Private/SkeletalRenderGPUSkin.h b/Engine/Source/Runtime/Engine/Private/SkeletalRenderGPUSkin.h
--- a/Engine/Source/Runtime/Engine/Private/SkeletalRenderGPUSkin.h
+++ b/Engine/Source/Runtime/Engine/Private/SkeletalRenderGPUSkin.h
@@ -53,6 +53,9 @@ public:
 	/** data for updating cloth section */
 	TArray<FClothSimulData> ClothSimulUpdateData;
 
+	/** a weight factor to blend between simulated positions and skinned positions */	
+	float ClothBlendWeight;
+
 	/**
 	* Compare the given set of active morph targets with the current list to check if different
 	* @param CompareActiveVertexAnims - array of vertex anims to compare
diff --git a/Engine/Source/Runtime/Engine/Public/GPUSkinVertexFactory.h b/Engine/Source/Runtime/Engine/Public/GPUSkinVertexFactory.h
--- a/Engine/Source/Runtime/Engine/Public/GPUSkinVertexFactory.h
+++ b/Engine/Source/Runtime/Engine/Public/GPUSkinVertexFactory.h
@@ -692,6 +692,12 @@ class FGPUBaseSkinAPEXClothVertexFactory
 public:
 	struct ClothShaderType
 	{
+		/**
+		 * weight to blend between simulated positions and key-framed poses
+		 * if ClothBlendWeight is 1.0, it shows only simulated positions and if it is 0.0, it shows only key-framed animation
+		 */
+		float ClothBlendWeight;
+
 		void UpdateClothUniformBuffer(const TArray<FVector4>& InSimulPositions, const TArray<FVector4>& InSimulNormals);
 
 		void ReleaseClothUniformBuffer()
