Increased number of shading models to 256.

diff --git a/Engine/Shaders/BasePassPixelShader.usf b/Engine/Shaders/BasePassPixelShader.usf
--- a/Engine/Shaders/BasePassPixelShader.usf
+++ b/Engine/Shaders/BasePassPixelShader.usf
@@ -713,7 +713,7 @@ void Main(
 		Data.Specular = Specular;
 #endif
 		Data.Roughness = Roughness;
-		Data.SubsurfaceColor = SubsurfaceColor;
+		Data.CustomData = EncodeSubsurfaceColor( SubsurfaceColor );
 		Data.IndirectIrradiance = IndirectIrradiance;
 		Data.PrecomputedShadowFactors = PrecomputedShadowFactors;
 		Data.GBufferAO = MaterialAO;
diff --git a/Engine/Shaders/DeferredDecal.usf b/Engine/Shaders/DeferredDecal.usf
--- a/Engine/Shaders/DeferredDecal.usf
+++ b/Engine/Shaders/DeferredDecal.usf
@@ -181,10 +181,10 @@ void MainPS(
 	Data.Specular = Specular;
 #endif
 	Data.Roughness = GetMaterialRoughness(MaterialParameters);
-	Data.SubsurfaceColor = GetMaterialSubsurfaceColor(MaterialParameters);
+	Data.CustomData = 0;
 	Data.IndirectIrradiance = 0;
 	Data.PrecomputedShadowFactors = 0;
-	Data.GBufferAO = GetMaterialAmbientOcclusion(MaterialParameters);
+	Data.GBufferAO = 1;
 	Data.Opacity = Opacity;
 	Data.LightingModelId = LIGHTINGMODELID_DEFAULT_LIT;
 	Data.DecalMask = 1;
diff --git a/Engine/Shaders/DeferredLightingCommon.usf b/Engine/Shaders/DeferredLightingCommon.usf
--- a/Engine/Shaders/DeferredLightingCommon.usf
+++ b/Engine/Shaders/DeferredLightingCommon.usf
@@ -274,6 +274,7 @@ float3 PointLightSpecularMIS( FScreenSpaceData ScreenSpaceData, FDeferredLightDa
 float3 PointLightSubsurface( FScreenSpaceData ScreenSpaceData, int LightingModelId, float3 L, float3 V, half3 N, float SubsurfaceExtinction )
 {
 	FGBufferData InGBufferData = ScreenSpaceData.GBuffer;
+	float3 SubsurfaceColor = DecodeSubsurfaceColor( InGBufferData.CustomData );
 
 	if (LightingModelId == LIGHTINGMODELID_SUBSURFACE)
 	{
@@ -290,7 +291,7 @@ float3 PointLightSubsurface( FScreenSpaceData ScreenSpaceData, int LightingModel
 		float BackScatter = InGBufferData.GBufferAO * NormalContribution / (PI * 2);
 
 		// lerp to never exceed 1 (energy conserving)
-		return InGBufferData.SubsurfaceColor * (lerp(BackScatter, 1, InScatter) * SubsurfaceExtinction);
+		return SubsurfaceColor * (lerp(BackScatter, 1, InScatter) * SubsurfaceExtinction);
 	}
 	else if (LightingModelId == LIGHTINGMODELID_PREINTEGRATED_SKIN)
 	{
@@ -298,7 +299,7 @@ float3 PointLightSubsurface( FScreenSpaceData ScreenSpaceData, int LightingModel
 
 		float OpacityFactor = InGBufferData.Opacity;
 		float3 PreintegratedBRDF = Texture2DSampleLevel(PreIntegratedBRDF, PreIntegratedBRDFSampler, float2(saturate(dot(N, L) * .5 + .5), 1 - OpacityFactor), 0).rgb;
-		return PreintegratedBRDF * InGBufferData.SubsurfaceColor * SubsurfaceExtinction;
+		return PreintegratedBRDF * SubsurfaceColor * SubsurfaceExtinction;
 	}
 
 	return 0;
diff --git a/Engine/Shaders/DeferredShadingCommon.usf b/Engine/Shaders/DeferredShadingCommon.usf
--- a/Engine/Shaders/DeferredShadingCommon.usf
+++ b/Engine/Shaders/DeferredShadingCommon.usf
@@ -51,9 +51,7 @@ float2 UnitVectorToOctahedron( float3 N )
 	N.xy /= dot( 1, abs(N) );
 	if( N.z <= 0 )
 	{
-		bool2 GreaterThanZero = N.xy >= 0;
-		float2 Value = ( GreaterThanZero ? (float2)1 : (float2)-1 );
-		N.xy = ( 1 - abs(N.yx) ) * Value;
+		N.xy = ( 1 - abs(N.yx) ) * ( N.xy >= 0 ? float2(1,1) : float2(-1,-1) );
 	}
 	return N.xy;
 }
@@ -63,9 +61,7 @@ float3 OctahedronToUnitVector( float2 Oct )
 	float3 N = float3( Oct, 1 - dot( 1, abs(Oct) ) );
 	if( N.z < 0 )
 	{
-		bool2 GreaterThanZero = N.xy >= 0;
-		float2 Value = ( GreaterThanZero ? (float2)1 : (float2)-1 );
-		N.xy = ( 1 - abs(N.yx) ) * Value;
+		N.xy = ( 1 - abs(N.yx) ) * ( N.xy >= 0 ? float2(1,1) : float2(-1,-1) );
 	}
 	return normalize(N);
 }
@@ -108,12 +104,14 @@ float2 Pack888To1212( float3 x )
 
 float3 EncodeNormal( float3 N )
 {
-	return Pack1212To888( UnitVectorToOctahedron( N ) * 0.5 + 0.5 );
+	return N * 0.5 + 0.5;
+	//return Pack1212To888( UnitVectorToOctahedron( N ) * 0.5 + 0.5 );
 }
 
 float3 DecodeNormal( float3 N )
 {
-	return OctahedronToUnitVector( Pack888To1212( N )* 2 - 1 );
+	return N * 2 - 1;
+	//return OctahedronToUnitVector( Pack888To1212( N ) * 2 - 1 );
 }
 
 float3 EncodeSpecularColor(float3 SpecularColor)
@@ -224,6 +222,11 @@ float ComputeRoughnessFromAngle( float Angle, const float Threshold = 0.04f )
 	return Roughness;
 }
 
+float AddAngleToRoughness( float Angle, float Roughness )
+{
+	return saturate( sqrt( Square( Roughness ) + 0.33333 * Angle ) );
+}
+
 // @param Scalar clamped in 0..1 range
 // @param Mask 0..1
 // @return 8bit in range 0..1
@@ -245,11 +248,6 @@ float Decode71(float Scalar, out int Mask)
 	return (Scalar - 0.5f * Mask) * 2.0f;
 }
 
-float AddAngleToRoughness( float Angle, float Roughness )
-{
-	return saturate( sqrt( Square( Roughness ) + 0.33333 * Angle ) );
-}
-
 #define LIGHTINGMODELID_UNLIT				0
 #define LIGHTINGMODELID_DEFAULT_LIT			1
 #define LIGHTINGMODELID_SUBSURFACE			2
@@ -272,7 +270,7 @@ struct FGBufferData
 	// 0..1
 	float Specular;
 	// 0..1
-	float3 SubsurfaceColor;
+	float3 CustomData;
 	// Indirect irradiance luma
 	float IndirectIrradiance;
 	// Static shadow factors for channels assigned by Lightmass
@@ -284,10 +282,10 @@ struct FGBufferData
 	float Opacity;
 	// 0..1 ambient occlusion  e.g.SSAO, wet surface mask, skylight mask, ...
 	float GBufferAO;
-	// 0..3 LIGHTINGMODELID_UNLIT / LIGHTINGMODELID_DEFAULT_LIT / LIGHTINGMODELID_SUBSURFACE / LIGHTINGMODELID_PREINTEGRATED_SKIN
-	int LightingModelId;
+	// 0..255 LIGHTINGMODELID_UNLIT / LIGHTINGMODELID_DEFAULT_LIT / LIGHTINGMODELID_SUBSURFACE / LIGHTINGMODELID_PREINTEGRATED_SKIN
+	uint LightingModelId;
 	// 0..1 decal receiver mask
-	int DecalMask;
+	uint DecalMask;
 	// in world units
 	float CustomDepth;
 	// in unreal units (linear), can be used to reconstruct world position,
@@ -334,10 +332,10 @@ void EncodeGBuffer(
 	float QuantizationBias = 0		// -0.5 to 0.5 random float. Used to bias quantization.
 	)
 {
-	OutGBufferA.rgb = float3(Data.WorldNormal * 0.5f + 0.5f);
+	OutGBufferA.rgb = EncodeNormal( Data.WorldNormal );
 
 	// compress in 2 bits
-	OutGBufferA.a = Data.LightingModelId / 3.0f;
+	OutGBufferA.a = Data.DecalMask * (1.0 / 3.0);
 
 	if (Data.LightingModelId == LIGHTINGMODELID_UNLIT)
 	{
@@ -359,24 +357,19 @@ void EncodeGBuffer(
 		OutGBufferB.b = 0;
 		OutGBufferC.rgb = EncodeBaseColor(Data.BaseColor);
 #endif
-		OutGBufferB.a = Encode71(Data.GBufferAO, Data.DecalMask);;
-		OutGBufferC.a = EncodeIndirectIrradiance(Data.IndirectIrradiance) + QuantizationBias * (1.0 / 255.0);
+		OutGBufferB.a = Data.LightingModelId * (1.0 / 255.0);
 
-		float3 BufferDColor = bSubsurface ? EncodeSubsurfaceColor(Data.SubsurfaceColor) : float3(0, 0, 0);
+		// No space for AO. Multiply IndirectIrradiance by AO instead of storing.
+		OutGBufferC.a = EncodeIndirectIrradiance( Data.IndirectIrradiance * Data.GBufferAO ) + QuantizationBias * (1.0 / 255.0);
 
 		// Roughness in OutGBufferD.r instead of OutGBufferB.a so that deferred decals can blend in roughness while using MRT
-		OutGBufferD = float4(Data.Roughness, BufferDColor.r, BufferDColor.g, BufferDColor.b);
+		OutGBufferD.r = Data.Roughness;
+		OutGBufferD.gba = Data.CustomData;
 	}
 
 	OutGBufferE = Data.PrecomputedShadowFactors;
 }
 
-int DecodeLightingModelId(float4 InGBufferA)
-{
-	// decompress from 2 bits (works with low and high precision GBuffer textures)
-	return (int)(InGBufferA.a * 3.999f);
-}
-
 /** Populates FGBufferData */
 FGBufferData DecodeGBufferData(
 	float4 InGBufferA,
@@ -390,14 +383,15 @@ FGBufferData DecodeGBufferData(
 {
 	FGBufferData ret;
 
-	ret.WorldNormal = InGBufferA.xyz * 2 - 1;
+	ret.WorldNormal = DecodeNormal( InGBufferA.xyz );
 
 	if(bGetNormalizedNormal)
 	{
 		ret.WorldNormal = normalize(ret.WorldNormal);
 	}
 
-	ret.LightingModelId = DecodeLightingModelId(InGBufferA);
+	ret.DecalMask		= (uint)( InGBufferA.a * 3.999 );
+	ret.LightingModelId	= (uint)( InGBufferB.a * 255.999 );
 	
 	bool bSubsurface = (ret.LightingModelId == LIGHTINGMODELID_SUBSURFACE || ret.LightingModelId == LIGHTINGMODELID_PREINTEGRATED_SKIN);
 
@@ -427,19 +421,11 @@ FGBufferData DecodeGBufferData(
 
 #endif
 
-	if( bSubsurface )
-	{
-		ret.SubsurfaceColor = DecodeSubsurfaceColor(InGBufferD.gba);
-		ret.Opacity = InGBufferB.r;
-	}
-	else
-	{
-		ret.SubsurfaceColor = 0;
-		ret.Opacity = 0;
-	}
+	ret.Opacity = bSubsurface ? InGBufferB.r : 0;
 	
 	ret.Roughness = InGBufferD.r;
-	ret.GBufferAO = Decode71(InGBufferB.a, ret.DecalMask);
+	ret.CustomData = InGBufferD.gba;
+	ret.GBufferAO = 1;
 	ret.IndirectIrradiance = DecodeIndirectIrradiance(InGBufferC.a);
 	ret.PrecomputedShadowFactors = InGBufferE;
 	ret.CustomDepth = ConvertFromDeviceZ(CustomNativeDepth);
@@ -460,8 +446,8 @@ FDBufferData DecodeDBufferData(
 	ret.Color = DBufferA.rgb;
 	ret.ColorOpacity = DBufferA.a;
 
-	// UNORM 4 channel
-	ret.WorldNormal = float3(DBufferB.rgb * 2 - 1);
+	// UNORM 4 channel, 128/255 represents 0
+	ret.WorldNormal = DBufferB.rgb * 2 - (256.0 / 255.0);
 	ret.NormalOpacity = DBufferB.a;
 
 	// UNORM 2 channel
diff --git a/Engine/Shaders/MaterialTemplate.usf b/Engine/Shaders/MaterialTemplate.usf
--- a/Engine/Shaders/MaterialTemplate.usf
+++ b/Engine/Shaders/MaterialTemplate.usf
@@ -879,7 +879,7 @@ float4 SceneTextureLookup(float2 UV, int SceneTextureIndex)
 	else if (SceneTextureIndex == 1)	return ScreenSpaceData.GBuffer.Depth;
 	else if (SceneTextureIndex == 2)	return float4(ScreenSpaceData.GBuffer.DiffuseColor, 0);
 	else if (SceneTextureIndex == 3)	return float4(ScreenSpaceData.GBuffer.SpecularColor, 0);
-	else if (SceneTextureIndex == 4)	return float4(ScreenSpaceData.GBuffer.SubsurfaceColor, 0);
+	else if (SceneTextureIndex == 4)	return float4(ScreenSpaceData.GBuffer.CustomData, 0);
 	else if (SceneTextureIndex == 5)	return float4(ScreenSpaceData.GBuffer.BaseColor, 0);
 	else if (SceneTextureIndex == 6)	return ScreenSpaceData.GBuffer.Specular;
 	else if (SceneTextureIndex == 7)	return ScreenSpaceData.GBuffer.Metallic;
@@ -913,7 +913,7 @@ float4 SceneTextureLookup(float2 UV, int SceneTextureIndex)
 		// PPI_SpecularColor
 		case 3: return float4(ScreenSpaceData.GBuffer.SpecularColor, 0);
 		// PPI_SubsurfaceColor
-		case 4: return float4(ScreenSpaceData.GBuffer.SubsurfaceColor, 0);
+		case 4: return float4(ScreenSpaceData.GBuffer.CustomData, 0);
 		// PPI_BaseColor
 		case 5: return float4(ScreenSpaceData.GBuffer.BaseColor, 0);
 		// PPI_Specular
diff --git a/Engine/Shaders/PostProcessAmbient.usf b/Engine/Shaders/PostProcessAmbient.usf
--- a/Engine/Shaders/PostProcessAmbient.usf
+++ b/Engine/Shaders/PostProcessAmbient.usf
@@ -230,6 +230,7 @@ void MainPS(in float4 UVAndScreenPos : TEXCOORD0, out float4 OutColor : SV_Targe
 		// some view dependent and some non view dependent (hard coded)
 		float DependentSplit = 0.5f;
 
+		float3 SubsurfaceColor = DecodeSubsurfaceColor( GBuffer.CustomData );
 		float3 SubsurfaceLighting = 0;
 
 		// view independent (shared lookup for diffuse for better performance
@@ -237,7 +238,7 @@ void MainPS(in float4 UVAndScreenPos : TEXCOORD0, out float4 OutColor : SV_Targe
 		// view dependent (blurriness is hard coded)
 		SubsurfaceLighting += TextureCubeSampleLevel(AmbientCubemap, AmbientCubemapSampler, ScreenVector, AbsoluteDiffuseMip - 2.5f).rgb * AmbientOcclusion * (1.0f - DependentSplit);
 
-		DiffuseAndSpecular += GBuffer.SubsurfaceColor * SubsurfaceLighting;
+		DiffuseAndSpecular += SubsurfaceColor * SubsurfaceLighting;
 	}
 	
 	OutColor = float4(DiffuseAndSpecular * AmbientCubemapColor.rgb, 0);
diff --git a/Engine/Shaders/PostProcessGBufferHints.usf b/Engine/Shaders/PostProcessGBufferHints.usf
--- a/Engine/Shaders/PostProcessGBufferHints.usf
+++ b/Engine/Shaders/PostProcessGBufferHints.usf
@@ -19,7 +19,7 @@ uint ComputeAdvice(FGBufferData BufferData, float3 OriginalSceneColor)
 	{
 		FLATTEN if(BufferData.LightingModelId == LIGHTINGMODELID_SUBSURFACE)
 		{
-			SubsurfaceLum = dot(OneThird.xxx, BufferData.SubsurfaceColor);
+			SubsurfaceLum = dot(OneThird.xxx, DecodeSubsurfaceColor(BufferData.CustomData));
 		}
 
 		float DiffuseLum = dot(OneThird.xxx, BufferData.DiffuseColor);
diff --git a/Engine/Shaders/PostProcessSubsurface.usf b/Engine/Shaders/PostProcessSubsurface.usf
--- a/Engine/Shaders/PostProcessSubsurface.usf
+++ b/Engine/Shaders/PostProcessSubsurface.usf
@@ -155,7 +155,7 @@ void MainPS(float4 UVAndScreenPos : TEXCOORD0, out float4 OutColor : SV_Target0)
 			//		float Subsurface = (1 - ScreenSpaceData.GBuffer.DecalMask) > 0.5f && ScreenSpaceData.GBuffer.Depth < 1000.0f;
 			FLATTEN if(LocalScreenSpaceData.GBuffer.LightingModelId == LIGHTINGMODELID_SUBSURFACE)
 			{
-				float3 Color = ColorOverDistance(LocalScreenSpaceData.GBuffer.SubsurfaceColor, r);
+				float3 Color = ColorOverDistance(DecodeSubsurfaceColor(LocalScreenSpaceData.GBuffer.CustomData), r);
 				float ColorLum = Luminance(Color);
 
 				OutColor += float4(Color, 1) * float4(SampleDiffuseLighting(LocalUV), 1);	
diff --git a/Engine/Shaders/PostProcessVisualizeGBuffer.usf b/Engine/Shaders/PostProcessVisualizeGBuffer.usf
--- a/Engine/Shaders/PostProcessVisualizeGBuffer.usf
+++ b/Engine/Shaders/PostProcessVisualizeGBuffer.usf
@@ -54,7 +54,7 @@ float3 GetGBufferChannelData(FGBufferData GBuffer, float2 UV, int ChannelIndex)
 		case GBUFFER_VISUALIZATION_MODE_DIFFUSE_RGB: return pow( GBuffer.DiffuseColor, 1/2.2 );
 		case GBUFFER_VISUALIZATION_MODE_SPECULAR_RGB: return pow( GBuffer.SpecularColor, 1/2.2 );
 		case GBUFFER_VISUALIZATION_MODE_BASE_RGB: return pow( GBuffer.BaseColor, 1/2.2 );
-		case GBUFFER_VISUALIZATION_MODE_SUBSURFACE_RGB: return pow( GBuffer.SubsurfaceColor, 1/2.2 );
+		case GBUFFER_VISUALIZATION_MODE_SUBSURFACE_RGB: return pow( GBuffer.CustomData, 1/2.2 );
 		case GBUFFER_VISUALIZATION_MODE_WORLD_NORMAL: return GBuffer.WorldNormal * 0.5f + 0.5f;
 		case GBUFFER_VISUALIZATION_MODE_SEP_TRANS_RGB: return Texture2DSample(PostprocessInput1, PostprocessInput1Sampler, UV).rgb*20;
 		case GBUFFER_VISUALIZATION_MODE_OPACITY: return GBuffer.Opacity;
