SubsurfaceColor now accounted for in indirect lighting.

diff --git a/Engine/Shaders/BasePassPixelShader.usf b/Engine/Shaders/BasePassPixelShader.usf
--- a/Engine/Shaders/BasePassPixelShader.usf
+++ b/Engine/Shaders/BasePassPixelShader.usf
@@ -627,7 +627,7 @@ void Main(
 
 	#if !MATERIAL_LIGHTINGMODEL_UNLIT
 
-		Color += GetPrecomputedIndirectLighting(MaterialParameters, Interpolants, BasePassInterpolants, IndirectIrradiance) * DiffuseColor * MaterialAO;
+		Color += GetPrecomputedIndirectLighting(MaterialParameters, Interpolants, BasePassInterpolants, IndirectIrradiance) * (DiffuseColor + SubsurfaceColor) * MaterialAO;
 
 		#if SIMPLE_DYNAMIC_LIGHTING
 			// always unshadowed so BiasedNDotL is not needed
diff --git a/Engine/Shaders/ReflectionEnvironmentComputeShaders.usf b/Engine/Shaders/ReflectionEnvironmentComputeShaders.usf
--- a/Engine/Shaders/ReflectionEnvironmentComputeShaders.usf
+++ b/Engine/Shaders/ReflectionEnvironmentComputeShaders.usf
@@ -579,6 +579,12 @@ void ReflectionEnvironmentTiledDeferredMain(
 		// Note: make sure this matches the lightmap mixing done for translucency (BasePassPixelShader.usf)
 		SpecularLighting.rgb *= GBuffer.IndirectIrradiance;
 #else
+		BRANCH
+		if( GBuffer.LightingModelId == LIGHTINGMODELID_SUBSURFACE )
+		{
+			GBuffer.DiffuseColor += DecodeSubsurfaceColor(GBuffer.CustomData);
+		}
+
 		OutColor.rgb += DiffuseLighting.rgb * GBuffer.DiffuseColor * AO;
 #endif
 
