Renamed Material Lighting Model to Material Shading Model.

diff --git a/Engine/Config/BaseEngine.ini b/Engine/Config/BaseEngine.ini
--- a/Engine/Config/BaseEngine.ini
+++ b/Engine/Config/BaseEngine.ini
@@ -710,6 +710,18 @@ bUseStreamingPause=false
 
 +ActiveClassRedirects=(OldClassName="EModifyFrequency",NewClassName="EComponentMobility")
 
+; MaterialLightingModel -> MaterialShadingModel
++TaggedPropertyRedirects=(ClassName="TextureFactory",OldPropertyName="LightingModel",NewPropertyName="ShadingModel")
++TaggedPropertyRedirects=(ClassName="Material",OldPropertyName="LightingModel",NewPropertyName="ShadingModel")
++TaggedPropertyRedirects=(ClassName="MaterialInstanceBasePropertyOverrides",OldPropertyName="bOverride_LightingModel",NewPropertyName="bOverride_ShadingModel")
++TaggedPropertyRedirects=(ClassName="MaterialInstanceBasePropertyOverrides",OldPropertyName="LightingModel",NewPropertyName="ShadingModel")
++EnumRedirects=(EnumName="ESceneTextureId",OldEnumEntry="PPI_LightingModel",NewEnumEntry="PPI_ShadingModel")
++EnumRedirects=(EnumName="EMaterialShadingModel",OldEnumEntry="MLM_DefaultLit",NewEnumEntry="MSM_DefaultLit")
++EnumRedirects=(EnumName="EMaterialShadingModel",OldEnumEntry="MLM_Unlit",NewEnumEntry="MSM_Unlit")
++EnumRedirects=(EnumName="EMaterialShadingModel",OldEnumEntry="MLM_Subsurface",NewEnumEntry="MSM_Subsurface")
++EnumRedirects=(EnumName="EMaterialShadingModel",OldEnumEntry="MLM_PreintegratedSkin",NewEnumEntry="MSM_PreintegratedSkin")
++ActiveClassRedirects=(OldClassName="EMaterialLightingModel",NewClassName="EMaterialShadingModel")
+
 ; enum change of attach location type
 +EnumRedirects=(EnumName="EAttachLocationType",OldEnumEntry="EAttachLocationType_RelativeOffset",NewEnumEntry="EAttachLocation::KeepRelativeOffset")
 +EnumRedirects=(EnumName="EAttachLocationType",OldEnumEntry="EAttachLocationType_AbsoluteWorld",NewEnumEntry="EAttachLocation::KeepWorldPosition")
@@ -2072,7 +2084,7 @@ CustomDepth=(Material="/Engine/BufferVisualization/CustomDepth.CustomDepth", Nam
 DecalMask=(Material="/Engine/BufferVisualization/DecalMask.DecalMask", Name=LOCTEXT("BaseDecalMaskMat", "Decal Mask"))
 DiffuseColor=(Material="/Engine/BufferVisualization/DiffuseColor.DiffuseColor", Name=LOCTEXT("BaseDiffuseColorMat", "Diffuse Color"))
 FinalImage=(Material="/Engine/BufferVisualization/FinalImage.FinalImage", Name=LOCTEXT("BaseFinalImageMat", "Final Image"))
-LightingModel=(Material="/Engine/BufferVisualization/LightingModel.LightingModel", Name=LOCTEXT("BaseLightingModelMat", "Lighting Model"))
+ShadingModel=(Material="/Engine/BufferVisualization/LightingModel.LightingModel", Name=LOCTEXT("BaseShadingModelMat", "Shading Model"))
 MaterialAO=(Material="/Engine/BufferVisualization/MaterialAO.MaterialAO", Name=LOCTEXT("BaseMaterialAOMat", "Material Ambient Occlusion"))
 Metallic=(Material="/Engine/BufferVisualization/Metallic.Metallic", Name=LOCTEXT("BaseMetallicMat", "Metallic"))
 Opacity=(Material="/Engine/BufferVisualization/Opacity.Opacity", Name=LOCTEXT("BaseOpacityMat", "Opacity"))
diff --git a/Engine/Shaders/BasePassForForwardShadingPixelShader.usf b/Engine/Shaders/BasePassForForwardShadingPixelShader.usf
--- a/Engine/Shaders/BasePassForForwardShadingPixelShader.usf
+++ b/Engine/Shaders/BasePassForForwardShadingPixelShader.usf
@@ -173,7 +173,7 @@ void Main(
 		#endif
 	#endif
 
-	#if !MATERIAL_LIGHTINGMODEL_UNLIT
+	#if !MATERIAL_SHADINGMODEL_UNLIT
 		half Shadow = GetPrimaryPrecomputedShadowMask(Interpolants).r;
 
 		#if SIMPLE_DYNAMIC_LIGHTING
diff --git a/Engine/Shaders/BasePassPixelShader.usf b/Engine/Shaders/BasePassPixelShader.usf
--- a/Engine/Shaders/BasePassPixelShader.usf
+++ b/Engine/Shaders/BasePassPixelShader.usf
@@ -195,7 +195,7 @@ float3 GetTranslucencyLighting(FMaterialPixelParameters MaterialParameters, floa
 
 			float3 BackscatteredLighting = 0;
 
-			#if MATERIAL_LIGHTINGMODEL_SUBSURFACE
+			#if MATERIAL_SHADINGMODEL_SUBSURFACE
 
 				float InScatterPower = GetMaterialTranslucentBackscatteringExponent();
 				// Setup a pow lobe to approximate anisotropic in-scattering near to the light direction
@@ -625,7 +625,7 @@ void Main(
 	half3 Color = 0;
 	float IndirectIrradiance = 0;
 
-	#if !MATERIAL_LIGHTINGMODEL_UNLIT
+	#if !MATERIAL_SHADINGMODEL_UNLIT
 
 		Color += GetPrecomputedIndirectLighting(MaterialParameters, Interpolants, BasePassInterpolants, IndirectIrradiance) * (DiffuseColor + SubsurfaceColor) * MaterialAO;
 
@@ -652,11 +652,11 @@ void Main(
 	#endif
 		
 	// Volume lighting for lit translucency
-	#if (MATERIAL_LIGHTINGMODEL_DEFAULT_LIT || MATERIAL_LIGHTINGMODEL_SUBSURFACE) && (MATERIALBLENDING_TRANSLUCENT || MATERIALBLENDING_ADDITIVE)
+	#if (MATERIAL_SHADINGMODEL_DEFAULT_LIT || MATERIAL_SHADINGMODEL_SUBSURFACE) && (MATERIALBLENDING_TRANSLUCENT || MATERIALBLENDING_ADDITIVE)
 		Color += GetTranslucencyLighting(MaterialParameters, DiffuseColor, Roughness, SpecularColor, IndirectIrradiance, MaterialAO);
 	#endif
 
-	#if !MATERIAL_LIGHTINGMODEL_UNLIT
+	#if !MATERIAL_SHADINGMODEL_UNLIT
 		Color = lerp(Color, DiffuseColor + SpecularColor, View.UnlitViewmodeMask);
 	#endif
 
@@ -720,16 +720,16 @@ void Main(
 		Data.Opacity = Opacity;
 		Data.DecalMask = Primitive.DecalReceiverMask;
 
-		#if MATERIAL_LIGHTINGMODEL_UNLIT
-			Data.LightingModelId = LIGHTINGMODELID_UNLIT;
-		#elif MATERIAL_LIGHTINGMODEL_DEFAULT_LIT
-			Data.LightingModelId = LIGHTINGMODELID_DEFAULT_LIT;
-		#elif MATERIAL_LIGHTINGMODEL_SUBSURFACE
-			Data.LightingModelId = LIGHTINGMODELID_SUBSURFACE;
-		#elif MATERIAL_LIGHTINGMODEL_PREINTEGRATED_SKIN
-			Data.LightingModelId = LIGHTINGMODELID_PREINTEGRATED_SKIN;
+		#if MATERIAL_SHADINGMODEL_UNLIT
+			Data.ShadingModelID = SHADINGMODELID_UNLIT;
+		#elif MATERIAL_SHADINGMODEL_DEFAULT_LIT
+			Data.ShadingModelID = SHADINGMODELID_DEFAULT_LIT;
+		#elif MATERIAL_SHADINGMODEL_SUBSURFACE
+			Data.ShadingModelID = SHADINGMODELID_SUBSURFACE;
+		#elif MATERIAL_SHADINGMODEL_PREINTEGRATED_SKIN
+			Data.ShadingModelID = SHADINGMODELID_PREINTEGRATED_SKIN;
 		#else
-			// missing lighting model, compiler should report LightingModelId is not set
+			// missing shading model, compiler should report ShadingModelID is not set
 		#endif
 
 		#if !ALLOW_STATIC_LIGHTING
diff --git a/Engine/Shaders/DeferredDecal.usf b/Engine/Shaders/DeferredDecal.usf
--- a/Engine/Shaders/DeferredDecal.usf
+++ b/Engine/Shaders/DeferredDecal.usf
@@ -186,7 +186,7 @@ void MainPS(
 	Data.PrecomputedShadowFactors = 0;
 	Data.GBufferAO = 1;
 	Data.Opacity = Opacity;
-	Data.LightingModelId = LIGHTINGMODELID_DEFAULT_LIT;
+	Data.ShadingModelID = SHADINGMODELID_DEFAULT_LIT;
 	Data.DecalMask = 1;
 
 	#if DECAL_RENDERSTAGE == 0
diff --git a/Engine/Shaders/DeferredLightPixelShaders.usf b/Engine/Shaders/DeferredLightPixelShaders.usf
--- a/Engine/Shaders/DeferredLightPixelShaders.usf
+++ b/Engine/Shaders/DeferredLightPixelShaders.usf
@@ -59,7 +59,7 @@ void DirectionalPixelMain(
 		float3 WorldPosition = ScreenVector * SceneDepth + View.ViewOrigin.xyz;
 
 		FDeferredLightData LightData = SetupLightDataForStandardDeferred();
-		OutColor = GetDynamicLighting(WorldPosition, CameraVector, InUV, ScreenSpaceData, ScreenSpaceData.GBuffer.LightingModelId, LightData, GetPerPixelLightAttenuation(InUV), uint2(0,0));
+		OutColor = GetDynamicLighting(WorldPosition, CameraVector, InUV, ScreenSpaceData, ScreenSpaceData.GBuffer.ShadingModelID, LightData, GetPerPixelLightAttenuation(InUV), uint2(0,0));
 	}
 }
 
@@ -93,7 +93,7 @@ void RadialPixelMain(
 		Random.x ^= View.Random;
 		Random.y ^= View.Random;
 
-		OutColor = GetDynamicLighting(WorldPosition, CameraVector, ScreenUV, ScreenSpaceData, ScreenSpaceData.GBuffer.LightingModelId, LightData, GetPerPixelLightAttenuation(ScreenUV), Random);
+		OutColor = GetDynamicLighting(WorldPosition, CameraVector, ScreenUV, ScreenSpaceData, ScreenSpaceData.GBuffer.ShadingModelID, LightData, GetPerPixelLightAttenuation(ScreenUV), Random);
 		OutColor *= ComputeLightProfileMultiplier(WorldPosition, DeferredLightUniforms.LightPosition, DeferredLightUniforms.NormalizedLightDirection);
 	}
 }
\ No newline at end of file
diff --git a/Engine/Shaders/DeferredLightingCommon.usf b/Engine/Shaders/DeferredLightingCommon.usf
--- a/Engine/Shaders/DeferredLightingCommon.usf
+++ b/Engine/Shaders/DeferredLightingCommon.usf
@@ -40,7 +40,7 @@ struct FDeferredLightData
 	bool bShadowed;
 };
 
-/** Data about a single light to be shaded with the simple lighting model, designed for speed and limited feature set. */
+/** Data about a single light to be shaded with the simple shading model, designed for speed and limited feature set. */
 struct FSimpleDeferredLightData
 {
 	float4 LightPositionAndInvRadius;
@@ -271,12 +271,12 @@ float3 PointLightSpecularMIS( FScreenSpaceData ScreenSpaceData, FDeferredLightDa
 	return SpecularLighting / (NumSamplesGGX + NumSamplesCone);
 }
 
-float3 PointLightSubsurface( FScreenSpaceData ScreenSpaceData, int LightingModelId, float3 L, float3 V, half3 N, float SubsurfaceExtinction )
+float3 PointLightSubsurface( FScreenSpaceData ScreenSpaceData, uint ShadingModelID, float3 L, float3 V, half3 N, float SubsurfaceExtinction )
 {
 	FGBufferData InGBufferData = ScreenSpaceData.GBuffer;
 	float3 SubsurfaceColor = DecodeSubsurfaceColor( InGBufferData.CustomData );
 
-	if (LightingModelId == LIGHTINGMODELID_SUBSURFACE)
+	if (ShadingModelID == SHADINGMODELID_SUBSURFACE)
 	{
 		L = normalize( L );
 		float3 H = normalize(V + L);
@@ -293,7 +293,7 @@ float3 PointLightSubsurface( FScreenSpaceData ScreenSpaceData, int LightingModel
 		// lerp to never exceed 1 (energy conserving)
 		return SubsurfaceColor * (lerp(BackScatter, 1, InScatter) * SubsurfaceExtinction);
 	}
-	else if (LightingModelId == LIGHTINGMODELID_PREINTEGRATED_SKIN)
+	else if (ShadingModelID == SHADINGMODELID_PREINTEGRATED_SKIN)
 	{
 		L = normalize( L );
 
@@ -351,7 +351,7 @@ void GetShadowTerms(FScreenSpaceData ScreenSpaceData, FDeferredLightData LightDa
 }
 
 /** Calculates lighting for a given position, normal, etc with a fully featured lighting model designed for quality. */
-float4 GetDynamicLighting(float3 WorldPosition, float3 CameraVector, float2 InUV, FScreenSpaceData ScreenSpaceData, int LightingModelId, FDeferredLightData LightData, float4 LightAttenuation, uint2 Random)
+float4 GetDynamicLighting(float3 WorldPosition, float3 CameraVector, float2 InUV, FScreenSpaceData ScreenSpaceData, uint ShadingModelID, FDeferredLightData LightData, float4 LightAttenuation, uint2 Random)
 {
 	float3 N = ScreenSpaceData.GBuffer.WorldNormal;
 	float3 L = LightData.LightDirection;
@@ -440,14 +440,14 @@ float4 GetDynamicLighting(float3 WorldPosition, float3 CameraVector, float2 InUV
 
 #if LIGHTCULLING_FINE == 1
 		// optimization that should help if there are a lot of shadowed pixels
-		BRANCH if (ShadowedAttenuation > 0 || LightingModelId == LIGHTINGMODELID_SUBSURFACE)
+		BRANCH if (ShadowedAttenuation > 0 || ShadingModelID == SHADINGMODELID_SUBSURFACE)
 #endif
 		{
 			const float3 LightColor = LightData.LightColorAndFalloffExponent.rgb;
 
 			float3 DiffuseLighting = PointLightDiffuse( ScreenSpaceData, L, -CameraVector, N );
 			float3 SpecularLighting = PointLightSpecular( ScreenSpaceData, LightData, L, -CameraVector, N );
-			float3 SubsurfaceLighting = PointLightSubsurface( ScreenSpaceData, LightingModelId, L, -CameraVector, N, SSSShadowTerm );
+			float3 SubsurfaceLighting = PointLightSubsurface( ScreenSpaceData, ShadingModelID, L, -CameraVector, N, SSSShadowTerm );
 			
 #define RAY_TRACE 0
 #if RAY_TRACE
diff --git a/Engine/Shaders/DeferredShadingCommon.usf b/Engine/Shaders/DeferredShadingCommon.usf
--- a/Engine/Shaders/DeferredShadingCommon.usf
+++ b/Engine/Shaders/DeferredShadingCommon.usf
@@ -248,11 +248,11 @@ float Decode71(float Scalar, out int Mask)
 	return (Scalar - 0.5f * Mask) * 2.0f;
 }
 
-#define LIGHTINGMODELID_UNLIT				0
-#define LIGHTINGMODELID_DEFAULT_LIT			1
-#define LIGHTINGMODELID_SUBSURFACE			2
-#define LIGHTINGMODELID_PREINTEGRATED_SKIN	3
-#define LIGHTINGMODELID_NUM					4
+#define SHADINGMODELID_UNLIT				0
+#define SHADINGMODELID_DEFAULT_LIT			1
+#define SHADINGMODELID_SUBSURFACE			2
+#define SHADINGMODELID_PREINTEGRATED_SKIN	3
+#define SHADINGMODELID_NUM					4
 
 // all values that are output by the forward rendering pass
 struct FGBufferData
@@ -282,8 +282,8 @@ struct FGBufferData
 	float Opacity;
 	// 0..1 ambient occlusion  e.g.SSAO, wet surface mask, skylight mask, ...
 	float GBufferAO;
-	// 0..255 LIGHTINGMODELID_UNLIT / LIGHTINGMODELID_DEFAULT_LIT / LIGHTINGMODELID_SUBSURFACE / LIGHTINGMODELID_PREINTEGRATED_SKIN
-	uint LightingModelId;
+	// 0..255 SHADINGMODELID_UNLIT / SHADINGMODELID_DEFAULT_LIT / SHADINGMODELID_SUBSURFACE / SHADINGMODELID_PREINTEGRATED_SKIN
+	uint ShadingModelID;
 	// 0..1 decal receiver mask
 	uint DecalMask;
 	// in world units
@@ -337,7 +337,7 @@ void EncodeGBuffer(
 	// compress in 2 bits
 	OutGBufferA.a = Data.DecalMask * (1.0 / 3.0);
 
-	if (Data.LightingModelId == LIGHTINGMODELID_UNLIT)
+	if (Data.ShadingModelID == SHADINGMODELID_UNLIT)
 	{
 		OutGBufferB = 0;
 		OutGBufferC = 0;
@@ -345,7 +345,7 @@ void EncodeGBuffer(
 	}
 	else
 	{
-		bool bSubsurface = (Data.LightingModelId == LIGHTINGMODELID_SUBSURFACE || Data.LightingModelId == LIGHTINGMODELID_PREINTEGRATED_SKIN);
+		bool bSubsurface = (Data.ShadingModelID == SHADINGMODELID_SUBSURFACE || Data.ShadingModelID == SHADINGMODELID_PREINTEGRATED_SKIN);
 
 #if DIFFUSE_SPEC_INPUTS
 		OutGBufferB.rgb = EncodeSpecularColor(Data.SpecularColor);
@@ -357,7 +357,7 @@ void EncodeGBuffer(
 		OutGBufferB.b = 0;
 		OutGBufferC.rgb = EncodeBaseColor(Data.BaseColor);
 #endif
-		OutGBufferB.a = Data.LightingModelId * (1.0 / 255.0);
+		OutGBufferB.a = Data.ShadingModelID * (1.0 / 255.0);
 
 		// No space for AO. Multiply IndirectIrradiance by AO instead of storing.
 		OutGBufferC.a = EncodeIndirectIrradiance( Data.IndirectIrradiance * Data.GBufferAO ) + QuantizationBias * (1.0 / 255.0);
@@ -391,9 +391,9 @@ FGBufferData DecodeGBufferData(
 	}
 
 	ret.DecalMask		= (uint)( InGBufferA.a * 3.999 );
-	ret.LightingModelId	= (uint)( InGBufferB.a * 255.999 );
+	ret.ShadingModelID	= (uint)( InGBufferB.a * 255.999 );
 	
-	bool bSubsurface = (ret.LightingModelId == LIGHTINGMODELID_SUBSURFACE || ret.LightingModelId == LIGHTINGMODELID_PREINTEGRATED_SKIN);
+	bool bSubsurface = (ret.ShadingModelID == SHADINGMODELID_SUBSURFACE || ret.ShadingModelID == SHADINGMODELID_PREINTEGRATED_SKIN);
 
 #if DIFFUSE_SPEC_INPUTS
 	ret.BaseColor = 0;
diff --git a/Engine/Shaders/Definitions.usf b/Engine/Shaders/Definitions.usf
--- a/Engine/Shaders/Definitions.usf
+++ b/Engine/Shaders/Definitions.usf
@@ -32,20 +32,20 @@
 #define MATERIALBLENDING_MODULATE						0
 #endif
 
-#ifndef MATERIAL_LIGHTINGMODEL_DEFAULT_LIT
-#define MATERIAL_LIGHTINGMODEL_DEFAULT_LIT				0
+#ifndef MATERIAL_SHADINGMODEL_DEFAULT_LIT
+#define MATERIAL_SHADINGMODEL_DEFAULT_LIT				0
 #endif
 
-#ifndef MATERIAL_LIGHTINGMODEL_SUBSURFACE
-#define MATERIAL_LIGHTINGMODEL_SUBSURFACE				0
+#ifndef MATERIAL_SHADINGMODEL_SUBSURFACE
+#define MATERIAL_SHADINGMODEL_SUBSURFACE				0
 #endif
 
-#ifndef MATERIAL_LIGHTINGMODEL_PREINTEGRATED_SKIN
-#define MATERIAL_LIGHTINGMODEL_PREINTEGRATED_SKIN		0
+#ifndef MATERIAL_SHADINGMODEL_PREINTEGRATED_SKIN
+#define MATERIAL_SHADINGMODEL_PREINTEGRATED_SKIN		0
 #endif
 
-#ifndef MATERIAL_LIGHTINGMODEL_UNLIT
-#define	MATERIAL_LIGHTINGMODEL_UNLIT					0
+#ifndef MATERIAL_SHADINGMODEL_UNLIT
+#define	MATERIAL_SHADINGMODEL_UNLIT						0
 #endif
 
 #ifndef MESH_MATERIAL_SHADER
diff --git a/Engine/Shaders/LocalVertexFactoryCommon.usf b/Engine/Shaders/LocalVertexFactoryCommon.usf
--- a/Engine/Shaders/LocalVertexFactoryCommon.usf
+++ b/Engine/Shaders/LocalVertexFactoryCommon.usf
@@ -109,7 +109,7 @@ void SetTangents(inout FVertexFactoryInterpolantsVSToPS Interpolants, float3 InT
 
 #define NUM_PACKED_UVS (NUM_MATERIAL_TEXCOORDS > 1 ? 3 : 1)
 
-#define USE_INTERPOLATOR_345 ((!MATERIAL_LIGHTINGMODEL_UNLIT) || USES_TRANSFORM_VECTOR || USES_DISTORTION)
+#define USE_INTERPOLATOR_345 ((!MATERIAL_SHADINGMODEL_UNLIT) || USES_TRANSFORM_VECTOR || USES_DISTORTION)
 
 /** Interpolants packed for ES2.  All UVs are in the xy to ensure independent texture fetches. */
 struct FVertexFactoryInterpolantsVSToPS
diff --git a/Engine/Shaders/MaterialTemplate.usf b/Engine/Shaders/MaterialTemplate.usf
--- a/Engine/Shaders/MaterialTemplate.usf
+++ b/Engine/Shaders/MaterialTemplate.usf
@@ -845,23 +845,23 @@ MaterialFloat2 GetLightmapUVs(FMaterialPixelParameters Parameters)
 #if NEEDS_SCENE_TEXTURES
 #include "DeferredShadingCommon.usf"		// GetGBufferData()
 
-float3 GetLightingModelColor(FGBufferData GBuffer)
+float3 GetShadingModelColor(FGBufferData GBuffer)
 {
 	// TODO: PS4 doesn't optimize out correctly the switch(), so it thinks it needs all the Samplers even if they get compiled out
 	//	This will get fixed after launch per Sony...
 #if PS4_PROFILE
-		 if (GBuffer.LightingModelId == LIGHTINGMODELID_UNLIT) return float3(0.1f, 0.1f, 0.2f); // Dark Blue
-	else if (GBuffer.LightingModelId == LIGHTINGMODELID_DEFAULT_LIT) return float3(0.1f, 1.0f, 0.1f); // Green
-	else if (GBuffer.LightingModelId == LIGHTINGMODELID_SUBSURFACE) return float3(1.0f, 0.1f, 0.1f); // Red
-	else if (GBuffer.LightingModelId == LIGHTINGMODELID_PREINTEGRATED_SKIN) return float3(0.6f, 0.4f, 0.1f); // Brown
+		 if (GBuffer.ShadingModelID == SHADINGMODELID_UNLIT) return float3(0.1f, 0.1f, 0.2f); // Dark Blue
+	else if (GBuffer.ShadingModelID == SHADINGMODELID_DEFAULT_LIT) return float3(0.1f, 1.0f, 0.1f); // Green
+	else if (GBuffer.ShadingModelID == SHADINGMODELID_SUBSURFACE) return float3(1.0f, 0.1f, 0.1f); // Red
+	else if (GBuffer.ShadingModelID == SHADINGMODELID_PREINTEGRATED_SKIN) return float3(0.6f, 0.4f, 0.1f); // Brown
 	else return float3(1.0f, 1.0f, 1.0f); // White
 #else
-	switch(GBuffer.LightingModelId)
+	switch(GBuffer.ShadingModelID)
 	{
-		case LIGHTINGMODELID_UNLIT: return float3(0.1f, 0.1f, 0.2f); // Dark Blue
-		case LIGHTINGMODELID_DEFAULT_LIT: return float3(0.1f, 1.0f, 0.1f); // Green
-		case LIGHTINGMODELID_SUBSURFACE: return float3(1.0f, 0.1f, 0.1f); // Red
-		case LIGHTINGMODELID_PREINTEGRATED_SKIN: return float3(0.6f, 0.4f, 0.1f); // Brown
+		case SHADINGMODELID_UNLIT: return float3(0.1f, 0.1f, 0.2f); // Dark Blue
+		case SHADINGMODELID_DEFAULT_LIT: return float3(0.1f, 1.0f, 0.1f); // Green
+		case SHADINGMODELID_SUBSURFACE: return float3(1.0f, 0.1f, 0.1f); // Red
+		case SHADINGMODELID_PREINTEGRATED_SKIN: return float3(0.6f, 0.4f, 0.1f); // Brown
 		default: return float3(1.0f, 1.0f, 1.0f); // White
 	}
 #endif
@@ -897,7 +897,7 @@ float4 SceneTextureLookup(float2 UV, int SceneTextureIndex)
 	else if (SceneTextureIndex == 19)	return Texture2DSample(PostprocessInput5, PostprocessInput5Sampler, UV);
 	else if (SceneTextureIndex == 20)	return Texture2DSample(PostprocessInput6, PostprocessInput6Sampler, UV);
 	else if (SceneTextureIndex == 21)	return ScreenSpaceData.GBuffer.DecalMask;
-	else if (SceneTextureIndex == 22)	return float4(GetLightingModelColor(ScreenSpaceData.GBuffer), 1);
+	else if (SceneTextureIndex == 22)	return float4(GetShadingModelColor(ScreenSpaceData.GBuffer), 1);
 	else if (SceneTextureIndex == 23)	return ScreenSpaceData.AmbientOcclusion;
 #else
 	switch(SceneTextureIndex)
@@ -948,8 +948,8 @@ float4 SceneTextureLookup(float2 UV, int SceneTextureIndex)
 		case 20: return Texture2DSample(PostprocessInput6, PostprocessInput6Sampler, UV);
 		// PPI_DecalMask
 		case 21: return ScreenSpaceData.GBuffer.DecalMask;
-		// PPI_LightingModel
-		case 22: return float4(GetLightingModelColor(ScreenSpaceData.GBuffer), 1);
+		// PPI_ShadingModel
+		case 22: return float4(GetShadingModelColor(ScreenSpaceData.GBuffer), 1);
 		// PPI_AmbientOcclusion
 		case 23: return ScreenSpaceData.AmbientOcclusion;
 		default:
diff --git a/Engine/Shaders/ParticleGPUSpriteVertexFactory.usf b/Engine/Shaders/ParticleGPUSpriteVertexFactory.usf
--- a/Engine/Shaders/ParticleGPUSpriteVertexFactory.usf
+++ b/Engine/Shaders/ParticleGPUSpriteVertexFactory.usf
@@ -9,7 +9,7 @@
 #include "VertexFactoryCommon.usf"
 #include "ParticleVertexFactoryCommon.usf"
 
-#define USE_PARTICLE_LIGHTING_OFFSET (FEATURE_LEVEL >= FEATURE_LEVEL_SM5 && !MATERIAL_LIGHTINGMODEL_UNLIT)
+#define USE_PARTICLE_LIGHTING_OFFSET (FEATURE_LEVEL >= FEATURE_LEVEL_SM5 && !MATERIAL_SHADINGMODEL_UNLIT)
 #define USE_PARTICLE_POSITION (FEATURE_LEVEL >= FEATURE_LEVEL_SM5 && NEEDS_PARTICLE_POSITION)
 #define USE_PARTICLE_VELOCITY (FEATURE_LEVEL >= FEATURE_LEVEL_SM5 && NEEDS_PARTICLE_VELOCITY)
 #define USE_PARTICLE_TIME (FEATURE_LEVEL >= FEATURE_LEVEL_SM5 && NEEDS_PARTICLE_TIME)
diff --git a/Engine/Shaders/ParticleSpriteVertexFactory.usf b/Engine/Shaders/ParticleSpriteVertexFactory.usf
--- a/Engine/Shaders/ParticleSpriteVertexFactory.usf
+++ b/Engine/Shaders/ParticleSpriteVertexFactory.usf
@@ -9,7 +9,7 @@
 #include "ParticleVertexFactoryCommon.usf"
 
 // These are only enabled for SM5 due to limited interpolators elsewhere
-#define USE_PARTICLE_LIGHTING_OFFSET (FEATURE_LEVEL >= FEATURE_LEVEL_SM5 && !MATERIAL_LIGHTINGMODEL_UNLIT)
+#define USE_PARTICLE_LIGHTING_OFFSET (FEATURE_LEVEL >= FEATURE_LEVEL_SM5 && !MATERIAL_SHADINGMODEL_UNLIT)
 #define USE_PARTICLE_POSITION (FEATURE_LEVEL >= FEATURE_LEVEL_SM5 && NEEDS_PARTICLE_POSITION)
 #define USE_PARTICLE_VELOCITY (FEATURE_LEVEL >= FEATURE_LEVEL_SM5 && NEEDS_PARTICLE_VELOCITY)
 #define USE_PARTICLE_TIME (FEATURE_LEVEL >= FEATURE_LEVEL_SM5 && NEEDS_PARTICLE_TIME)
diff --git a/Engine/Shaders/PostProcessAmbient.usf b/Engine/Shaders/PostProcessAmbient.usf
--- a/Engine/Shaders/PostProcessAmbient.usf
+++ b/Engine/Shaders/PostProcessAmbient.usf
@@ -225,7 +225,7 @@ void MainPS(in float4 UVAndScreenPos : TEXCOORD0, out float4 OutColor : SV_Targe
 	float AmbientOcclusion = GBuffer.GBufferAO * ScreenSpaceData.AmbientOcclusion;
 
 	// Subsurface
-	BRANCH if(GBuffer.LightingModelId == LIGHTINGMODELID_SUBSURFACE || GBuffer.LightingModelId == LIGHTINGMODELID_PREINTEGRATED_SKIN)
+	BRANCH if(GBuffer.ShadingModelID == SHADINGMODELID_SUBSURFACE || GBuffer.ShadingModelID == SHADINGMODELID_PREINTEGRATED_SKIN)
 	{
 		// some view dependent and some non view dependent (hard coded)
 		float DependentSplit = 0.5f;
diff --git a/Engine/Shaders/PostProcessAmbientOcclusion.usf b/Engine/Shaders/PostProcessAmbientOcclusion.usf
--- a/Engine/Shaders/PostProcessAmbientOcclusion.usf
+++ b/Engine/Shaders/PostProcessAmbientOcclusion.usf
@@ -558,7 +558,7 @@ void BasePassAOPS(in float4 UVAndScreenPos : TEXCOORD0, out float4 OutColor : SV
 	float StaticFraction = ScreenSpaceAOandSSRShaderParams[3].z;
 
 	// can be optimized
-	float AOMask = (ScreenSpaceData.GBuffer.LightingModelId != LIGHTINGMODELID_UNLIT);
+	float AOMask = (ScreenSpaceData.GBuffer.ShadingModelID != SHADINGMODELID_UNLIT);
 
 	OutColor = lerp(1.0f, ScreenSpaceData.AmbientOcclusion * ScreenSpaceData.GBuffer.GBufferAO, AOMask * StaticFraction);
 }
diff --git a/Engine/Shaders/PostProcessGBufferHints.usf b/Engine/Shaders/PostProcessGBufferHints.usf
--- a/Engine/Shaders/PostProcessGBufferHints.usf
+++ b/Engine/Shaders/PostProcessGBufferHints.usf
@@ -15,9 +15,9 @@ uint ComputeAdvice(FGBufferData BufferData, float3 OriginalSceneColor)
 	float SubsurfaceLum = 0;
 	const float OneThird = 1.0f / 3.0f;
 
-	FLATTEN if(BufferData.LightingModelId != LIGHTINGMODELID_UNLIT)
+	FLATTEN if(BufferData.ShadingModelID != SHADINGMODELID_UNLIT)
 	{
-		FLATTEN if(BufferData.LightingModelId == LIGHTINGMODELID_SUBSURFACE)
+		FLATTEN if(BufferData.ShadingModelID == SHADINGMODELID_SUBSURFACE)
 		{
 			SubsurfaceLum = dot(OneThird.xxx, DecodeSubsurfaceColor(BufferData.CustomData));
 		}
diff --git a/Engine/Shaders/PostProcessSubsurface.usf b/Engine/Shaders/PostProcessSubsurface.usf
--- a/Engine/Shaders/PostProcessSubsurface.usf
+++ b/Engine/Shaders/PostProcessSubsurface.usf
@@ -57,7 +57,7 @@ void SetupPS(in float4 UVAndScreenPos : TEXCOORD0, out float4 OutColor : SV_Targ
 
 	FScreenSpaceData ScreenSpaceData = GetScreenSpaceData(UV);
 
-	FLATTEN if(ScreenSpaceData.GBuffer.LightingModelId >= LIGHTINGMODELID_SUBSURFACE)
+	FLATTEN if(ScreenSpaceData.GBuffer.ShadingModelID >= SHADINGMODELID_SUBSURFACE)
 	{
 		float4 SceneColor4 = Texture2DSample(PostprocessInput0, PostprocessInput0Sampler, UV);
 
@@ -153,7 +153,7 @@ void MainPS(float4 UVAndScreenPos : TEXCOORD0, out float4 OutColor : SV_Target0)
 			// hack
 			//		float Subsurface = ScreenSpaceData.GBuffer.Roughness > 0.4f;
 			//		float Subsurface = (1 - ScreenSpaceData.GBuffer.DecalMask) > 0.5f && ScreenSpaceData.GBuffer.Depth < 1000.0f;
-			FLATTEN if(LocalScreenSpaceData.GBuffer.LightingModelId == LIGHTINGMODELID_SUBSURFACE)
+			FLATTEN if(LocalScreenSpaceData.GBuffer.ShadingModelID == SHADINGMODELID_SUBSURFACE)
 			{
 				float3 Color = ColorOverDistance(DecodeSubsurfaceColor(LocalScreenSpaceData.GBuffer.CustomData), r);
 				float ColorLum = Luminance(Color);
@@ -173,7 +173,7 @@ void MainPS(float4 UVAndScreenPos : TEXCOORD0, out float4 OutColor : SV_Target0)
 	float4 SceneColor4 = Texture2DSample(PostprocessInput1, PostprocessInput1Sampler, UV);
 	float4 SSSColor = OutColor;
 
-	SSSColor.a = (ScreenSpaceData.GBuffer.LightingModelId == LIGHTINGMODELID_SUBSURFACE);
+	SSSColor.a = (ScreenSpaceData.GBuffer.ShadingModelID == SHADINGMODELID_SUBSURFACE);
 
 	// partly translucent surface (0.9 means 10% of the opaque is blended in)
 	SSSColor.a *= 0.3f;
diff --git a/Engine/Shaders/PostProcessVisualizeGBuffer.usf b/Engine/Shaders/PostProcessVisualizeGBuffer.usf
--- a/Engine/Shaders/PostProcessVisualizeGBuffer.usf
+++ b/Engine/Shaders/PostProcessVisualizeGBuffer.usf
@@ -6,164 +6,9 @@
 
 #include "Common.usf"
 #include "PostProcessCommon.usf"
-#include "DeferredShadingCommon.usf"
-#include "PostProcessHistogramCommon.usf"
-
-// This should match the enumeration defined in ShowFlags.h.
-// TODO: Perhaps we could share this with the C++ side?
-#define GBUFFER_VISUALIZATION_MODE_ALL				0
-#define GBUFFER_VISUALIZATION_MODE_DIFFUSE_RGB		1
-#define GBUFFER_VISUALIZATION_MODE_SPECULAR_RGB		2
-#define GBUFFER_VISUALIZATION_MODE_BASE_RGB			3
-#define GBUFFER_VISUALIZATION_MODE_SUBSURFACE_RGB	4
-#define GBUFFER_VISUALIZATION_MODE_WORLD_NORMAL		5
-#define GBUFFER_VISUALIZATION_MODE_SEP_TRANS_RGB	6
-#define GBUFFER_VISUALIZATION_MODE_SEP_TRANS_A		7
-#define GBUFFER_VISUALIZATION_MODE_OPACITY			8
-#define GBUFFER_VISUALIZATION_MODE_DECAL_MASK		9
-#define GBUFFER_VISUALIZATION_MODE_SCENE_DEPTH		10
-#define GBUFFER_VISUALIZATION_MODE_METALLIC			11
-#define GBUFFER_VISUALIZATION_MODE_SPECULAR			12
-#define GBUFFER_VISUALIZATION_MODE_ROUGHNESS		13
-#define GBUFFER_VISUALIZATION_MODE_MATERIAL_AO		14
-#define GBUFFER_VISUALIZATION_MODE_LIGHTING_MODEL	15
-
-/** Helper function for computing a distinct color per lighting model */
-float3 GetLightingModelColor(FGBufferData GBuffer)
-{
-//@todo This seems to ICE PSSL...
-#if !PS4_PROFILE
-	switch(GBuffer.LightingModelId)
-	{
-		case LIGHTINGMODELID_UNLIT: return float3(0.1f, 0.1f, 0.2f); // Dark Blue
-		case LIGHTINGMODELID_DEFAULT_LIT: return float3(0.1f, 1.0f, 0.1f); // Green
-		case LIGHTINGMODELID_SUBSURFACE: return float3(1.0f, 0.1f, 0.1f); // Red
-		case LIGHTINGMODELID_PREINTEGRATED_SKIN: return float3(0.6f, 0.4f, 0.1f); // Brown
-	}
-#endif
-	return float3(1.0f, 1.0f, 1.0f); // White
-}
-
-/** Compute the color of the given pixel in the specified GBuffer channel. Note that ChannelIndex is */
-float3 GetGBufferChannelData(FGBufferData GBuffer, float2 UV, int ChannelIndex)
-{
-//@todo This seems to ICE PSSL...
-#if !PS4_PROFILE
-	switch (ChannelIndex)
-	{
-		case GBUFFER_VISUALIZATION_MODE_DIFFUSE_RGB: return pow( GBuffer.DiffuseColor, 1/2.2 );
-		case GBUFFER_VISUALIZATION_MODE_SPECULAR_RGB: return pow( GBuffer.SpecularColor, 1/2.2 );
-		case GBUFFER_VISUALIZATION_MODE_BASE_RGB: return pow( GBuffer.BaseColor, 1/2.2 );
-		case GBUFFER_VISUALIZATION_MODE_SUBSURFACE_RGB: return pow( GBuffer.CustomData, 1/2.2 );
-		case GBUFFER_VISUALIZATION_MODE_WORLD_NORMAL: return GBuffer.WorldNormal * 0.5f + 0.5f;
-		case GBUFFER_VISUALIZATION_MODE_SEP_TRANS_RGB: return Texture2DSample(PostprocessInput1, PostprocessInput1Sampler, UV).rgb*20;
-		case GBUFFER_VISUALIZATION_MODE_OPACITY: return GBuffer.Opacity;
-		case GBUFFER_VISUALIZATION_MODE_SEP_TRANS_A: return Texture2DSample(PostprocessInput1, PostprocessInput1Sampler, UV).a;
-		case GBUFFER_VISUALIZATION_MODE_DECAL_MASK: return GBuffer.DecalMask;
-		case GBUFFER_VISUALIZATION_MODE_SCENE_DEPTH: return frac(GBuffer.Depth * 0.0001f);
-		case GBUFFER_VISUALIZATION_MODE_METALLIC: return GBuffer.Metallic;
-		case GBUFFER_VISUALIZATION_MODE_SPECULAR: return GBuffer.Specular;
-		case GBUFFER_VISUALIZATION_MODE_ROUGHNESS: return GBuffer.Roughness;
-		case GBUFFER_VISUALIZATION_MODE_MATERIAL_AO: return GBuffer.GBufferAO;
-		case GBUFFER_VISUALIZATION_MODE_LIGHTING_MODEL: return GetLightingModelColor(GBuffer);
-	}
-#endif
-	// Unknown 
-	return float3(0,0,0);
-}
-
-int ComputeChannelIndexFromQuarterTile(int2 QuarterTileXY)
-{
-	const int ChannelIndexTable[16] =
-	{
-		1, 2, 3, 4,
-		5, 0, 0, 6,
-		7, 0, 0, 8,
-		9, 10, 11, 12,
-	};
-
-	return ChannelIndexTable[QuarterTileXY.y*4+QuarterTileXY.x];
-}
-
-/** For "All" mode, we split the screen up into a 4x4 grid of tiles, each one able to display a different component of the gbuffer. 
-	This function maps a tile coordinate (X,Y) onto a specific GBuffer component and computes the colour at the given UV */
-
-float3 ComputeQuarterColor(float2 UV, int2 QuarterTileXY)
-{	
-	FGBufferData GBuffer = GetGBufferData(UV);
-	int index = ComputeChannelIndexFromQuarterTile(QuarterTileXY);
-	return GetGBufferChannelData(GBuffer, UV, index);
-}
 
 /** Main pixel shader function */
 void MainPS(float4 UVAndScreenPos : TEXCOORD0, out float4 OutColor : SV_Target0)
 {
-#if GBUFFER_VISUALIZATION_MODE == GBUFFER_VISUALIZATION_MODE_ALL
-	int2 PixelPos = (int2)(UVAndScreenPos.zw * ScreenPosToPixel.xy + ScreenPosToPixel.zw + 0.5f);
-
-	// half viewport size
-	int2 HalfInsetSize = (int2)(ViewportSize.xy / 2);
-	// quarter viewport size
-	int2 QuarterInsetSize = (int2)(ViewportSize.xy / 4);
-
-	// 16 tiles over the full viewport
-	int2 QuarterTileXY = (int2)((uint2)PixelPos / QuarterInsetSize);
-	// pixel position within a quartertile 
-	int2 QuarterTileLocalPos = PixelPos - QuarterTileXY * QuarterInsetSize;
-	
-	// GBuffer size
-	uint2 BufferSize;
-	InputNew0.GetDimensions(BufferSize.x, BufferSize.y);
-	float2 InvBufferSize = 1.0f / (float2)BufferSize;
-
-	int2 HalfTileLocalPos = PixelPos - QuarterInsetSize;
-
-	float2 HalfScaledUV = HalfTileLocalPos / float2(BufferSize) * 2;
-
-	float2 QuarterScaledUV = QuarterTileLocalPos / float2(BufferSize) * 4;
-
 	OutColor = 0;
-
-	float QuarterSoftMask;
-	{
-		const int Border1 = 10;
-		const int Border2 = 8;
-		float BorderDistance = ComputeDistanceToRect(QuarterTileLocalPos, int2(0, 0) + Border1.xx, QuarterInsetSize - 2 * Border1.xx);
-		
-		QuarterSoftMask = 1 - saturate((BorderDistance - Border2.x) *0.5f);
-	}
-
-	float3 BackgroundColor = Texture2DSample(PostprocessInput0, PostprocessInput0Sampler, UVAndScreenPos.xy).rgb;
-
-	// anti aliasing
-	float3 QuarterColor = 0;
-	{
-		int DownsampleFactor = 4;
-
-		for(int dy = 0; dy < DownsampleFactor; ++dy)
-		{
-			for(int dx = 0; dx < DownsampleFactor; ++dx)
-			{
-				float3 SampleColor = ComputeQuarterColor(QuarterScaledUV + float2( dx,  dy) * InvBufferSize, QuarterTileXY);
-
-				QuarterColor += SampleColor;
-			}
-		}
-
-		QuarterColor *= 1.0f / (DownsampleFactor * DownsampleFactor);
-	}
-
-	if(QuarterTileXY.x < 1 || QuarterTileXY.y < 1 || QuarterTileXY.x > 2 || QuarterTileXY.y > 2)
-	{
-		OutColor.rgb = lerp(BackgroundColor, QuarterColor, QuarterSoftMask);
-		OutColor = lerp(OutColor, float4(1,1,1,1), QuarterSoftMask * (1 - QuarterSoftMask));
-	}
-	else
-	{
-		OutColor.rgb = BackgroundColor;
-	}
-#else
-	FGBufferData GBuffer = GetGBufferData(UVAndScreenPos.xy);
-	OutColor = float4(GetGBufferChannelData(GBuffer, UVAndScreenPos.xy, GBUFFER_VISUALIZATION_MODE), 1);
-#endif
 }
\ No newline at end of file
diff --git a/Engine/Shaders/ReflectionEnvironmentComputeShaders.usf b/Engine/Shaders/ReflectionEnvironmentComputeShaders.usf
--- a/Engine/Shaders/ReflectionEnvironmentComputeShaders.usf
+++ b/Engine/Shaders/ReflectionEnvironmentComputeShaders.usf
@@ -344,7 +344,7 @@ void ReflectionEnvironmentTiledDeferredMain(
 
 		BRANCH
 		// Only light pixels marked as using deferred shading
-		if (GBuffer.LightingModelId > 0)
+		if (GBuffer.ShadingModelID > 0)
 		{
 			// Accumulate reflections from captures affecting this tile, applying largest captures first so that the smallest ones display on top
 			LOOP
@@ -540,7 +540,7 @@ void ReflectionEnvironmentTiledDeferredMain(
 #endif
 
 	// Only light pixels marked as using deferred shading
-	BRANCH if (GBuffer.LightingModelId > 0 && SkyLightParameters.y > 0 && !Covered)
+	BRANCH if (GBuffer.ShadingModelID > 0 && SkyLightParameters.y > 0 && !Covered)
 	{
 #if !USE_LIGHTMAPS
 		DiffuseLighting.rgb += DiffuseLighting.a * GetSkyLightReflection(GBuffer.WorldNormal, 1, USE_LIGHTMAPS);
@@ -580,7 +580,7 @@ void ReflectionEnvironmentTiledDeferredMain(
 		SpecularLighting.rgb *= GBuffer.IndirectIrradiance;
 #else
 		BRANCH
-		if( GBuffer.LightingModelId == LIGHTINGMODELID_SUBSURFACE )
+		if( GBuffer.ShadingModelID == SHADINGMODELID_SUBSURFACE )
 		{
 			GBuffer.DiffuseColor += DecodeSubsurfaceColor(GBuffer.CustomData);
 		}
diff --git a/Engine/Shaders/ReflectionEnvironmentShaders.usf b/Engine/Shaders/ReflectionEnvironmentShaders.usf
--- a/Engine/Shaders/ReflectionEnvironmentShaders.usf
+++ b/Engine/Shaders/ReflectionEnvironmentShaders.usf
@@ -386,7 +386,7 @@ void StandardDeferredReflectionPS(
 
 	BRANCH
 	// Only light pixels marked as using deferred shading
-	if (ScreenSpaceData.GBuffer.LightingModelId > 0)
+	if (ScreenSpaceData.GBuffer.ShadingModelID > 0)
 	{
 		FGBufferData InGBufferData = ScreenSpaceData.GBuffer;
 		float SceneDepth = CalcSceneDepth(ScreenUV);
@@ -495,7 +495,7 @@ void ReflectionApplyPS(
 	
 	BRANCH
 	// Only light pixels marked as lit
-	if( GBuffer.LightingModelId > 0 )
+	if( GBuffer.ShadingModelID > 0 )
 	{
 		// TODO if SpecularColor and Roughness were in same texture would remove 2 fetches
 		float3 V = -normalize( mul( float4(UVAndScreenPos.zw, 1, 0), View.ScreenToWorld ).xyz );
diff --git a/Engine/Shaders/ScreenSpaceReflections.usf b/Engine/Shaders/ScreenSpaceReflections.usf
--- a/Engine/Shaders/ScreenSpaceReflections.usf
+++ b/Engine/Shaders/ScreenSpaceReflections.usf
@@ -303,7 +303,7 @@ void ScreenSpaceReflectionsPS(in float4 UVAndScreenPos : TEXCOORD0, out float4 O
 
 #if 1
 	// Early out
-	BRANCH if( RoughnessFade == 0 || ScreenSpaceData.GBuffer.LightingModelId == 0 )
+	BRANCH if( RoughnessFade == 0 || ScreenSpaceData.GBuffer.ShadingModelID == 0 )
 	{
 		return;
 	}
diff --git a/Engine/Shaders/ShadowProjectionPixelShader.usf b/Engine/Shaders/ShadowProjectionPixelShader.usf
--- a/Engine/Shaders/ShadowProjectionPixelShader.usf
+++ b/Engine/Shaders/ShadowProjectionPixelShader.usf
@@ -463,7 +463,7 @@ void Main(
 	#if FEATURE_LEVEL >= FEATURE_LEVEL_SM4
 
 		BRANCH
-		if (InGBufferData.LightingModelId == LIGHTINGMODELID_SUBSURFACE || InGBufferData.LightingModelId == LIGHTINGMODELID_PREINTEGRATED_SKIN)
+		if (InGBufferData.ShadingModelID == SHADINGMODELID_SUBSURFACE || InGBufferData.ShadingModelID == SHADINGMODELID_PREINTEGRATED_SKIN)
 		{
 			// Derive density from a heuristic using opacity, tweaked for useful falloff ranges and to give a linear depth falloff with opacity
 			float Density = -.05f * log(1 - min(InGBufferData.Opacity, .999f));
diff --git a/Engine/Shaders/TiledDeferredLightShaders.usf b/Engine/Shaders/TiledDeferredLightShaders.usf
--- a/Engine/Shaders/TiledDeferredLightShaders.usf
+++ b/Engine/Shaders/TiledDeferredLightShaders.usf
@@ -167,15 +167,15 @@ void TiledDeferredLightingMain(
 	FGBufferData InGBufferData = ScreenSpaceData.GBuffer; 
 	float3 CameraVector = normalize(WorldPosition - View.ViewOrigin.xyz);
 
-	// Multiple lighting models are being handled by unrolling a top level loop
-	// This generates faster code with the D3D11 shader compiler than a single pass over all lights, which branches on GBuffer.LightingModelId in PointLightSubsurface
-	// Presumably because the compiler was able to remove the lighting model branches in the unrolled loop iterations
+	// Multiple shading models are being handled by unrolling a top level loop
+	// This generates faster code with the D3D11 shader compiler than a single pass over all lights, which branches on GBuffer.ShadingModelID in PointLightSubsurface
+	// Presumably because the compiler was able to remove the shading model branches in the unrolled loop iterations
 	UNROLL
-	for (int TileLightingModelId = 1; TileLightingModelId < LIGHTINGMODELID_NUM; TileLightingModelId++)
+	for (uint TileShadingModelID = 1; TileShadingModelID < SHADINGMODELID_NUM; TileShadingModelID++)
 	{
-		// Only light pixels with the lighting model being processed in this loop iteration
+		// Only light pixels with the shading model being processed in this loop iteration
 		BRANCH
-		if (InGBufferData.LightingModelId == TileLightingModelId)
+		if (InGBufferData.ShadingModelID == TileShadingModelID)
 		{
 			LOOP
 			for (uint TileLightIndex = 0; TileLightIndex < NumLightsAffectingTile; TileLightIndex++) 
@@ -200,14 +200,14 @@ void TiledDeferredLightingMain(
 				LightData.bShadowed = dot(LightData.ShadowMapChannelMask, float4(1, 1, 1, 1));
 	
 				// Lights requiring light attenuation are not supported tiled for now
-				CompositedLighting += GetDynamicLighting(WorldPosition, CameraVector, ScreenUV, ScreenSpaceData, TileLightingModelId, LightData, float4(1, 1, 1, 1), uint2(0,0));
+				CompositedLighting += GetDynamicLighting(WorldPosition, CameraVector, ScreenUV, ScreenSpaceData, TileShadingModelID, LightData, float4(1, 1, 1, 1), uint2(0,0));
 			}
 		}
 	}
 		
-	// The simple lighting model does not depend on LightingModelId, so use it anytime the material is lit
+	// The simple shading model does not depend on ShadingModelID, so use it anytime the material is lit
 	BRANCH
-	if (InGBufferData.LightingModelId != LIGHTINGMODELID_UNLIT)
+	if (InGBufferData.ShadingModelID != SHADINGMODELID_UNLIT)
 	{
 		LOOP
 		for (uint TileLightIndex = 0; TileLightIndex < NumSimpleLightsAffectingTile; TileLightIndex++) 
diff --git a/Engine/Source/Editor/FontEditor/Private/FontEditor.cpp b/Engine/Source/Editor/FontEditor/Private/FontEditor.cpp
--- a/Engine/Source/Editor/FontEditor/Private/FontEditor.cpp
+++ b/Engine/Source/Editor/FontEditor/Private/FontEditor.cpp
@@ -132,7 +132,7 @@ void FFontEditor::InitFontEditor(const EToolkitMode::Type Mode, const TSharedPtr
 	Factory = ConstructObject<UTextureFactory>(UTextureFactory::StaticClass());
 	// Set the defaults
 	Factory->Blending = BLEND_Opaque;
-	Factory->LightingModel = MLM_Unlit;
+	Factory->ShadingModel = MSM_Unlit;
 	Factory->bDeferCompression = true;
 	Factory->MipGenSettings = TMGS_NoMipmaps;
 	
diff --git a/Engine/Source/Editor/MaterialEditor/Private/MaterialEditor.cpp b/Engine/Source/Editor/MaterialEditor/Private/MaterialEditor.cpp
--- a/Engine/Source/Editor/MaterialEditor/Private/MaterialEditor.cpp
+++ b/Engine/Source/Editor/MaterialEditor/Private/MaterialEditor.cpp
@@ -232,7 +232,7 @@ void FMaterialEditor::InitEditorForMaterialFunction(UMaterialFunction* InMateria
 		// Hack: serialize the new material with an archive that does nothing so that its material resources are created
 		Material->Serialize(DummyArchive);
 	}
-	Material->SetLightingModel(MLM_Unlit);
+	Material->SetShadingModel(MSM_Unlit);
 
 	// Propagate all object flags except for RF_Standalone, otherwise the preview material function won't GC once
 	// the material editor releases the reference.
diff --git a/Engine/Source/Editor/MaterialEditor/Private/MaterialEditor.h b/Engine/Source/Editor/MaterialEditor/Private/MaterialEditor.h
--- a/Engine/Source/Editor/MaterialEditor/Private/MaterialEditor.h
+++ b/Engine/Source/Editor/MaterialEditor/Private/MaterialEditor.h
@@ -117,7 +117,7 @@ public:
 	virtual bool IsWireframe() const { return false; }
 	virtual bool IsMasked() const { return false; }
 	virtual enum EBlendMode GetBlendMode() const { return BLEND_Opaque; }
-	virtual enum EMaterialLightingModel GetLightingModel() const { return MLM_Unlit; }
+	virtual enum EMaterialShadingModel GetShadingModel() const { return MSM_Unlit; }
 	virtual float GetOpacityMaskClipValue() const { return 0.5f; }
 	virtual FString GetFriendlyName() const { return FString::Printf(TEXT("FMatExpressionPreview %s"), Expression.IsValid() ? *Expression->GetName() : TEXT("NULL")); }
 	/**
diff --git a/Engine/Source/Editor/UnrealEd/Classes/Factories/TextureFactory.h b/Engine/Source/Editor/UnrealEd/Classes/Factories/TextureFactory.h
--- a/Engine/Source/Editor/UnrealEd/Classes/Factories/TextureFactory.h
+++ b/Engine/Source/Editor/UnrealEd/Classes/Factories/TextureFactory.h
@@ -60,9 +60,9 @@ class UNREALED_API UTextureFactory : public UFactory
 	UPROPERTY(EditAnywhere, Category=CreateMaterial, meta=(ToolTip="The blend mode of the created material"))
 	TEnumAsByte<enum EBlendMode> Blending;
 
-	/** The lighting model of the created material */
-	UPROPERTY(EditAnywhere, Category=CreateMaterial, meta=(ToolTip="The lighting model of the created material"))
-	TEnumAsByte<enum EMaterialLightingModel> LightingModel;
+	/** The shading model of the created material */
+	UPROPERTY(EditAnywhere, Category=CreateMaterial, meta=(ToolTip="The shading model of the created material"))
+	TEnumAsByte<enum EMaterialShadingModel> ShadingModel;
 
 	/** The mip-map generation settings for the texture; Allows customization of the content of the mip-map chain */
 	UPROPERTY(EditAnywhere, Category=TextureFactory, meta=(ToolTip="The mip-map generation settings for the texture; Allows customization of the content of the mip-map chain"))
diff --git a/Engine/Source/Editor/UnrealEd/Private/EditorExporters.cpp b/Engine/Source/Editor/UnrealEd/Private/EditorExporters.cpp
--- a/Engine/Source/Editor/UnrealEd/Private/EditorExporters.cpp
+++ b/Engine/Source/Editor/UnrealEd/Private/EditorExporters.cpp
@@ -788,7 +788,7 @@ public:
 		{
 			UMaterial* ProxyMaterial = MaterialInterface->GetMaterial();
 			EBlendMode BlendMode = MaterialInterface->GetBlendMode();
-			EMaterialLightingModel LightingModel = MaterialInterface->GetLightingModel();
+			EMaterialShadingModel ShadingModel = MaterialInterface->GetShadingModel();
 			check(ProxyMaterial);
 			switch (PropertyToCompile)
 			{
@@ -835,7 +835,7 @@ public:
 				}
 				else if (BlendMode == BLEND_Modulate)
 				{
-					if (LightingModel == MLM_Unlit)
+					if (ShadingModel == MSM_Unlit)
 					{
 						return Compiler->ForceCast(MaterialInterface->CompileProperty(Compiler, MP_EmissiveColor),MCT_Float3,true,true);
 					}
@@ -847,7 +847,7 @@ public:
 				else if ((BlendMode == BLEND_Translucent) || (BlendMode == BLEND_Additive))
 				{
 					int32 ColoredOpacity = -1;
-					if (LightingModel == MLM_Unlit)
+					if (ShadingModel == MSM_Unlit)
 					{
 						ColoredOpacity = Compiler->ForceCast(MaterialInterface->CompileProperty(Compiler, MP_EmissiveColor),MCT_Float3,true,true);
 					}
@@ -928,7 +928,7 @@ public:
 	}
 	virtual bool IsMasked() const									{ return false; }
 	virtual enum EBlendMode GetBlendMode() const					{ return BLEND_Opaque; }
-	virtual enum EMaterialLightingModel GetLightingModel() const	{ return MLM_Unlit; }
+	virtual enum EMaterialShadingModel GetShadingModel() const		{ return MSM_Unlit; }
 	virtual float GetOpacityMaskClipValue() const					{ return 0.5f; }
 	virtual FString GetFriendlyName() const { return FString::Printf(TEXT("FExportMaterialRenderer %s"), MaterialInterface ? *MaterialInterface->GetName() : TEXT("NULL")); }
 	/**
@@ -1000,7 +1000,7 @@ public:
 		OutUniformValue.A = 0;
 
 		EBlendMode BlendMode = MaterialInterface->GetBlendMode();
-		EMaterialLightingModel LightingModel = MaterialInterface->GetLightingModel();
+		EMaterialShadingModel ShadingModel = MaterialInterface->GetShadingModel();
 		
 		check(Material);
 		bool bExpressionIsNULL = false;
@@ -1048,7 +1048,7 @@ public:
 				(BlendMode == BLEND_Additive))
 			{
 				bool bColorInputIsNULL = false;
-				if (LightingModel == MLM_Unlit)
+				if (ShadingModel == MSM_Unlit)
 				{
 					bColorInputIsNULL = !IsMaterialInputConnected(Material, MP_EmissiveColor);
 				}
@@ -2805,7 +2805,7 @@ namespace MaterialExportUtils
 		
 		UMaterial* Material = ConstructObject<UMaterial>(UMaterial::StaticClass(), Outer, MaterialName, Flags);
 		Material->TwoSided = false;
-		Material->SetLightingModel(MLM_DefaultLit);
+		Material->SetShadingModel(MSM_DefaultLit);
 
 		//Set Metallic as a default constant
 		{
diff --git a/Engine/Source/Editor/UnrealEd/Private/Factories/EditorFactories.cpp b/Engine/Source/Editor/UnrealEd/Private/Factories/EditorFactories.cpp
--- a/Engine/Source/Editor/UnrealEd/Private/Factories/EditorFactories.cpp
+++ b/Engine/Source/Editor/UnrealEd/Private/Factories/EditorFactories.cpp
@@ -4491,7 +4491,7 @@ UObject* UTextureFactory::FactoryCreateBinary
 
 		Material->TwoSided	= bTwoSided;
 		Material->BlendMode = Blending;
-		Material->SetLightingModel(LightingModel);
+		Material->SetShadingModel(ShadingModel);
 
 		Material->PostEditChange();
 	}
diff --git a/Engine/Source/Editor/UnrealEd/Private/Fbx/FbxMainExport.cpp b/Engine/Source/Editor/UnrealEd/Private/Fbx/FbxMainExport.cpp
--- a/Engine/Source/Editor/UnrealEd/Private/Fbx/FbxMainExport.cpp
+++ b/Engine/Source/Editor/UnrealEd/Private/Fbx/FbxMainExport.cpp
@@ -925,14 +925,14 @@ FbxSurfaceMaterial* FFbxExporter::ExportMaterial(UMaterial* Material)
 	// Create the Fbx material
 	FbxSurfaceMaterial* FbxMaterial = NULL;
 	
-	// Set the lighting model
-	if (Material->GetLightingModel() == MLM_DefaultLit)
+	// Set the shading model
+	if (Material->GetShadingModel() == MSM_DefaultLit)
 	{
 		FbxMaterial = FbxSurfacePhong::Create(Scene, TCHAR_TO_ANSI(*Material->GetName()));
 		((FbxSurfacePhong*)FbxMaterial)->Specular.Set(SetMaterialComponent(Material->SpecularColor));
 		//((FbxSurfacePhong*)FbxMaterial)->Shininess.Set(Material->SpecularPower.Constant);
 	}
-	else // if (Material->LightingModel == MLM_Unlit)
+	else // if (Material->ShadingModel == MSM_Unlit)
 	{
 		FbxMaterial = FbxSurfaceLambert::Create(Scene, TCHAR_TO_ANSI(*Material->GetName()));
 	}
diff --git a/Engine/Source/Runtime/Engine/Classes/Engine/EngineTypes.h b/Engine/Source/Runtime/Engine/Classes/Engine/EngineTypes.h
--- a/Engine/Source/Runtime/Engine/Classes/Engine/EngineTypes.h
+++ b/Engine/Source/Runtime/Engine/Classes/Engine/EngineTypes.h
@@ -100,13 +100,13 @@ enum ETranslucencyLightingMode
 };
 
 UENUM()
-enum EMaterialLightingModel
+enum EMaterialShadingModel
 {
-	MLM_DefaultLit UMETA(DisplayName="Default Lit"),
-	MLM_Unlit UMETA(DisplayName="Unlit"),
-	MLM_Subsurface UMETA(DisplayName="Subsurface"),
-	MLM_PreintegratedSkin UMETA(DisplayName="Preintegrated Skin"),
-	MLM_MAX,
+	MSM_DefaultLit			UMETA(DisplayName="Default Lit"),
+	MSM_Unlit				UMETA(DisplayName="Unlit"),
+	MSM_Subsurface			UMETA(DisplayName="Subsurface"),
+	MSM_PreintegratedSkin	UMETA(DisplayName="Preintegrated Skin"),
+	MSM_MAX,
 };
 
 // This is used by the drawing passes to determine tessellation policy, so changes here need to be supported in native code.
diff --git a/Engine/Source/Runtime/Engine/Classes/Materials/Material.h b/Engine/Source/Runtime/Engine/Classes/Materials/Material.h
--- a/Engine/Source/Runtime/Engine/Classes/Materials/Material.h
+++ b/Engine/Source/Runtime/Engine/Classes/Materials/Material.h
@@ -253,7 +253,7 @@ struct FMaterialParameterCollectionInfo
 
 /**
  * A Material is an asset which can be applied to a mesh to control the visual look of the scene. In general,
- * when light from the scene hits the surface, the lighting model of the material is used to calculate how
+ * when light from the scene hits the surface, the shading model of the material is used to calculate how
  * that light interacts with the surface. 
  */
 UCLASS(hidecategories=Object, MinimalAPI, BlueprintType)
@@ -332,7 +332,7 @@ class UMaterial : public UMaterialInterface
 private:
 	/** Determines how inputs are combined to create the material's final color. */
 	UPROPERTY(EditAnywhere, Category=Material, AssetRegistrySearchable)
-	TEnumAsByte<enum EMaterialLightingModel> LightingModel;
+	TEnumAsByte<enum EMaterialShadingModel> ShadingModel;
 public:
 
 	/** If BlendMode is BLEND_Masked, the surface is not rendered where OpacityMask < OpacityMaskClipValue. */
@@ -351,7 +351,7 @@ public:
 	UPROPERTY()
 	FScalarMaterialInput TessellationMultiplier;
 
-	/** Inner material color, only used for LightingModel=Subsurface */
+	/** Inner material color, only used for ShadingModel=Subsurface */
 	UPROPERTY()
 	FColorMaterialInput SubsurfaceColor;
 
@@ -426,7 +426,7 @@ public:
 	float TranslucentSelfShadowSecondOpacity;
 
 	/** 
-	 * Controls how diffuse the material's backscattering is when using the MLM_Subsurface lighting model.
+	 * Controls how diffuse the material's backscattering is when using the MSM_Subsurface shading model.
 	 * Larger exponents give a less diffuse look (smaller, brighter backscattering highlight).
 	 * This is only used when the object is casting a volumetric translucent shadow from a directional light.
 	 */
@@ -750,10 +750,10 @@ public:
 
 	ENGINE_API virtual float GetOpacityMaskClipValue_Internal() const;
 	ENGINE_API virtual EBlendMode GetBlendMode_Internal() const;
-	ENGINE_API virtual EMaterialLightingModel GetLightingModel_Internal() const;
+	ENGINE_API virtual EMaterialShadingModel GetShadingModel_Internal() const;
 	ENGINE_API virtual bool IsTwoSided_Internal() const;
 
-	ENGINE_API void SetLightingModel(EMaterialLightingModel NewModel) {LightingModel = NewModel;}
+	ENGINE_API void SetShadingModel(EMaterialShadingModel NewModel) {ShadingModel = NewModel;}
 
 	/** Checks to see if an input property should be active, based on the state of the material */
 	ENGINE_API virtual bool IsPropertyActive(EMaterialProperty InProperty) const;
@@ -1072,8 +1072,8 @@ public:
 	void GetReferencedParameterCollectionIds(TArray<FGuid>& OutIds) const;
 
 	/* Helper functions for text output of properties. */
-	static const TCHAR* GetMaterialLightingModelString(EMaterialLightingModel InMaterialLightingModel);
-	static EMaterialLightingModel GetMaterialLightingModelFromString(const TCHAR* InMaterialLightingModelStr);
+	static const TCHAR* GetMaterialShadingModelString(EMaterialShadingModel InMaterialShadingModel);
+	static EMaterialShadingModel GetMaterialShadingModelFromString(const TCHAR* InMaterialShadingModelStr);
 	static const TCHAR* GetBlendModeString(EBlendMode InBlendMode);
 	static EBlendMode GetBlendModeFromString(const TCHAR* InBlendModeStr);
 
diff --git a/Engine/Source/Runtime/Engine/Classes/Materials/MaterialExpressionSceneTexture.h b/Engine/Source/Runtime/Engine/Classes/Materials/MaterialExpressionSceneTexture.h
--- a/Engine/Source/Runtime/Engine/Classes/Materials/MaterialExpressionSceneTexture.h
+++ b/Engine/Source/Runtime/Engine/Classes/Materials/MaterialExpressionSceneTexture.h
@@ -53,8 +53,8 @@ enum ESceneTextureId
 	PPI_PostProcessInput6 UMETA(DisplayName="PostProcessInput6"),
 	/** Decal Mask, single bit */
 	PPI_DecalMask UMETA(DisplayName="Decal Mask"),
-	/** Lighting model, two bits */
-	PPI_LightingModel UMETA(DisplayName="Lighting Model"),
+	/** Shading model */
+	PPI_ShadingModel UMETA(DisplayName="Shading Model"),
 	/** Ambient Occlusion, single channel */
 	PPI_AmbientOcclusion UMETA(DisplayName="Ambient Occlusion"),
 };
diff --git a/Engine/Source/Runtime/Engine/Classes/Materials/MaterialInstance.h b/Engine/Source/Runtime/Engine/Classes/Materials/MaterialInstance.h
--- a/Engine/Source/Runtime/Engine/Classes/Materials/MaterialInstance.h
+++ b/Engine/Source/Runtime/Engine/Classes/Materials/MaterialInstance.h
@@ -245,7 +245,7 @@ public:
 	
 	ENGINE_API virtual float GetOpacityMaskClipValue_Internal() const;
 	ENGINE_API virtual EBlendMode GetBlendMode_Internal() const;
-	ENGINE_API virtual EMaterialLightingModel GetLightingModel_Internal() const;
+	ENGINE_API virtual EMaterialShadingModel GetShadingModel_Internal() const;
 	ENGINE_API virtual bool IsTwoSided_Internal() const;
 
 	/** Returns true and sets Result if the property was overridden in the instance. Otherwise, returns false and the base material property should be used. */
@@ -253,7 +253,7 @@ public:
 	/** Returns true and sets Result if the property was overridden in the instance. Otherwise, returns false and the base material property should be used. */
 	ENGINE_API virtual bool GetBlendModeOverride(EBlendMode& OutResult) const;
 	/** Returns true and sets Result if the property was overridden in the instance. Otherwise, returns false and the base material property should be used. */
-	ENGINE_API virtual bool GetLightingModelOverride(EMaterialLightingModel& OutResult) const;
+	ENGINE_API virtual bool GetShadingModelOverride(EMaterialShadingModel& OutResult) const;
 	/** Returns true and sets Result if the property was overridden in the instance. Otherwise, returns false and the base material property should be used. */
 	ENGINE_API virtual bool IsTwoSidedOverride(bool& OutResult) const;
 
diff --git a/Engine/Source/Runtime/Engine/Classes/Materials/MaterialInstanceBasePropertyOverrides.h b/Engine/Source/Runtime/Engine/Classes/Materials/MaterialInstanceBasePropertyOverrides.h
--- a/Engine/Source/Runtime/Engine/Classes/Materials/MaterialInstanceBasePropertyOverrides.h
+++ b/Engine/Source/Runtime/Engine/Classes/Materials/MaterialInstanceBasePropertyOverrides.h
@@ -20,9 +20,9 @@ struct ENGINE_API FMaterialInstanceBasePropertyOverrides
 	UPROPERTY(EditAnywhere, Category=Material)
 	bool bOverride_BlendMode;
 
-	/** Enables override of the lighting model. */
+	/** Enables override of the shading model. */
 	UPROPERTY(EditAnywhere, Category=Material)
-	bool bOverride_LightingModel;
+	bool bOverride_ShadingModel;
 
 	/** Enables override of the two sided property. */
 	UPROPERTY(EditAnywhere, Category=Material)
@@ -36,9 +36,9 @@ struct ENGINE_API FMaterialInstanceBasePropertyOverrides
 	UPROPERTY(EditAnywhere, Category=Material, meta=(editcondition = "bOverride_BlendMode"))
 	TEnumAsByte<EBlendMode> BlendMode;
 
-	/** The lighting model */
-	UPROPERTY(EditAnywhere, Category=Material, meta=(editcondition = "bOverride_LightingModel"))
-	TEnumAsByte<EMaterialLightingModel> LightingModel;	
+	/** The shading model */
+	UPROPERTY(EditAnywhere, Category=Material, meta=(editcondition = "bOverride_ShadingModel"))
+	TEnumAsByte<EMaterialShadingModel> ShadingModel;	
 	
 	/** If the material is two sided or not. */
 	UPROPERTY(EditAnywhere, Category=Material, meta=(editcondition = "bOverride_TwoSided"))
@@ -59,7 +59,7 @@ struct ENGINE_API FMaterialInstanceBasePropertyOverrides
 		if( Ar.UE4Ver() >= VER_UE4_MATERIAL_INSTANCE_BASE_PROPERTY_OVERRIDES_PHASE_2 )
 		{
 			Ar	<< Ref.bOverride_BlendMode << Ref.BlendMode
-				<< Ref.bOverride_LightingModel << Ref.LightingModel
+				<< Ref.bOverride_ShadingModel << Ref.ShadingModel
 				<< Ref.bOverride_TwoSided;
 
 			bool bTwoSided = Ref.TwoSided != 0;
diff --git a/Engine/Source/Runtime/Engine/Classes/Materials/MaterialInterface.h b/Engine/Source/Runtime/Engine/Classes/Materials/MaterialInterface.h
--- a/Engine/Source/Runtime/Engine/Classes/Materials/MaterialInterface.h
+++ b/Engine/Source/Runtime/Engine/Classes/Materials/MaterialInterface.h
@@ -444,13 +444,13 @@ public:
 	*/
 	ENGINE_API float GetOpacityMaskClipValue() const;
 	ENGINE_API EBlendMode GetBlendMode() const;
-	ENGINE_API EMaterialLightingModel GetLightingModel() const;
+	ENGINE_API EMaterialShadingModel GetShadingModel() const;
 	ENGINE_API bool IsTwoSided() const;
 
 	/** Game thread versions of the accessors. On the render thread there are equivelant getters in FMaterial. */
 	ENGINE_API virtual float GetOpacityMaskClipValue_Internal() const;
 	ENGINE_API virtual EBlendMode GetBlendMode_Internal() const;
-	ENGINE_API virtual EMaterialLightingModel GetLightingModel_Internal() const;
+	ENGINE_API virtual EMaterialShadingModel GetShadingModel_Internal() const;
 	ENGINE_API virtual bool IsTwoSided_Internal() const;
 	/**
 	 * Force the streaming system to disregard the normal logic for the specified duration and
diff --git a/Engine/Source/Runtime/Engine/Private/BufferVisualizationData.cpp b/Engine/Source/Runtime/Engine/Private/BufferVisualizationData.cpp
--- a/Engine/Source/Runtime/Engine/Private/BufferVisualizationData.cpp
+++ b/Engine/Source/Runtime/Engine/Private/BufferVisualizationData.cpp
@@ -86,7 +86,7 @@ void FBufferVisualizationData::ConfigureConsoleCommand()
 	ConsoleDocumentationOverviewTargets += AvailableVisualizationMaterials;
 
 	IConsoleManager::Get().RegisterConsoleVariable(TEXT("r.BufferVisualizationOverviewTargets"),
-		TEXT("DiffuseColor,SpecularColor,SubsurfaceColor,WorldNormal,SeparateTranslucencyRGB,,,Opacity,SeparateTranslucencyA,,,DecalMask,SceneDepth,Roughness,MaterialAO,LightingModel"), 
+		TEXT("DiffuseColor,SpecularColor,SubsurfaceColor,WorldNormal,SeparateTranslucencyRGB,,,Opacity,SeparateTranslucencyA,,,DecalMask,SceneDepth,Roughness,MaterialAO,ShadingModel"), 
 		*ConsoleDocumentationOverviewTargets,
 		ECVF_Default);
 }
diff --git a/Engine/Source/Runtime/Engine/Private/Materials/HLSLMaterialTranslator.h b/Engine/Source/Runtime/Engine/Private/Materials/HLSLMaterialTranslator.h
--- a/Engine/Source/Runtime/Engine/Private/Materials/HLSLMaterialTranslator.h
+++ b/Engine/Source/Runtime/Engine/Private/Materials/HLSLMaterialTranslator.h
@@ -389,7 +389,7 @@ public:
 			}
 
 
-			if (Material->GetBlendMode() == BLEND_Modulate && Material->GetLightingModel() != MLM_Unlit && !Material->IsUsedWithDeferredDecal())
+			if (Material->GetBlendMode() == BLEND_Modulate && Material->GetShadingModel() != MSM_Unlit && !Material->IsUsedWithDeferredDecal())
 			{
 				Errorf(TEXT("Dynamically lit translucency is not supported for BLEND_Modulate materials."));
 			}
@@ -427,12 +427,12 @@ public:
 				Errorf(TEXT("Light function materials must be opaque."));
 			}
 
-			if (Material->IsLightFunction() && Material->GetLightingModel() != MLM_Unlit)
+			if (Material->IsLightFunction() && Material->GetShadingModel() != MSM_Unlit)
 			{
 				Errorf(TEXT("Light function materials must use unlit."));
 			}
 
-			if (Material->GetMaterialDomain() == MD_PostProcess && Material->GetLightingModel() != MLM_Unlit)
+			if (Material->GetMaterialDomain() == MD_PostProcess && Material->GetShadingModel() != MSM_Unlit)
 			{
 				Errorf(TEXT("Post process materials must use unlit."));
 			}
@@ -2344,7 +2344,7 @@ protected:
 			|| SceneTextureId == PPI_Roughness
 			|| SceneTextureId == PPI_MaterialAO
 			|| SceneTextureId == PPI_DecalMask
-			|| SceneTextureId == PPI_LightingModel;
+			|| SceneTextureId == PPI_ShadingModel;
 
 		// not yet tracked:
 		//   PPI_SeparateTranslucency, PPI_CustomDepth, PPI_AmbientOcclusion
diff --git a/Engine/Source/Runtime/Engine/Private/Materials/Material.cpp b/Engine/Source/Runtime/Engine/Private/Materials/Material.cpp
--- a/Engine/Source/Runtime/Engine/Private/Materials/Material.cpp
+++ b/Engine/Source/Runtime/Engine/Private/Materials/Material.cpp
@@ -1288,21 +1288,21 @@ UPhysicalMaterial* UMaterial::GetPhysicalMaterial() const
 #define TEXT_TO_ENUM(eVal, txt)		if (FCString::Stricmp(TEXT(#eVal), txt) == 0)	return eVal;
 #endif
 
-const TCHAR* UMaterial::GetMaterialLightingModelString(EMaterialLightingModel InMaterialLightingModel)
+const TCHAR* UMaterial::GetMaterialShadingModelString(EMaterialShadingModel InMaterialShadingModel)
 {
-	switch (InMaterialLightingModel)
+	switch (InMaterialShadingModel)
 	{
-		FOREACH_ENUM_EMATERIALLIGHTINGMODEL(CASE_ENUM_TO_TEXT)
+		FOREACH_ENUM_EMATERIALSHADINGMODEL(CASE_ENUM_TO_TEXT)
 	}
-	return TEXT("MLM_DefaultLit");
+	return TEXT("MSM_DefaultLit");
 }
 
-EMaterialLightingModel UMaterial::GetMaterialLightingModelFromString(const TCHAR* InMaterialLightingModelStr)
+EMaterialShadingModel UMaterial::GetMaterialShadingModelFromString(const TCHAR* InMaterialShadingModelStr)
 {
-	#define TEXT_TO_LIGHTINGMODEL(m) TEXT_TO_ENUM(m, InMaterialLightingModelStr);
-	FOREACH_ENUM_EMATERIALLIGHTINGMODEL(TEXT_TO_LIGHTINGMODEL)
-	#undef TEXT_TO_LIGHTINGMODEL
-	return MLM_DefaultLit;
+	#define TEXT_TO_SHADINGMODEL(m) TEXT_TO_ENUM(m, InMaterialShadingModelStr);
+	FOREACH_ENUM_EMATERIALSHADINGMODEL(TEXT_TO_SHADINGMODEL)
+	#undef TEXT_TO_SHADINGMODEL
+	return MSM_DefaultLit;
 }
 
 const TCHAR* UMaterial::GetBlendModeString(EBlendMode InBlendMode)
@@ -1678,7 +1678,7 @@ void UMaterial::BackwardsCompatibilityInputConversion()
 
 		Roughness.Constant = 0.4238f;
 
-		if( LightingModel != MLM_Unlit )
+		if( ShadingModel != MSM_Unlit )
 		{
 			// Multiply SpecularColor by FresnelBaseReflectFraction
 			if( SpecularColor.IsConnected() && FresnelBaseReflectFraction_DEPRECATED != 1.0f )
@@ -1725,7 +1725,7 @@ void UMaterial::BackwardsCompatibilityInputConversion()
 		}
 	}
 
-	if( LightingModel != MLM_Unlit && UseDiffuseSpecularMaterialInputs->GetValueOnGameThread() == 0 )
+	if( ShadingModel != MSM_Unlit && UseDiffuseSpecularMaterialInputs->GetValueOnGameThread() == 0 )
 	{
 		bool bIsDS = DiffuseColor.IsConnected() || SpecularColor.IsConnected();
 		bool bIsBMS = BaseColor.IsConnected() || Metallic.IsConnected() || Specular.IsConnected();
@@ -1901,10 +1901,10 @@ void UMaterial::PostLoad()
 	bUsedAsLightFunction_DEPRECATED = false;
 	bUsedWithDeferredDecal_DEPRECATED = false;
 
-	// Fix the lighting model to be valid.  Loading a material saved with a lighting model that has been removed will yield a MLM_MAX.
-	if(LightingModel == MLM_MAX)
+	// Fix the shading model to be valid.  Loading a material saved with a shading model that has been removed will yield a MSM_MAX.
+	if(ShadingModel == MSM_MAX)
 	{
-		LightingModel = MLM_DefaultLit;
+		ShadingModel = MSM_DefaultLit;
 	}
 
 	if(DecalBlendMode == DBM_MAX)
@@ -2114,7 +2114,7 @@ bool UMaterial::CanEditChange(const UProperty* InProperty) const
 			return MaterialDomain == MD_Surface;
 		}
 	
-		if (PropertyName == GET_MEMBER_NAME_STRING_CHECKED(UMaterial, LightingModel))
+		if (PropertyName == GET_MEMBER_NAME_STRING_CHECKED(UMaterial, ShadingModel))
 		{
 			return MaterialDomain == MD_Surface;
 		}
@@ -2153,7 +2153,7 @@ bool UMaterial::CanEditChange(const UProperty* InProperty) const
 			|| PropertyName == GET_MEMBER_NAME_STRING_CHECKED(UMaterial, TranslucentMultipleScatteringExtinction)
 			|| PropertyName == GET_MEMBER_NAME_STRING_CHECKED(UMaterial, TranslucentShadowStartOffset))
 		{
-			return IsTranslucentBlendMode(BlendMode) && LightingModel != MLM_Unlit;
+			return IsTranslucentBlendMode(BlendMode) && ShadingModel != MSM_Unlit;
 		}
 	}
 
@@ -3366,22 +3366,22 @@ EBlendMode UMaterial::GetBlendMode_Internal() const
 	return BlendMode;
 }
 
-EMaterialLightingModel UMaterial::GetLightingModel_Internal() const
+EMaterialShadingModel UMaterial::GetShadingModel_Internal() const
 {
 	switch (MaterialDomain)
 	{
 		case MD_Surface:
 		case MD_DeferredDecal:
-			return LightingModel;
+			return ShadingModel;
 
 		// Post process and light function materials must be rendered with the unlit model.
 		case MD_PostProcess:
 		case MD_LightFunction:
-			return MLM_Unlit;
+			return MSM_Unlit;
 
 		default:
 			checkNoEntry();
-			return MLM_Unlit;
+			return MSM_Unlit;
 	}
 }
 
@@ -3398,7 +3398,7 @@ bool UMaterial::IsPropertyActive(EMaterialProperty InProperty)const
 	}
 	else if(MaterialDomain == MD_LightFunction)
 	{
-		// light functions should already use MLM_Unlit but we also we don't want WorldPosOffset
+		// light functions should already use MSM_Unlit but we also we don't want WorldPosOffset
 		return InProperty == MP_EmissiveColor;
 	}
 	else if(MaterialDomain == MD_DeferredDecal)
@@ -3489,7 +3489,7 @@ bool UMaterial::IsPropertyActive(EMaterialProperty InProperty)const
 		break;
 	case MP_Opacity:
 		Active = IsTranslucentBlendMode((EBlendMode)BlendMode) && BlendMode != BLEND_Modulate;
-		if(LightingModel == MLM_Subsurface || LightingModel == MLM_PreintegratedSkin)
+		if(ShadingModel == MSM_Subsurface || ShadingModel == MSM_PreintegratedSkin)
 		{
 			Active = true;
 		}
@@ -3504,13 +3504,13 @@ bool UMaterial::IsPropertyActive(EMaterialProperty InProperty)const
 	case MP_Specular:
 	case MP_Roughness:
 	case MP_AmbientOcclusion:
-		Active = LightingModel != MLM_Unlit;
+		Active = ShadingModel != MSM_Unlit;
 		break;
 	case MP_Normal:
-		Active = LightingModel != MLM_Unlit || Refraction.IsConnected();
+		Active = ShadingModel != MSM_Unlit || Refraction.IsConnected();
 		break;
 	case MP_SubsurfaceColor:
-		Active = LightingModel == MLM_Subsurface || LightingModel == MLM_PreintegratedSkin;
+		Active = ShadingModel == MSM_Subsurface || ShadingModel == MSM_PreintegratedSkin;
 		break;
 	case MP_TessellationMultiplier:
 	case MP_WorldDisplacement:
diff --git a/Engine/Source/Runtime/Engine/Private/Materials/MaterialExpressions.cpp b/Engine/Source/Runtime/Engine/Private/Materials/MaterialExpressions.cpp
--- a/Engine/Source/Runtime/Engine/Private/Materials/MaterialExpressions.cpp
+++ b/Engine/Source/Runtime/Engine/Private/Materials/MaterialExpressions.cpp
@@ -897,14 +897,14 @@ void UMaterialExpression::ConnectToPreviewMaterial(UMaterial* InMaterial, int32
 
 		if( bUseMaterialAttributes )
 		{
-			InMaterial->SetLightingModel(MLM_DefaultLit);
+			InMaterial->SetShadingModel(MSM_DefaultLit);
 			InMaterial->bUseMaterialAttributes = true;
 			FExpressionInput* MaterialInput = InMaterial->GetExpressionInputForProperty(MP_MaterialAttributes);
 			ConnectExpression( MaterialInput, OutputIndex );
 		}
 		else
 		{
-			InMaterial->SetLightingModel(MLM_Unlit);
+			InMaterial->SetShadingModel(MSM_Unlit);
 			InMaterial->bUseMaterialAttributes = false;
 
 			// Connect the selected expression to the emissive node of the expression preview material.  The emissive material is not affected by light which is why its a good choice.
diff --git a/Engine/Source/Runtime/Engine/Private/Materials/MaterialInstance.cpp b/Engine/Source/Runtime/Engine/Private/Materials/MaterialInstance.cpp
--- a/Engine/Source/Runtime/Engine/Private/Materials/MaterialInstance.cpp
+++ b/Engine/Source/Runtime/Engine/Private/Materials/MaterialInstance.cpp
@@ -2331,14 +2331,14 @@ EBlendMode UMaterialInstance::GetBlendMode_Internal() const
 	return GetMaterial()->GetBlendMode();
 }
 
-EMaterialLightingModel UMaterialInstance::GetLightingModel_Internal() const
+EMaterialShadingModel UMaterialInstance::GetShadingModel_Internal() const
 {
 	checkSlow(IsInGameThread());
-	if( bOverrideBaseProperties && BasePropertyOverrides && BasePropertyOverrides->bOverride_LightingModel )
+	if( bOverrideBaseProperties && BasePropertyOverrides && BasePropertyOverrides->bOverride_ShadingModel )
 	{
-		return BasePropertyOverrides->LightingModel;
+		return BasePropertyOverrides->ShadingModel;
 	}
-	return GetMaterial()->GetLightingModel();
+	return GetMaterial()->GetShadingModel();
 }
 
 bool UMaterialInstance::IsTwoSided_Internal() const
@@ -2371,11 +2371,11 @@ bool UMaterialInstance::GetBlendModeOverride(EBlendMode& OutResult) const
 	return false;
 }
 
-bool UMaterialInstance::GetLightingModelOverride(EMaterialLightingModel& OutResult) const
+bool UMaterialInstance::GetShadingModelOverride(EMaterialShadingModel& OutResult) const
 {
-	if( bOverrideBaseProperties && BasePropertyOverrides && BasePropertyOverrides->bOverride_LightingModel )
+	if( bOverrideBaseProperties && BasePropertyOverrides && BasePropertyOverrides->bOverride_ShadingModel )
 	{
-		OutResult = BasePropertyOverrides->LightingModel;
+		OutResult = BasePropertyOverrides->ShadingModel;
 		return true;
 	}
 	return false;
diff --git a/Engine/Source/Runtime/Engine/Private/Materials/MaterialInterface.cpp b/Engine/Source/Runtime/Engine/Private/Materials/MaterialInterface.cpp
--- a/Engine/Source/Runtime/Engine/Private/Materials/MaterialInterface.cpp
+++ b/Engine/Source/Runtime/Engine/Private/Materials/MaterialInterface.cpp
@@ -28,7 +28,7 @@ FMaterialRelevance UMaterialInterface::GetRelevance_Internal(const UMaterial* Ma
 	{
 		const FMaterialResource* MaterialResource = Material->GetMaterialResource(GRHIFeatureLevel);
 		const bool bIsTranslucent = IsTranslucentBlendMode((EBlendMode)GetBlendMode());
-		const bool bIsLit = GetLightingModel() != MLM_Unlit;
+		const bool bIsLit = GetShadingModel() != MSM_Unlit;
 		// Determine the material's view relevance.
 		FMaterialRelevance MaterialRelevance;
 		MaterialRelevance.bOpaque = !bIsTranslucent;
@@ -248,15 +248,15 @@ bool UMaterialInterface::IsTwoSided()const
 	return GetRenderProxy(0)->GetMaterial(GRHIFeatureLevel)->IsTwoSided();
 }
 
-EMaterialLightingModel UMaterialInterface::GetLightingModel()const
+EMaterialShadingModel UMaterialInterface::GetShadingModel()const
 {
 	if( IsInGameThread() )
 	{
-		return GetLightingModel_Internal();
+		return GetShadingModel_Internal();
 	}
 
 	//We're on the render thread so get it from the proxy.
-	return GetRenderProxy(0)->GetMaterial(GRHIFeatureLevel)->GetLightingModel();
+	return GetRenderProxy(0)->GetMaterial(GRHIFeatureLevel)->GetShadingModel();
 }
 
 float UMaterialInterface::GetOpacityMaskClipValue_Internal()const
@@ -274,9 +274,9 @@ bool UMaterialInterface::IsTwoSided_Internal()const
 	return false;
 }
 
-EMaterialLightingModel UMaterialInterface::GetLightingModel_Internal()const
+EMaterialShadingModel UMaterialInterface::GetShadingModel_Internal()const
 {
-	return MLM_DefaultLit;
+	return MSM_DefaultLit;
 }
 
 void UMaterialInterface::SetFeatureLevelToCompile(ERHIFeatureLevel::Type FeatureLevel, bool bShouldCompile)
diff --git a/Engine/Source/Runtime/Engine/Private/Materials/MaterialShader.cpp b/Engine/Source/Runtime/Engine/Private/Materials/MaterialShader.cpp
--- a/Engine/Source/Runtime/Engine/Private/Materials/MaterialShader.cpp
+++ b/Engine/Source/Runtime/Engine/Private/Materials/MaterialShader.cpp
@@ -33,19 +33,19 @@ uint32 FMaterialShaderMap::NextCompilingId = 1;
  */
 TMap<TRefCountPtr<FMaterialShaderMap>, TArray<FMaterial*> > FMaterialShaderMap::ShaderMapsBeingCompiled;
 
-/** Converts an EMaterialLightingModel to a string description. */
-FString GetLightingModelString(EMaterialLightingModel LightingModel)
+/** Converts an EMaterialShadingModel to a string description. */
+FString GetShadingModelString(EMaterialShadingModel ShadingModel)
 {
-	FString LightingModelName;
-	switch(LightingModel)
+	FString ShadingModelName;
+	switch(ShadingModel)
 	{
-		case MLM_DefaultLit: LightingModelName = TEXT("MLM_DefaultLit"); break;
-		case MLM_Unlit: LightingModelName = TEXT("MLM_Unlit"); break;
-		case MLM_Subsurface: LightingModelName = TEXT("MLM_Subsurface"); break;
-		case MLM_PreintegratedSkin: LightingModelName = TEXT("MLM_PreintegratedSkin"); break;
-		default: LightingModelName = TEXT("Unknown"); break;
+		case MSM_DefaultLit: ShadingModelName = TEXT("MSM_DefaultLit"); break;
+		case MSM_Unlit: ShadingModelName = TEXT("MSM_Unlit"); break;
+		case MSM_Subsurface: ShadingModelName = TEXT("MSM_Subsurface"); break;
+		case MSM_PreintegratedSkin: ShadingModelName = TEXT("MSM_PreintegratedSkin"); break;
+		default: ShadingModelName = TEXT("Unknown"); break;
 	}
-	return LightingModelName;
+	return ShadingModelName;
 }
 
 /** Converts an EBlendMode to a string description. */
@@ -76,12 +76,12 @@ void UpdateMaterialShaderCompilingStats(const FMaterial* Material)
 		default: INC_DWORD_STAT_BY(STAT_ShaderCompiling_NumTransparentMaterialShaders,1); break;
 	}
 
-	switch(Material->GetLightingModel())
+	switch(Material->GetShadingModel())
 	{
-		case MLM_Subsurface: 
-		case MLM_PreintegratedSkin: 
-		case MLM_DefaultLit: INC_DWORD_STAT_BY(STAT_ShaderCompiling_NumLitMaterialShaders,1); break;
-		case MLM_Unlit: INC_DWORD_STAT_BY(STAT_ShaderCompiling_NumUnlitMaterialShaders,1); break;
+		case MSM_Subsurface: 
+		case MSM_PreintegratedSkin: 
+		case MSM_DefaultLit: INC_DWORD_STAT_BY(STAT_ShaderCompiling_NumLitMaterialShaders,1); break;
+		case MSM_Unlit: INC_DWORD_STAT_BY(STAT_ShaderCompiling_NumUnlitMaterialShaders,1); break;
 		default: break;
 	};
 
diff --git a/Engine/Source/Runtime/Engine/Private/Materials/MaterialShared.cpp b/Engine/Source/Runtime/Engine/Private/Materials/MaterialShared.cpp
--- a/Engine/Source/Runtime/Engine/Private/Materials/MaterialShared.cpp
+++ b/Engine/Source/Runtime/Engine/Private/Materials/MaterialShared.cpp
@@ -807,12 +807,12 @@ EBlendMode FMaterialResource::GetBlendMode() const
 	return Ret;
 }
 
-EMaterialLightingModel FMaterialResource::GetLightingModel() const 
+EMaterialShadingModel FMaterialResource::GetShadingModel() const 
 { 
-	EMaterialLightingModel Ret;
-	if( !(MaterialInstance && MaterialInstance->GetLightingModelOverride(Ret)) )
+	EMaterialShadingModel Ret;
+	if( !(MaterialInstance && MaterialInstance->GetShadingModelOverride(Ret)) )
 	{
-		Ret = Material->GetLightingModel_Internal();
+		Ret = Material->GetShadingModel_Internal();
 	}
 	return Ret;
 }
@@ -877,7 +877,7 @@ void FMaterialResource::NotifyCompilationFinished()
 }
 
 /**
- * Gets instruction counts that best represent the likely usage of this material based on lighting model and other factors.
+ * Gets instruction counts that best represent the likely usage of this material based on shading model and other factors.
  * @param Descriptions - an array of descriptions to be populated
  * @param InstructionCounts - an array of instruction counts matching the Descriptions.  
  *		The dimensions of these arrays are guaranteed to be identical and all values are valid.
@@ -927,7 +927,7 @@ void FMaterialResource::GetRepresentativeShaderTypesAndDescriptions(TArray<FStri
 
 	if (GetFeatureLevel() > ERHIFeatureLevel::ES2)
 	{
-		if (GetLightingModel() == MLM_Unlit)
+		if (GetShadingModel() == MSM_Unlit)
 		{
 			//unlit materials are never lightmapped
 			new (ShaderTypeNames)FString(TEXT("TBasePassPSFNoLightMapPolicy"));
@@ -961,7 +961,7 @@ void FMaterialResource::GetRepresentativeShaderTypesAndDescriptions(TArray<FStri
 		const TCHAR* ShaderSuffix = bMobileHDR ? TEXT("HDRLinear64") : TEXT("LDRGamma32");
 		const TCHAR* DescSuffix = bMobileHDR ? TEXT(" (HDR)") : TEXT(" (LDR)");
 
-		if (GetLightingModel() == MLM_Unlit)
+		if (GetShadingModel() == MSM_Unlit)
 		{
 			//unlit materials are never lightmapped
 			new (ShaderTypeNames)FString(FString::Printf(TEXT("TBasePassForForwardShadingPSFNoLightMapPolicy%s"), ShaderSuffix));
@@ -1153,15 +1153,15 @@ void FMaterial::SetupMaterialEnvironment(
 	OutEnvironment.SetDefine(TEXT("MATERIAL_USE_LM_DIRECTIONALITY"), UseLmDirectionality() ? TEXT("1") : TEXT("0"));
 	OutEnvironment.SetDefine(TEXT("MATERIAL_INJECT_EMISSIVE_INTO_LPV"), ShouldInjectEmissiveIntoLPV() ? TEXT("1") : TEXT("0"));
 
-	switch(GetLightingModel())
+	switch(GetShadingModel())
 	{
-	case MLM_DefaultLit: OutEnvironment.SetDefine(TEXT("MATERIAL_LIGHTINGMODEL_DEFAULT_LIT"),TEXT("1")); break;
-	case MLM_Subsurface: OutEnvironment.SetDefine(TEXT("MATERIAL_LIGHTINGMODEL_SUBSURFACE"),TEXT("1")); break;
-	case MLM_PreintegratedSkin: OutEnvironment.SetDefine(TEXT("MATERIAL_LIGHTINGMODEL_PREINTEGRATED_SKIN"),TEXT("1")); break;
-	case MLM_Unlit: OutEnvironment.SetDefine(TEXT("MATERIAL_LIGHTINGMODEL_UNLIT"),TEXT("1")); break;
+	case MSM_DefaultLit: OutEnvironment.SetDefine(TEXT("MATERIAL_SHADINGMODEL_DEFAULT_LIT"),TEXT("1")); break;
+	case MSM_Subsurface: OutEnvironment.SetDefine(TEXT("MATERIAL_SHADINGMODEL_SUBSURFACE"),TEXT("1")); break;
+	case MSM_PreintegratedSkin: OutEnvironment.SetDefine(TEXT("MATERIAL_SHADINGMODEL_PREINTEGRATED_SKIN"),TEXT("1")); break;
+	case MSM_Unlit: OutEnvironment.SetDefine(TEXT("MATERIAL_SHADINGMODEL_UNLIT"),TEXT("1")); break;
 	default: 
-		UE_LOG(LogMaterial, Warning, TEXT("Unknown material lighting model: %u  Setting to MLM_DefaultLit"),(int32)GetLightingModel());
-		OutEnvironment.SetDefine(TEXT("MATERIAL_LIGHTINGMODEL_DEFAULT_LIT"),TEXT("1"));
+		UE_LOG(LogMaterial, Warning, TEXT("Unknown material shading model: %u  Setting to MSM_DefaultLit"),(int32)GetShadingModel());
+		OutEnvironment.SetDefine(TEXT("MATERIAL_SHADINGMODEL_DEFAULT_LIT"),TEXT("1"));
 	
 	};
 
@@ -1735,7 +1735,7 @@ int32 FMaterialResource::GetSamplerUsage() const
 FString FMaterialResource::GetMaterialUsageDescription() const
 {
 	check(Material);
-	FString BaseDescription = GetLightingModelString(GetLightingModel()) + TEXT(", ") + GetBlendModeString(GetBlendMode());
+	FString BaseDescription = GetShadingModelString(GetShadingModel()) + TEXT(", ") + GetBlendModeString(GetBlendMode());
 
 	if (IsSpecialEngineMaterial())
 	{
@@ -2310,11 +2310,11 @@ void DoMaterialAttributeReorder(FExpressionInput* Input)
 FMaterialInstanceBasePropertyOverrides::FMaterialInstanceBasePropertyOverrides()
 	:bOverride_OpacityMaskClipValue(false)
 	,bOverride_BlendMode(false)
-	,bOverride_LightingModel(false)
+	,bOverride_ShadingModel(false)
 	,bOverride_TwoSided(false)
 	,OpacityMaskClipValue(0.0f)
 	,BlendMode(BLEND_Opaque)
-	,LightingModel(MLM_DefaultLit)
+	,ShadingModel(MSM_DefaultLit)
 	,TwoSided(0)
 {
 
@@ -2324,7 +2324,7 @@ void FMaterialInstanceBasePropertyOverrides::Init(const UMaterialInstance& Insta
 {
 	OpacityMaskClipValue = Instance.GetOpacityMaskClipValue();
 	BlendMode = Instance.GetBlendMode();
-	LightingModel = Instance.GetLightingModel();
+	ShadingModel = Instance.GetShadingModel();
 	TwoSided = (uint32)Instance.IsTwoSided();
 }
 
@@ -2344,11 +2344,11 @@ void FMaterialInstanceBasePropertyOverrides::UpdateHash(FSHA1& HashState) const
  		HashState.Update((const uint8*)&BlendMode, sizeof(BlendMode));
 	}
 
-	if(bOverride_LightingModel)
+	if(bOverride_ShadingModel)
 	{
-		const FString HashString = TEXT("bOverride_LightingModel");
+		const FString HashString = TEXT("bOverride_ShadingModel");
 		HashState.UpdateWithString(*HashString, HashString.Len());
-		HashState.Update((const uint8*)&LightingModel, sizeof(LightingModel));
+		HashState.Update((const uint8*)&ShadingModel, sizeof(ShadingModel));
 	}
 
 	//This does seem like it needs to be in the hash but due to some shaders being added to when two sided is enabled
@@ -2374,8 +2374,8 @@ bool FMaterialInstanceBasePropertyOverrides::Update(const FMaterialInstanceBaseP
 	bRet |= (Updated.bOverride_BlendMode != bOverride_BlendMode);
 	bRet |= Updated.bOverride_BlendMode && (BlendMode != Updated.BlendMode);
 
-	bRet |= (Updated.bOverride_LightingModel != bOverride_LightingModel);
-	bRet |= Updated.bOverride_LightingModel && (LightingModel != Updated.LightingModel);
+	bRet |= (Updated.bOverride_ShadingModel != bOverride_ShadingModel);
+	bRet |= Updated.bOverride_ShadingModel && (ShadingModel != Updated.ShadingModel);
 
 	bRet |= (Updated.bOverride_TwoSided != bOverride_TwoSided);
 	bRet |= Updated.bOverride_TwoSided && (TwoSided != Updated.TwoSided);
diff --git a/Engine/Source/Runtime/Engine/Private/PrimitiveDrawingUtils.cpp b/Engine/Source/Runtime/Engine/Private/PrimitiveDrawingUtils.cpp
--- a/Engine/Source/Runtime/Engine/Private/PrimitiveDrawingUtils.cpp
+++ b/Engine/Source/Runtime/Engine/Private/PrimitiveDrawingUtils.cpp
@@ -1262,7 +1262,7 @@ int32 DrawRichMesh(
 	else if(PDI->View->Family->EngineShowFlags.LightComplexity)
 	{
 		// Don't render unlit translucency when in 'light complexity' viewmode.
-		if (!Mesh.IsTranslucent() || Mesh.MaterialRenderProxy->GetMaterial(GRHIFeatureLevel)->GetLightingModel() != MLM_Unlit)
+		if (!Mesh.IsTranslucent() || Mesh.MaterialRenderProxy->GetMaterial(GRHIFeatureLevel)->GetShadingModel() != MSM_Unlit)
 		{
 			// Count the number of lights interacting with this primitive.
 			int32 NumDynamicLights = GetRendererModule().GetNumDynamicLightsAffectingPrimitive(PrimitiveSceneProxy->GetPrimitiveSceneInfo(),Mesh.LCI);
@@ -1286,7 +1286,7 @@ int32 DrawRichMesh(
 	else if(!EngineShowFlags.Materials && !PDI->View->bForceShowMaterials)
 	{
 		// Don't render unlit translucency when in 'lighting only' viewmode.
-		if (Mesh.MaterialRenderProxy->GetMaterial(GRHIFeatureLevel)->GetLightingModel() != MLM_Unlit
+		if (Mesh.MaterialRenderProxy->GetMaterial(GRHIFeatureLevel)->GetShadingModel() != MSM_Unlit
 			// Don't render translucency in 'lighting only', since the viewmode works by overriding with an opaque material
 			// This would cause a mismatch of the material's blend mode with the primitive's view relevance,
 			// And make faint particles block the view
diff --git a/Engine/Source/Runtime/Engine/Private/StaticMeshComponent.cpp b/Engine/Source/Runtime/Engine/Private/StaticMeshComponent.cpp
--- a/Engine/Source/Runtime/Engine/Private/StaticMeshComponent.cpp
+++ b/Engine/Source/Runtime/Engine/Private/StaticMeshComponent.cpp
@@ -1233,7 +1233,7 @@ bool UStaticMeshComponent::UsesOnlyUnlitMaterials() const
 				UMaterialInterface*	MaterialInterface	= GetMaterial(LOD.Sections[ElementIndex].MaterialIndex);
 				UMaterial*			Material			= MaterialInterface ? MaterialInterface->GetMaterial() : NULL;
 
-				bUsesOnlyUnlitMaterials = Material && Material->GetLightingModel_Internal() == MLM_Unlit;
+				bUsesOnlyUnlitMaterials = Material && Material->GetShadingModel_Internal() == MSM_Unlit;
 			}
 		}
 		return bUsesOnlyUnlitMaterials;
diff --git a/Engine/Source/Runtime/Engine/Public/MaterialShader.h b/Engine/Source/Runtime/Engine/Public/MaterialShader.h
--- a/Engine/Source/Runtime/Engine/Public/MaterialShader.h
+++ b/Engine/Source/Runtime/Engine/Public/MaterialShader.h
@@ -21,8 +21,8 @@
 		Frequency \
 		);
 
-/** Converts an EMaterialLightingModel to a string description. */
-extern FString GetLightingModelString(EMaterialLightingModel LightingModel);
+/** Converts an EMaterialShadingModel to a string description. */
+extern FString GetShadingModelString(EMaterialShadingModel ShadingModel);
 
 /** Converts an EBlendMode to a string description. */
 extern FString GetBlendModeString(EBlendMode BlendMode);
diff --git a/Engine/Source/Runtime/Engine/Public/MaterialShared.h b/Engine/Source/Runtime/Engine/Public/MaterialShared.h
--- a/Engine/Source/Runtime/Engine/Public/MaterialShared.h
+++ b/Engine/Source/Runtime/Engine/Public/MaterialShared.h
@@ -1108,7 +1108,7 @@ public:
 	virtual bool UseLmDirectionality() const { return true; }
 	virtual bool IsMasked() const = 0;
 	virtual enum EBlendMode GetBlendMode() const = 0;
-	virtual enum EMaterialLightingModel GetLightingModel() const = 0;
+	virtual enum EMaterialShadingModel GetShadingModel() const = 0;
 	virtual enum ETranslucencyLightingMode GetTranslucencyLightingMode() const { return TLM_VolumetricNonDirectional; };
 	virtual float GetOpacityMaskClipValue() const = 0;
 	virtual bool IsDistorted() const { return false; };
@@ -1756,7 +1756,7 @@ public:
 	ENGINE_API virtual uint32 GetDecalBlendMode() const;
 	ENGINE_API virtual uint32 GetMaterialDecalResponse() const;
 	ENGINE_API virtual bool HasNormalConnected() const;
-	ENGINE_API virtual enum EMaterialLightingModel GetLightingModel() const;
+	ENGINE_API virtual enum EMaterialShadingModel GetShadingModel() const;
 	ENGINE_API virtual enum ETranslucencyLightingMode GetTranslucencyLightingMode() const;
 	ENGINE_API virtual float GetOpacityMaskClipValue() const;
 	ENGINE_API virtual bool IsDistorted() const;
@@ -1784,7 +1784,7 @@ public:
 	ENGINE_API virtual void NotifyCompilationFinished() OVERRIDE;
 
 	/**
-	 * Gets instruction counts that best represent the likely usage of this material based on lighting model and other factors.
+	 * Gets instruction counts that best represent the likely usage of this material based on shading model and other factors.
 	 * @param Descriptions - an array of descriptions to be populated
 	 * @param InstructionCounts - an array of instruction counts matching the Descriptions.  
 	 *		The dimensions of these arrays are guaranteed to be identical and all values are valid.
diff --git a/Engine/Source/Runtime/Renderer/Private/BasePassRendering.cpp b/Engine/Source/Runtime/Renderer/Private/BasePassRendering.cpp
--- a/Engine/Source/Runtime/Renderer/Private/BasePassRendering.cpp
+++ b/Engine/Source/Runtime/Renderer/Private/BasePassRendering.cpp
@@ -182,7 +182,7 @@ public:
 				LightMapPolicy,
 				Parameters.BlendMode,
 				Parameters.TextureMode,
-				Parameters.LightingModel != MLM_Unlit && Scene && Scene->SkyLight && !Scene->SkyLight->bHasStaticLighting,
+				Parameters.ShadingModel != MSM_Unlit && Scene && Scene->SkyLight && !Scene->SkyLight->bHasStaticLighting,
 				IsTranslucentBlendMode(Parameters.BlendMode) && Scene->HasAtmosphericFog()
 				),
 			Scene->GetFeatureLevel()
@@ -258,7 +258,7 @@ public:
 		const typename LightMapPolicyType::ElementDataType& LightMapElementData
 		) const
 	{
-		const bool bIsLitMaterial = Parameters.LightingModel != MLM_Unlit;
+		const bool bIsLitMaterial = Parameters.ShadingModel != MSM_Unlit;
 
 #if !(UE_BUILD_SHIPPING || UE_BUILD_TEST)
 		// When rendering masked materials in the shader complexity viewmode, 
diff --git a/Engine/Source/Runtime/Renderer/Private/BasePassRendering.h b/Engine/Source/Runtime/Renderer/Private/BasePassRendering.h
--- a/Engine/Source/Runtime/Renderer/Private/BasePassRendering.h
+++ b/Engine/Source/Runtime/Renderer/Private/BasePassRendering.h
@@ -418,7 +418,7 @@ public:
 	static bool ShouldCache(EShaderPlatform Platform,const FMaterial* Material,const FVertexFactoryType* VertexFactoryType)
 	{
 		// Only compile skylight version for lit materials
-		const bool bCacheShaders = !bEnableSkyLight || (Material->GetLightingModel() != MLM_Unlit);
+		const bool bCacheShaders = !bEnableSkyLight || (Material->GetShadingModel() != MSM_Unlit);
 
 		return bCacheShaders
 			&& TBasePassPixelShaderBaseType<LightMapPolicyType>::ShouldCache(Platform, Material, VertexFactoryType);
@@ -780,7 +780,7 @@ public:
 	const FMaterial* Material;
 	const FPrimitiveSceneProxy* PrimitiveSceneProxy;
 	EBlendMode BlendMode;
-	EMaterialLightingModel LightingModel;
+	EMaterialShadingModel ShadingModel;
 	const bool bAllowFog;
 	/** Whether or not to perform depth test in the pixel shader */
 	const bool bEditorCompositeDepthTest;
@@ -800,7 +800,7 @@ public:
 		Material(InMaterial),
 		PrimitiveSceneProxy(InPrimitiveSceneProxy),
 		BlendMode(InMaterial->GetBlendMode()),
-		LightingModel(InMaterial->GetLightingModel()),
+		ShadingModel(InMaterial->GetShadingModel()),
 		bAllowFog(InbAllowFog),
 		bEditorCompositeDepthTest(bInEditorCompositeDepthTest),
 		TextureMode(InTextureMode)
@@ -822,7 +822,7 @@ public:
 		Material(InMaterial),
 		PrimitiveSceneProxy(InPrimitiveSceneProxy),
 		BlendMode(InMaterial->GetBlendMode()),
-		LightingModel(InMaterial->GetLightingModel()),
+		ShadingModel(InMaterial->GetShadingModel()),
 		bAllowFog(InbAllowFog),
 		bEditorCompositeDepthTest(bInEditorCompositeDepthTest),
 		TextureMode(InTextureMode)
@@ -839,7 +839,7 @@ void ProcessBasePassMesh(
 	)
 {
 	// Check for a cached light-map.
-	const bool bIsLitMaterial = Parameters.LightingModel != MLM_Unlit;
+	const bool bIsLitMaterial = Parameters.ShadingModel != MSM_Unlit;
 	const bool bNeedsSceneTextures = Parameters.Material->NeedsSceneTextures();
 	static const auto AllowStaticLightingVar = IConsoleManager::Get().FindTConsoleVariableDataInt(TEXT("r.AllowStaticLighting"));
 	const bool bAllowStaticLighting = (!AllowStaticLightingVar || AllowStaticLightingVar->GetValueOnRenderThread() != 0);
diff --git a/Engine/Source/Runtime/Renderer/Private/DepthRendering.cpp b/Engine/Source/Runtime/Renderer/Private/DepthRendering.cpp
--- a/Engine/Source/Runtime/Renderer/Private/DepthRendering.cpp
+++ b/Engine/Source/Runtime/Renderer/Private/DepthRendering.cpp
@@ -501,7 +501,7 @@ bool FDepthDrawingPolicyFactory::DrawStaticMesh(
 	bool bDirty = false;
 
 	const FMaterial* Material = StaticMesh.MaterialRenderProxy->GetMaterial(View.GetFeatureLevel());
-	const EMaterialLightingModel LightingModel = Material->GetLightingModel();
+	const EMaterialShadingModel ShadingModel = Material->GetShadingModel();
 	bDirty |= DrawMesh(
 		View,
 		DrawingContext,
diff --git a/Engine/Source/Runtime/Renderer/Private/ForwardBasePassRendering.cpp b/Engine/Source/Runtime/Renderer/Private/ForwardBasePassRendering.cpp
--- a/Engine/Source/Runtime/Renderer/Private/ForwardBasePassRendering.cpp
+++ b/Engine/Source/Runtime/Renderer/Private/ForwardBasePassRendering.cpp
@@ -155,7 +155,7 @@ public:
 		}
 #endif
 
-		const bool bIsLitMaterial = Parameters.LightingModel != MLM_Unlit;
+		const bool bIsLitMaterial = Parameters.ShadingModel != MSM_Unlit;
 
 		TBasePassForForwardShadingDrawingPolicy<LightMapPolicyType> DrawingPolicy(
 			Parameters.Mesh.VertexFactory,
diff --git a/Engine/Source/Runtime/Renderer/Private/ForwardBasePassRendering.h b/Engine/Source/Runtime/Renderer/Private/ForwardBasePassRendering.h
--- a/Engine/Source/Runtime/Renderer/Private/ForwardBasePassRendering.h
+++ b/Engine/Source/Runtime/Renderer/Private/ForwardBasePassRendering.h
@@ -459,7 +459,7 @@ void ProcessBasePassMeshForForwardShading(
 	)
 {
 	// Check for a cached light-map.
-	const bool bIsLitMaterial = Parameters.LightingModel != MLM_Unlit;
+	const bool bIsLitMaterial = Parameters.ShadingModel != MSM_Unlit;
 
 	const FLightMapInteraction LightMapInteraction = (Parameters.Mesh.LCI && bIsLitMaterial) 
 		? Parameters.Mesh.LCI->GetLightMapInteraction() 
diff --git a/Engine/Source/Runtime/Renderer/Private/ForwardTranslucentRendering.cpp b/Engine/Source/Runtime/Renderer/Private/ForwardTranslucentRendering.cpp
--- a/Engine/Source/Runtime/Renderer/Private/ForwardTranslucentRendering.cpp
+++ b/Engine/Source/Runtime/Renderer/Private/ForwardTranslucentRendering.cpp
@@ -125,7 +125,7 @@ public:
 		const typename LightMapPolicyType::ElementDataType& LightMapElementData
 		) const
 	{
-		const bool bIsLitMaterial = Parameters.LightingModel != MLM_Unlit;
+		const bool bIsLitMaterial = Parameters.ShadingModel != MSM_Unlit;
 
 		TBasePassForForwardShadingDrawingPolicy<LightMapPolicyType> DrawingPolicy(
 			Parameters.Mesh.VertexFactory,
diff --git a/Engine/Source/Runtime/Renderer/Private/LightMapDensityRendering.cpp b/Engine/Source/Runtime/Renderer/Private/LightMapDensityRendering.cpp
--- a/Engine/Source/Runtime/Renderer/Private/LightMapDensityRendering.cpp
+++ b/Engine/Source/Runtime/Renderer/Private/LightMapDensityRendering.cpp
@@ -115,7 +115,7 @@ bool FLightMapDensityDrawingPolicyFactory::DrawDynamicMesh(
 		MaterialRenderProxy = GEngine->LevelColorationLitMaterial->GetRenderProxy(false);
 	}
 
-	bool bIsLitMaterial = (Material->GetLightingModel() != MLM_Unlit);
+	bool bIsLitMaterial = (Material->GetShadingModel() != MSM_Unlit);
 	/*const */FLightMapInteraction LightMapInteraction = (Mesh.LCI && bIsLitMaterial) ? Mesh.LCI->GetLightMapInteraction() : FLightMapInteraction();
 	// force simple lightmaps based on system settings
 	bool bAllowHighQualityLightMaps = AllowHighQualityLightmaps() && LightMapInteraction.AllowsHighQualityLightmaps();
diff --git a/Engine/Source/Runtime/Renderer/Private/LightMapRendering.h b/Engine/Source/Runtime/Renderer/Private/LightMapRendering.h
--- a/Engine/Source/Runtime/Renderer/Private/LightMapRendering.h
+++ b/Engine/Source/Runtime/Renderer/Private/LightMapRendering.h
@@ -218,7 +218,7 @@ public:
 		static const auto AllowStaticLightingVar = IConsoleManager::Get().FindTConsoleVariableDataInt(TEXT("r.AllowStaticLighting"));
 
 		// GetValueOnAnyThread() as it's possible that ShouldCache is called from rendering thread. That is to output some error message.
-		return Material->GetLightingModel() != MLM_Unlit 
+		return Material->GetShadingModel() != MSM_Unlit 
 			&& VertexFactoryType->SupportsStaticLighting() 
 			&& (!AllowStaticLightingVar || AllowStaticLightingVar->GetValueOnAnyThread() != 0)
 			&& (Material->IsUsedWithStaticLighting() || Material->IsSpecialEngineMaterial());
@@ -426,7 +426,7 @@ public:
 
 	static bool ShouldCache(EShaderPlatform Platform,const FMaterial* Material,const FVertexFactoryType* VertexFactoryType)
 	{
-		return Material->GetLightingModel() != MLM_Unlit && VertexFactoryType->SupportsStaticLighting();
+		return Material->GetShadingModel() != MSM_Unlit && VertexFactoryType->SupportsStaticLighting();
 	}
 
 	void Set(
@@ -528,7 +528,7 @@ public:
 
 	static bool ShouldCache(EShaderPlatform Platform,const FMaterial* Material,const FVertexFactoryType* VertexFactoryType)
 	{
-		return Material->GetLightingModel() != MLM_Unlit && IsTranslucentBlendMode(Material->GetBlendMode()) && IsFeatureLevelSupported(Platform, ERHIFeatureLevel::SM4);
+		return Material->GetShadingModel() != MSM_Unlit && IsTranslucentBlendMode(Material->GetBlendMode()) && IsFeatureLevelSupported(Platform, ERHIFeatureLevel::SM4);
 	}
 
 	static void ModifyCompilationEnvironment(EShaderPlatform Platform, const FMaterial* Material, FShaderCompilerEnvironment& OutEnvironment)
@@ -638,7 +638,7 @@ public:
 	{
 		static const auto AllowStaticLightingVar = IConsoleManager::Get().FindTConsoleVariableDataInt(TEXT("r.AllowStaticLighting"));
 
-		return Material->GetLightingModel() != MLM_Unlit 
+		return Material->GetShadingModel() != MSM_Unlit 
 			&& !IsTranslucentBlendMode(Material->GetBlendMode()) 
 			&& (!AllowStaticLightingVar || AllowStaticLightingVar->GetValueOnAnyThread() != 0)
 			&& IsFeatureLevelSupported(Platform, ERHIFeatureLevel::SM4);
@@ -720,7 +720,7 @@ public:
 		static const auto AllowStaticLightingVar = IConsoleManager::Get().FindTConsoleVariableDataInt(TEXT("r.AllowStaticLighting"));
 		
 		// @todo Mac OS X: For GL 3.3 devices which don't support volume-texture rendering we need to cache the simpler point indirect lighting shaders.
-		return Material->GetLightingModel() != MLM_Unlit 
+		return Material->GetShadingModel() != MSM_Unlit 
 			&& (IsTranslucentBlendMode(Material->GetBlendMode()) || GetMaxSupportedFeatureLevel(Platform) == ERHIFeatureLevel::ES2
 				|| (PLATFORM_MAC && GetMaxSupportedFeatureLevel(Platform) == ERHIFeatureLevel::SM4))
 			&& (!AllowStaticLightingVar || AllowStaticLightingVar->GetValueOnAnyThread() != 0);
@@ -781,7 +781,7 @@ public:
 	{
 		static IConsoleVariable* AllowStaticLightingVar = IConsoleManager::Get().FindConsoleVariable(TEXT("r.AllowStaticLighting"));
 
-		return Material->GetLightingModel() != MLM_Unlit 
+		return Material->GetShadingModel() != MSM_Unlit 
 			&& IsTranslucentBlendMode(Material->GetBlendMode()) 
 			&& (!AllowStaticLightingVar || AllowStaticLightingVar->GetInt() != 0)
 			&& FSelfShadowedTranslucencyPolicy::ShouldCache(Platform, Material, VertexFactoryType);
@@ -818,7 +818,7 @@ public:
 
 	static bool ShouldCache(EShaderPlatform Platform,const FMaterial* Material,const FVertexFactoryType* VertexFactoryType)
 	{
-		return Material->GetLightingModel() != MLM_Unlit;
+		return Material->GetShadingModel() != MSM_Unlit;
 	}
 
 	static void ModifyCompilationEnvironment(EShaderPlatform Platform, const FMaterial* Material, FShaderCompilerEnvironment& OutEnvironment)
diff --git a/Engine/Source/Runtime/Renderer/Private/ShadowRendering.cpp b/Engine/Source/Runtime/Renderer/Private/ShadowRendering.cpp
--- a/Engine/Source/Runtime/Renderer/Private/ShadowRendering.cpp
+++ b/Engine/Source/Runtime/Renderer/Private/ShadowRendering.cpp
@@ -148,7 +148,7 @@ public:
 			// Reflective shadow map shaders must be compiled for every material because they access the material normal
 			return !bUsePositionOnlyStream
 				// Don't render ShadowDepth for translucent unlit materials, unless we're injecting emissive
-				&& ((!IsTranslucentBlendMode(Material->GetBlendMode()) && Material->GetLightingModel() != MLM_Unlit) || Material->ShouldInjectEmissiveIntoLPV() )
+				&& ((!IsTranslucentBlendMode(Material->GetBlendMode()) && Material->GetShadingModel() != MSM_Unlit) || Material->ShouldInjectEmissiveIntoLPV() )
 				&& IsFeatureLevelSupported(Platform, ERHIFeatureLevel::SM5);
 		}
 		else
@@ -161,7 +161,7 @@ public:
 					// Only compile position-only shaders for vertex factories that support it.
 					&& (!bUsePositionOnlyStream || VertexFactoryType->SupportsPositionOnly())
 					// Don't render ShadowDepth for translucent unlit materials
-					&& (!IsTranslucentBlendMode(Material->GetBlendMode()) && Material->GetLightingModel() != MLM_Unlit)
+					&& (!IsTranslucentBlendMode(Material->GetBlendMode()) && Material->GetShadingModel() != MSM_Unlit)
 					&& IsFeatureLevelSupported(Platform, ERHIFeatureLevel::SM3);
 		}
 	}
@@ -493,7 +493,7 @@ public:
 			// Reflective shadow map shaders must be compiled for every material because they access the material normal
 			return 
 				// Only compile one pass point light shaders for feature levels >= SM4
-				( (!IsTranslucentBlendMode(Material->GetBlendMode()) && Material->GetLightingModel() != MLM_Unlit) || Material->ShouldInjectEmissiveIntoLPV() )
+				( (!IsTranslucentBlendMode(Material->GetBlendMode()) && Material->GetShadingModel() != MSM_Unlit) || Material->ShouldInjectEmissiveIntoLPV() )
 				&& IsFeatureLevelSupported(Platform, ERHIFeatureLevel::SM5);
 		}
 		else
@@ -507,7 +507,7 @@ public:
 				// Only compile one pass point light shaders for feature levels >= SM4
 				&& (ShaderMode != PixelShadowDepth_OnePassPointLight || IsFeatureLevelSupported(Platform, ERHIFeatureLevel::SM4))
 				// Don't render ShadowDepth for translucent unlit materials
-				&& (!IsTranslucentBlendMode(Material->GetBlendMode()) && Material->GetLightingModel() != MLM_Unlit)
+				&& (!IsTranslucentBlendMode(Material->GetBlendMode()) && Material->GetShadingModel() != MSM_Unlit)
 				&& IsFeatureLevelSupported(Platform, ERHIFeatureLevel::SM3);
 		}
 	}
@@ -954,11 +954,11 @@ void FShadowDepthDrawingPolicyFactory::AddStaticMesh(FScene* Scene,FStaticMesh*
 		const FMaterialRenderProxy* MaterialRenderProxy = StaticMesh->MaterialRenderProxy;
 		const FMaterial* Material = MaterialRenderProxy->GetMaterial(Scene->GetFeatureLevel());
 		const EBlendMode BlendMode = Material->GetBlendMode();
-		const EMaterialLightingModel LightingModel = Material->GetLightingModel();
+		const EMaterialShadingModel ShadingModel = Material->GetShadingModel();
 
 		const bool bLightPropagationVolume = UseLightPropagationVolumeRT();
 		const bool bTwoSided  = Material->IsTwoSided() || StaticMesh->PrimitiveSceneInfo->Proxy->CastsShadowAsTwoSided();
-		const bool bLitOpaque = !IsTranslucentBlendMode(BlendMode) && LightingModel != MLM_Unlit;
+		const bool bLitOpaque = !IsTranslucentBlendMode(BlendMode) && ShadingModel != MSM_Unlit;
 		if(bLightPropagationVolume && (bLitOpaque || Material->ShouldInjectEmissiveIntoLPV()))
 		{
 			// Add the static mesh to the shadow's subject draw list.
@@ -1021,12 +1021,12 @@ bool FShadowDepthDrawingPolicyFactory::DrawDynamicMesh(
 		const FMaterialRenderProxy* MaterialRenderProxy = Mesh.MaterialRenderProxy;
 		const FMaterial* Material = MaterialRenderProxy->GetMaterial(View.GetFeatureLevel());
 		const EBlendMode BlendMode = Material->GetBlendMode();
-		const EMaterialLightingModel LightingModel = Material->GetLightingModel();
+		const EMaterialShadingModel ShadingModel = Material->GetShadingModel();
 
 		const bool bOnePassPointLightShadow = Context.ShadowInfo->bOnePassPointLightShadow;
 		const bool bReflectiveShadowmap = Context.ShadowInfo->bReflectiveShadowmap && !bOnePassPointLightShadow;
 
-		if ( (!IsTranslucentBlendMode(BlendMode) || bReflectiveShadowmap ) && LightingModel != MLM_Unlit )
+		if ( (!IsTranslucentBlendMode(BlendMode) || bReflectiveShadowmap ) && ShadingModel != MSM_Unlit )
 		{
 			const bool bDirectionalLight = Context.ShadowInfo->bDirectionalLight;
 			const bool bPreShadow = Context.ShadowInfo->bPreShadow;
diff --git a/Engine/Source/Runtime/Renderer/Private/ShadowSetup.cpp b/Engine/Source/Runtime/Renderer/Private/ShadowSetup.cpp
--- a/Engine/Source/Runtime/Renderer/Private/ShadowSetup.cpp
+++ b/Engine/Source/Runtime/Renderer/Private/ShadowSetup.cpp
@@ -812,9 +812,9 @@ void FProjectedShadowInfo::AddSubjectPrimitive(FPrimitiveSceneInfo* PrimitiveSce
 							const FMaterialRenderProxy* MaterialRenderProxy = StaticMesh.MaterialRenderProxy;
 							const FMaterial* Material = MaterialRenderProxy->GetMaterial(GRHIFeatureLevel);
 							const EBlendMode BlendMode = Material->GetBlendMode();
-							const EMaterialLightingModel LightingModel = Material->GetLightingModel();
+							const EMaterialShadingModel ShadingModel = Material->GetShadingModel();
 
-							if(((!IsTranslucentBlendMode(BlendMode)) && LightingModel != MLM_Unlit) || (bReflectiveShadowmap && Material->ShouldInjectEmissiveIntoLPV())) 
+							if(((!IsTranslucentBlendMode(BlendMode)) && ShadingModel != MSM_Unlit) || (bReflectiveShadowmap && Material->ShouldInjectEmissiveIntoLPV())) 
 							{
 								const bool bTwoSided = Material->IsTwoSided() || PrimitiveSceneInfo->Proxy->CastsShadowAsTwoSided();
 								OverrideWithDefaultMaterialForShadowDepth(MaterialRenderProxy, Material, bReflectiveShadowmap, GRHIFeatureLevel);
diff --git a/Engine/Source/Runtime/Renderer/Private/TranslucentRendering.cpp b/Engine/Source/Runtime/Renderer/Private/TranslucentRendering.cpp
--- a/Engine/Source/Runtime/Renderer/Private/TranslucentRendering.cpp
+++ b/Engine/Source/Runtime/Renderer/Private/TranslucentRendering.cpp
@@ -252,7 +252,7 @@ public:
 		const typename LightMapPolicyType::ElementDataType& LightMapElementData
 		) const
 	{
-		const bool bIsLitMaterial = Parameters.LightingModel != MLM_Unlit;
+		const bool bIsLitMaterial = Parameters.ShadingModel != MSM_Unlit;
 
 		const FScene* Scene = Parameters.PrimitiveSceneProxy ? Parameters.PrimitiveSceneProxy->GetPrimitiveSceneInfo()->Scene : NULL;
 
@@ -317,7 +317,7 @@ bool FTranslucencyDrawingPolicyFactory::DrawMesh(
 	// Determine the mesh's material and blend mode.
 	const FMaterial* Material = Mesh.MaterialRenderProxy->GetMaterial(View.GetFeatureLevel());
 	const EBlendMode BlendMode = Material->GetBlendMode();
-	const EMaterialLightingModel LightingModel = Material->GetLightingModel();
+	const EMaterialShadingModel ShadingModel = Material->GetShadingModel();
 
 	// Only render translucent materials.
 	if(IsTranslucentBlendMode(BlendMode))
