Fixed dynamic object lighting by removing hardcoded gbuffer lookup which didn't use decode gbuffer functions.

diff --git a/Engine/Shaders/DeferredLightPixelShaders.usf b/Engine/Shaders/DeferredLightPixelShaders.usf
--- a/Engine/Shaders/DeferredLightPixelShaders.usf
+++ b/Engine/Shaders/DeferredLightPixelShaders.usf
@@ -47,13 +47,11 @@ void DirectionalPixelMain(
 {
 	OutColor = 0;
 	float3 CameraVector = normalize(ScreenVector);
-	float4 WorldNormal = Texture2DSampleLevel(GBufferATexture, GBufferATextureSampler, InUV, 0) * float4(2,2,2,1) - float4(1,1,1,0);
 
 	FScreenSpaceData ScreenSpaceData = GetScreenSpaceData(InUV);
 
-	BRANCH
 	// Only light pixels marked as using deferred shading
-	if (WorldNormal.w > 0)
+	BRANCH if( ScreenSpaceData.GBuffer.ShadingModelID > 0 )
 	{
 		float SceneDepth = CalcSceneDepth(InUV);
 		float3 WorldPosition = ScreenVector * SceneDepth + View.ViewOrigin.xyz;
@@ -74,13 +72,11 @@ void RadialPixelMain(
 {
 	OutColor = 0;
 	float2 ScreenUV = InScreenPosition.xy / InScreenPosition.w * View.ScreenPositionScaleBias.xy + View.ScreenPositionScaleBias.wz;
-	float4 WorldNormal = Texture2DSampleLevel(GBufferATexture, GBufferATextureSampler, ScreenUV, 0) * float4(2,2,2,1) - float4(1,1,1,0);
 	
 	FScreenSpaceData ScreenSpaceData = GetScreenSpaceData(ScreenUV);
 
-	BRANCH
 	// Only light pixels marked as using deferred shading
-	if (WorldNormal.w > 0)
+	BRANCH if( ScreenSpaceData.GBuffer.ShadingModelID > 0 )
 	{
 		float SceneDepth = CalcSceneDepth(ScreenUV);
 		float4 HomogeneousWorldPosition = mul(float4(InScreenPosition.xy / InScreenPosition.w * SceneDepth, SceneDepth, 1), View.ScreenToWorld);
diff --git a/Engine/Shaders/StationaryLightOverlapShaders.usf b/Engine/Shaders/StationaryLightOverlapShaders.usf
--- a/Engine/Shaders/StationaryLightOverlapShaders.usf
+++ b/Engine/Shaders/StationaryLightOverlapShaders.usf
@@ -18,10 +18,10 @@ void OverlapDirectionalPixelMain(
 	)
 {
 	OutColor = 0;
-	float4 WorldNormal = Texture2DSampleLevel(GBufferATexture, GBufferATextureSampler, InUV, 0) * float4(2,2,2,1) - float4(1,1,1,0);
+	FGBufferData GBuffer = GetGBufferData(InUV);
 
 	// Only light pixels marked as using deferred shading
-	if (WorldNormal.w > 0)
+	if( GBuffer.ShadingModelID > 0 )
 	{
 		// Output minor complexity to visualize overlap
 		OutColor.rgb = .04f;
@@ -41,13 +41,11 @@ void OverlapRadialPixelMain(
 {
 	OutColor = 0;
 	float2 ScreenUV = InScreenPosition.xy / InScreenPosition.w * View.ScreenPositionScaleBias.xy + View.ScreenPositionScaleBias.wz;
-	float4 WorldNormal = Texture2DSampleLevel(GBufferATexture, GBufferATextureSampler, ScreenUV, 0) * float4(2,2,2,1) - float4(1,1,1,0);
 	
-	FScreenSpaceData ScreenSpaceData = GetScreenSpaceData(ScreenUV);
+	FGBufferData GBuffer = GetGBufferData(ScreenUV);
 
-	BRANCH
 	// Only light pixels marked as using deferred shading
-	if (WorldNormal.w > 0)
+	BRANCH if( GBuffer.ShadingModelID > 0 )
 	{
 		float SceneDepth = CalcSceneDepth(ScreenUV);
 		float4 HomogeneousWorldPosition = mul(float4(InScreenPosition.xy / InScreenPosition.w * SceneDepth, SceneDepth, 1), View.ScreenToWorld);
@@ -66,8 +64,7 @@ void OverlapRadialPixelMain(
 
 		float RadialAttenuationRatio = length(WorldPosition - DeferredLightUniforms.LightPosition) * DeferredLightUniforms.LightInvRadius;
 
-		if (RadialAttenuationRatio < 1
-			&& SpotFalloff > .0001f)
+		if (RadialAttenuationRatio < 1 && SpotFalloff > .0001f)
 		{
 			// Output minor complexity to visualize overlap
 			OutColor.rgb = .04f;
