Distance field AO using a surface cache
* Provides mid-range stable AO for dynamic rigid meshes
* Movable sky lights are now supported, and distance field AO is used for shadowing Movable sky lighting from dynamic scenes
* Static meshes are preprocessed into signed distance field volumes at mesh build time when the r.AllowMeshDistanceFieldRepresentations project setting is enabled
* Non-uniform scaling does not work with this method (mirroring is fine), animating through world position offset also causes artifacts as the two representations diverge
* Occlusion is computed along cones to reduce over-shadowing, object distance fields are operated on directly (no hierarchy) to obtain enough resolution to prevent leaking, visibility traces are done with cone stepping to allow better parallelization, and shading is done adaptively in space and time using the surface cache.

diff --git a/Engine/Shaders/BasePassPixelShader.usf b/Engine/Shaders/BasePassPixelShader.usf
--- a/Engine/Shaders/BasePassPixelShader.usf
+++ b/Engine/Shaders/BasePassPixelShader.usf
@@ -272,30 +272,6 @@ float3 IndirectlightingCachePrimitiveScale;
 float3 IndirectlightingCacheMinUV;
 float3 IndirectlightingCacheMaxUV;
 
-/** 
- * Computes sky diffuse lighting from the SH irradiance map.  
- * This has the SH basis evaluation and diffuse convolusion weights combined for minimal ALU's - see "Stupid Spherical Harmonics (SH) Tricks" 
- */
-float3 GetSkySHDiffuse(float3 Normal)
-{
-	float4 NormalVector = float4(Normal, 1);
-
-	float3 Intermediate0, Intermediate1, Intermediate2;
-	Intermediate0.x = dot(View.SkyIrradianceEnvironmentMap[0], NormalVector);
-	Intermediate0.y = dot(View.SkyIrradianceEnvironmentMap[1], NormalVector);
-	Intermediate0.z = dot(View.SkyIrradianceEnvironmentMap[2], NormalVector);
-
-	float4 vB = NormalVector.xyzz * NormalVector.yzzx;
-	Intermediate1.x = dot(View.SkyIrradianceEnvironmentMap[3], vB);
-	Intermediate1.y = dot(View.SkyIrradianceEnvironmentMap[4], vB);
-	Intermediate1.z = dot(View.SkyIrradianceEnvironmentMap[5], vB);
-
-	float vC = NormalVector.x * NormalVector.x - NormalVector.y * NormalVector.y;
-	Intermediate2 = View.SkyIrradianceEnvironmentMap[6].xyz * vC;
-
-	return Intermediate0 + Intermediate1 + Intermediate2;
-}
-
 /** Computes sky diffuse lighting, including precomputed shadowing. */
 float3 GetSkyLighting(float3 WorldNormal, float2 LightmapUV)
 {
@@ -719,6 +695,7 @@ void Main(
 		Data.GBufferAO = MaterialAO;
 		Data.Opacity = Opacity;
 		Data.DecalMask = Primitive.DecalReceiverMask;
+		Data.HasDistanceFieldRepresentation = Primitive.HasDistanceFieldRepresentation;
 
 		#if MATERIAL_SHADINGMODEL_UNLIT
 			Data.ShadingModelID = SHADINGMODELID_UNLIT;
diff --git a/Engine/Shaders/DeferredDecal.usf b/Engine/Shaders/DeferredDecal.usf
--- a/Engine/Shaders/DeferredDecal.usf
+++ b/Engine/Shaders/DeferredDecal.usf
@@ -188,6 +188,7 @@ void MainPS(
 	Data.Opacity = Opacity;
 	Data.ShadingModelID = SHADINGMODELID_DEFAULT_LIT;
 	Data.DecalMask = 1;
+	Data.HasDistanceFieldRepresentation = 1;
 
 	#if DECAL_RENDERSTAGE == 0
 	{
diff --git a/Engine/Shaders/DeferredShadingCommon.usf b/Engine/Shaders/DeferredShadingCommon.usf
--- a/Engine/Shaders/DeferredShadingCommon.usf
+++ b/Engine/Shaders/DeferredShadingCommon.usf
@@ -286,6 +286,8 @@ struct FGBufferData
 	uint ShadingModelID;
 	// 0..1 decal receiver mask
 	uint DecalMask;
+	// 0..1
+	uint HasDistanceFieldRepresentation; 
 	// in world units
 	float CustomDepth;
 	// in unreal units (linear), can be used to reconstruct world position,
@@ -335,7 +337,8 @@ void EncodeGBuffer(
 	OutGBufferA.rgb = EncodeNormal( Data.WorldNormal );
 
 	// compress in 2 bits
-	OutGBufferA.a = Data.DecalMask * (1.0 / 3.0);
+	uint PackedAlpha = Data.DecalMask * 2 + (Data.HasDistanceFieldRepresentation);
+	OutGBufferA.a = PackedAlpha * (1.0 / 3.0);
 
 	if (Data.ShadingModelID == SHADINGMODELID_UNLIT)
 	{
@@ -390,7 +393,9 @@ FGBufferData DecodeGBufferData(
 		ret.WorldNormal = normalize(ret.WorldNormal);
 	}
 
-	ret.DecalMask		= (uint)( InGBufferA.a * 3.999 );
+	uint PackedAlpha = (uint)( InGBufferA.a * 3.999 );
+	ret.DecalMask = PackedAlpha & 2;
+	ret.HasDistanceFieldRepresentation = PackedAlpha & 1;
 	ret.ShadingModelID	= (uint)( InGBufferB.a * 255.999 );
 	
 	bool bSubsurface = (ret.ShadingModelID == SHADINGMODELID_SUBSURFACE || ret.ShadingModelID == SHADINGMODELID_PREINTEGRATED_SKIN);
diff --git a/Engine/Shaders/DistanceFieldLightingShared.usf b/Engine/Shaders/DistanceFieldLightingShared.usf
new file mode 100644
--- /dev/null
+++ b/Engine/Shaders/DistanceFieldLightingShared.usf
@@ -0,0 +1,221 @@
+// Copyright 1998-2013 Epic Games, Inc. All Rights Reserved.
+
+/*=============================================================================
+
+=============================================================================*/
+
+#ifndef THREADGROUP_SIZEX
+#define THREADGROUP_SIZEX 1
+#endif
+
+#ifndef THREADGROUP_SIZEY
+#define THREADGROUP_SIZEY 1
+#endif
+
+#define THREADGROUP_TOTALSIZE (THREADGROUP_SIZEX * THREADGROUP_SIZEY)
+
+#ifndef SORTTHREADGROUP_SIZEX
+#define SORTTHREADGROUP_SIZEX 1
+#endif
+
+#ifndef SORTTHREADGROUP_SIZEY
+#define SORTTHREADGROUP_SIZEY 1
+#endif
+
+#ifndef NUM_SAMPLES
+#define NUM_SAMPLES 1
+#endif
+
+#ifndef DOWNSAMPLE_FACTOR
+#define DOWNSAMPLE_FACTOR 1
+#endif
+
+#ifndef CONE_TRACE_PHASE
+#define CONE_TRACE_PHASE -1
+#endif
+
+#ifndef MAX_OBJECTS_PER_TILE
+#define MAX_OBJECTS_PER_TILE 1
+#endif
+
+// Must match C++
+#define NUM_CONE_STEPS 10
+
+float AOMaxDistance;
+float AOStepScale;
+float AOStepExponentScale;
+float AOMaxViewDistance;
+
+#define NUM_CULLED_OBJECT_LISTS 3
+
+float GetPhaseParameters(uint Phase, out uint StartIndex, out uint EndIndex)
+{
+	if (Phase == 0)
+	{
+		StartIndex = 0;
+		EndIndex = 4;
+	}
+	else if (Phase == 1)
+	{
+		StartIndex = 5;
+		EndIndex = 7;
+	}
+	else
+	{
+		StartIndex = 8;
+		EndIndex = 9;
+	}
+
+	return AOStepScale;
+}
+
+float GetStepOffset(uint StepIndex, float StepScale)
+{
+	return StepScale * exp2(AOStepExponentScale * StepIndex);
+}
+
+float3 DistanceFieldVolumePositionToUV(float3 VolumePosition, float3 UVScale, float3 UVAdd)
+{
+	float3 VolumeUV = VolumePosition * UVScale + UVAdd;
+	return VolumeUV;
+}
+
+Texture3D DistanceFieldTexture;
+SamplerState DistanceFieldSampler;
+float3 DistanceFieldAtlasTexelSize;
+
+Buffer<uint> TileHeadDataUnpacked;
+Buffer<uint4> TileHeadData;
+Buffer<uint> TileArrayData;
+uint2 TileListGroupSize;
+
+
+#define UNPACKED_TILE_HEAD 1
+#define COHERENT_OBJECT_LIST_WRITES 1
+
+uint4 GetTileHead(uint2 TileCoordinate)
+{
+	uint TileIndex = TileCoordinate.y * TileListGroupSize.x + TileCoordinate.x;
+#if UNPACKED_TILE_HEAD
+	return uint4(
+		TileHeadDataUnpacked[TileIndex * 4 + 0], 
+		min(TileHeadDataUnpacked[TileIndex * 4 + 1], MAX_OBJECTS_PER_TILE), 
+		min(TileHeadDataUnpacked[TileIndex * 4 + 2], MAX_OBJECTS_PER_TILE), 
+		min(TileHeadDataUnpacked[TileIndex * 4 + 3], MAX_OBJECTS_PER_TILE));
+#else
+	return TileHeadData[TileIndex];
+#endif
+}
+
+Texture2D DistanceFieldNormalTexture;
+SamplerState DistanceFieldNormalSampler;
+
+int CurrentLevelDownsampleFactor;
+float2 AOBufferSize;
+float4 AOScreenPositionScaleBias;
+
+int DownsampleFactorToBaseLevel;
+float2 BaseLevelTexelSize;
+
+float3 EncodeNormalForAO(float3 InNormal)
+{
+	return InNormal.xyz;
+}
+
+float3 DecodeNormalForAO(float3 InNormal)
+{
+	return InNormal.xyz;
+}
+
+uint NumObjects;
+
+// In float4's.  Must match equivalent C++ variables.
+#define OBJECT_DATA_STRIDE 8
+#define OBJECT_DATA2_STRIDE 4
+
+Buffer<float4> ObjectBounds;
+Buffer<float4> ObjectData;
+Buffer<float4> ObjectData2;
+
+float4 LoadObjectPositionAndRadius(uint ObjectIndex)
+{
+	return ObjectBounds.Load(ObjectIndex);
+}
+
+float4x4 LoadObjectWorldToVolume(uint ObjectIndex)
+{
+	float4 M0 = ObjectData.Load(ObjectIndex * OBJECT_DATA_STRIDE + 0);
+	float4 M1 = ObjectData.Load(ObjectIndex * OBJECT_DATA_STRIDE + 1);
+	float4 M2 = ObjectData.Load(ObjectIndex * OBJECT_DATA_STRIDE + 2);
+	float4 M3 = ObjectData.Load(ObjectIndex * OBJECT_DATA_STRIDE + 3);
+
+	return float4x4(M0, M1, M2, M3);
+}
+
+float3 LoadObjectLocalPositionExtent(uint ObjectIndex)
+{
+	return ObjectData.Load(ObjectIndex * OBJECT_DATA_STRIDE + 4).xyz;
+}
+
+float4 LoadObjectUVScale(uint ObjectIndex)
+{
+	return ObjectData.Load(ObjectIndex * OBJECT_DATA_STRIDE + 5).xyzw;
+}
+
+float3 LoadObjectUVAdd(uint ObjectIndex)
+{
+	return ObjectData.Load(ObjectIndex * OBJECT_DATA_STRIDE + 6).xyz;
+}
+
+float4x4 LoadObjectVolumeToWorld(uint ObjectIndex)
+{
+	float4 M0 = ObjectData2.Load(ObjectIndex * OBJECT_DATA2_STRIDE + 0);
+	float4 M1 = ObjectData2.Load(ObjectIndex * OBJECT_DATA2_STRIDE + 1);
+	float4 M2 = ObjectData2.Load(ObjectIndex * OBJECT_DATA2_STRIDE + 2);
+	float4 M3 = ObjectData2.Load(ObjectIndex * OBJECT_DATA2_STRIDE + 3);
+
+	return float4x4(M0, M1, M2, M3);
+}
+
+bool SphereIntersectCone(float4 SphereCenterAndRadius, float3 ConeVertex, float3 ConeAxis, float ConeAngleCos, float ConeAngleSin)
+{
+	float3 U = ConeVertex - (SphereCenterAndRadius.w / ConeAngleSin) * ConeAxis;
+	float3 D = SphereCenterAndRadius.xyz - U;
+	float DSizeSq = dot(D, D);
+	float E = dot(ConeAxis, D);
+
+	if (E > 0 && E * E >= DSizeSq * ConeAngleCos * ConeAngleCos)
+	{
+		D = SphereCenterAndRadius.xyz - ConeVertex;
+		DSizeSq = dot(D, D);
+		E = -dot(ConeAxis, D);
+
+		if (E > 0 && E * E >= DSizeSq * ConeAngleSin * ConeAngleSin)
+		{
+			return DSizeSq <= SphereCenterAndRadius.w * SphereCenterAndRadius.w;
+		}
+		else
+		{
+			return true;
+		}
+	}
+
+	return false;
+}
+
+bool SphereIntersectConeWithDepthRanges(float4 SphereCenterAndRadius, float3 ConeVertex, float3 ConeAxis, float ConeAngleCos, float ConeAngleSin, float4 ConeAxisDepthRanges)
+{
+	if (SphereIntersectCone(SphereCenterAndRadius, ConeVertex, ConeAxis, ConeAngleCos, ConeAngleSin))
+	{
+		float ConeAxisDistance = dot(SphereCenterAndRadius.xyz - ConeVertex, ConeAxis);
+		float2 ConeAxisDistanceMinMax = float2(ConeAxisDistance + SphereCenterAndRadius.w, ConeAxisDistance - SphereCenterAndRadius.w);
+
+		if (ConeAxisDistanceMinMax.x > ConeAxisDepthRanges.x && ConeAxisDistanceMinMax.y < ConeAxisDepthRanges.y
+			|| ConeAxisDistanceMinMax.x > ConeAxisDepthRanges.z && ConeAxisDistanceMinMax.y < ConeAxisDepthRanges.w)
+		{
+			return true;
+		}
+	}
+
+	return false;
+}
\ No newline at end of file
diff --git a/Engine/Shaders/DistanceFieldSurfaceCacheLighting.usf b/Engine/Shaders/DistanceFieldSurfaceCacheLighting.usf
new file mode 100644
--- /dev/null
+++ b/Engine/Shaders/DistanceFieldSurfaceCacheLighting.usf
@@ -0,0 +1,715 @@
+// Copyright 1998-2013 Epic Games, Inc. All Rights Reserved.
+
+/*=============================================================================
+
+=============================================================================*/
+
+#include "Common.usf"
+#include "DeferredShadingCommon.usf"
+#include "SHCommon.usf"
+#include "ReflectionEnvironmentShared.usf"
+#include "DistanceFieldLightingShared.usf"
+
+struct FObjectCullVertexOutput
+{
+	float4 Position : SV_POSITION;
+	nointerpolation float4 PositionAndRadius : TEXCOORD0;
+	nointerpolation uint ObjectIndex : TEXCOORD1;
+};
+
+float ConservativeRadiusScale;
+
+/** */
+void ObjectCullVS(
+	float4 InPosition : ATTRIBUTE0,
+	uint ObjectIndex : SV_InstanceID,
+	out FObjectCullVertexOutput Output
+	)
+{
+	//@todo - implement ConservativelyBoundSphere
+	float4 ObjectPositionAndRadius = LoadObjectPositionAndRadius(ObjectIndex);
+	//@todo - expand to handle conservative rasterization
+	float EffectiveRadius = (ObjectPositionAndRadius.w + AOMaxDistance) * ConservativeRadiusScale;
+	float3 WorldPosition = InPosition.xyz * EffectiveRadius + ObjectPositionAndRadius.xyz;
+	Output.Position = mul(float4(WorldPosition, 1), View.WorldToClip);
+	Output.PositionAndRadius = ObjectPositionAndRadius;
+	Output.ObjectIndex = ObjectIndex;
+} 
+
+RWBuffer<uint> RWTileHeadDataUnpacked;
+RWBuffer<uint> RWTileArrayData;
+
+Buffer<float4> TileConeAxisAndCos;
+Buffer<float4> TileConeDepthRanges;
+
+float2 NumGroups;
+
+void ObjectCullPS(
+	FObjectCullVertexOutput Input, 
+	in float4 SVPos : SV_POSITION,
+	out float4 OutColor : COLOR0)
+{
+	OutColor = 0;
+	
+	uint2 TilePosition = (uint2)SVPos.xy;
+	uint TileIndex = TilePosition.y * NumGroups.x + TilePosition.x;
+	float4 ConeAxisAndCos = TileConeAxisAndCos.Load(TileIndex);
+	float4 ConeAxisDepthRanges = TileConeDepthRanges.Load(TileIndex);
+	float3 TileConeVertex = 0;
+	float3 TileConeAxis = ConeAxisAndCos.xyz;
+	float TileConeAngleCos = ConeAxisAndCos.w;
+	float TileConeAngleSin = sqrt(1 - TileConeAngleCos * TileConeAngleCos);
+
+	// A value of 1 is conservative, but has a huge impact on performance
+	float RadiusScale = .5f;
+	float4 SphereCenterAndRadius = Input.PositionAndRadius;
+	float3 ViewSpaceSphereCenter = mul(float4(SphereCenterAndRadius.xyz + View.PreViewTranslation.xyz, 1), View.TranslatedWorldToView).xyz;
+	
+	int SmallestGroupIndex = -1;
+
+	UNROLL
+	for (int GroupIndex = NUM_CULLED_OBJECT_LISTS - 1; GroupIndex >= 0; GroupIndex--)
+	{
+		uint StartIndex;
+		uint EndIndex;
+		float StepScale = GetPhaseParameters(GroupIndex, StartIndex, EndIndex);
+		float GroupMaxSampleRadius = GetStepOffset(EndIndex, StepScale) * 2 * RadiusScale;
+	
+		BRANCH
+		if (SphereIntersectConeWithDepthRanges(float4(ViewSpaceSphereCenter, SphereCenterAndRadius.w + GroupMaxSampleRadius), TileConeVertex, TileConeAxis, TileConeAngleCos, TileConeAngleSin, ConeAxisDepthRanges))
+		{
+			SmallestGroupIndex = GroupIndex;
+		}
+	}
+
+	if (SmallestGroupIndex >= 0)
+	{
+		uint ArrayIndex;
+		InterlockedAdd(RWTileHeadDataUnpacked[TileIndex * 4 + 1 + SmallestGroupIndex], 1U, ArrayIndex);
+
+		if (ArrayIndex < MAX_OBJECTS_PER_TILE)
+		{
+#if COHERENT_OBJECT_LIST_WRITES
+			uint DataIndex = (ArrayIndex * (uint)(NumGroups.x * NumGroups.y + .5f) + TileIndex) * NUM_CULLED_OBJECT_LISTS + SmallestGroupIndex;
+#else
+			uint TileArrayStart = TileIndex * MAX_OBJECTS_PER_TILE;
+			uint DataIndex = TileArrayStart + ArrayIndex;
+#endif
+
+			RWTileArrayData[DataIndex] = Input.ObjectIndex;
+		}
+	}
+}
+
+
+float2 UVToTileHead;
+
+void ComputeDistanceFieldNormalPS(
+	in float4 UVAndScreenPos : TEXCOORD0, 
+	in float4 SVPos : SV_POSITION,
+	out float4 OutColor : SV_Target0)
+{
+	uint2 TileCoordinate = (uint2)(SVPos) / uint2(THREADGROUP_SIZEX, THREADGROUP_SIZEY);
+	uint4 TileHead = GetTileHead(TileCoordinate);
+
+	float2 ScreenUV = float2((((uint2)SVPos.xy) * DOWNSAMPLE_FACTOR + .5f) * View.ViewSizeAndSceneTexelSize.zw);
+	float2 ScreenPosition = (ScreenUV.xy - View.ScreenPositionScaleBias.wz) / View.ScreenPositionScaleBias.xy;
+
+	float SceneDepth = CalcSceneDepth(ScreenUV);
+	float4 HomogeneousWorldPosition = mul(float4(ScreenPosition * SceneDepth, SceneDepth, 1), View.ScreenToWorld);
+	float3 OpaqueWorldPosition = HomogeneousWorldPosition.xyz / HomogeneousWorldPosition.w;
+	
+	int ClosestDistanceFieldObject = -1;
+	float ClosestDistanceFieldSurface = 10000;
+	
+	LOOP 
+	// Iterate over objects whose surface we might be on
+	for (uint ListObjectIndex = 0; ListObjectIndex < TileHead.y; ListObjectIndex++)
+	{
+#if COHERENT_OBJECT_LIST_WRITES
+		uint ListIndex = 0;
+		uint ObjectIndex = TileArrayData.Load((ListObjectIndex * TileListGroupSize.x * TileListGroupSize.y + TileHead.x) * NUM_CULLED_OBJECT_LISTS + ListIndex);
+#else
+		uint ObjectIndex = TileArrayData.Load(TileHead.x + ListObjectIndex);
+#endif
+		float3 LocalPositionExtent = LoadObjectLocalPositionExtent(ObjectIndex);
+		float4x4 WorldToVolume = LoadObjectWorldToVolume(ObjectIndex);
+		float4 UVScaleAndVolumeScale = LoadObjectUVScale(ObjectIndex);
+		float3 UVAdd = LoadObjectUVAdd(ObjectIndex);
+		float3 VolumeShadingPosition = mul(float4(OpaqueWorldPosition, 1), WorldToVolume).xyz;
+
+		BRANCH
+		if (all(VolumeShadingPosition > -LocalPositionExtent) && all(VolumeShadingPosition < LocalPositionExtent))
+		{
+			float3 LocalShadingUV = DistanceFieldVolumePositionToUV(VolumeShadingPosition, UVScaleAndVolumeScale.xyz, UVAdd);
+			float DistanceFromSurface = Texture3DSampleLevel(DistanceFieldTexture, DistanceFieldSampler, LocalShadingUV, 0).x;
+			float WorldSpaceDistanceFromSurface = DistanceFromSurface * UVScaleAndVolumeScale.w;
+
+			if (WorldSpaceDistanceFromSurface < ClosestDistanceFieldSurface)
+			{
+				ClosestDistanceFieldSurface = WorldSpaceDistanceFromSurface;
+				ClosestDistanceFieldObject = ObjectIndex;
+			}
+		}
+	}
+
+	FGBufferData GBufferData = GetGBufferData(ScreenUV);
+
+	float3 WorldNormal = 0;
+	
+	BRANCH
+	if (ClosestDistanceFieldObject >= 0 && GBufferData.HasDistanceFieldRepresentation > 0)
+	{
+		uint ObjectIndex = ClosestDistanceFieldObject;
+
+		float3 LocalPositionExtent = LoadObjectLocalPositionExtent(ObjectIndex);
+		float4x4 WorldToVolume = LoadObjectWorldToVolume(ObjectIndex);
+		float4 UVScaleAndVolumeScale = LoadObjectUVScale(ObjectIndex);
+		float3 UVAdd = LoadObjectUVAdd(ObjectIndex);
+
+		float4x4 VolumeToWorld = LoadObjectVolumeToWorld(ObjectIndex);
+
+		float3 VolumeShadingPosition = mul(float4(OpaqueWorldPosition, 1), WorldToVolume).xyz;
+		float3 LocalShadingUV = DistanceFieldVolumePositionToUV(VolumeShadingPosition, UVScaleAndVolumeScale.xyz, UVAdd);
+
+		float R = Texture3DSampleLevel(DistanceFieldTexture, DistanceFieldSampler, LocalShadingUV + float3(DistanceFieldAtlasTexelSize.x, 0, 0), 0).x;
+		float L = Texture3DSampleLevel(DistanceFieldTexture, DistanceFieldSampler, LocalShadingUV + float3(-DistanceFieldAtlasTexelSize.x, 0, 0), 0).x;
+		float F = Texture3DSampleLevel(DistanceFieldTexture, DistanceFieldSampler, LocalShadingUV + float3(0, DistanceFieldAtlasTexelSize.y, 0), 0).x;
+		float B = Texture3DSampleLevel(DistanceFieldTexture, DistanceFieldSampler, LocalShadingUV + float3(0, -DistanceFieldAtlasTexelSize.y, 0), 0).x;
+		float U = Texture3DSampleLevel(DistanceFieldTexture, DistanceFieldSampler, LocalShadingUV + float3(0, 0, DistanceFieldAtlasTexelSize.z), 0).x;
+		float D = Texture3DSampleLevel(DistanceFieldTexture, DistanceFieldSampler, LocalShadingUV + float3(0, 0, -DistanceFieldAtlasTexelSize.z), 0).x;
+		
+		float3 Gradient = .5f * float3(R - L, F - B, U - D);
+		float3 LocalNormal = normalize(Gradient);
+		WorldNormal = normalize(mul(LocalNormal, (float3x3)VolumeToWorld));
+		float DistanceAlpha = saturate(abs(ClosestDistanceFieldSurface) / 2);
+		//WorldNormal = normalize(lerp(WorldNormal, GBufferData.WorldNormal, .5f * DistanceAlpha));
+	}
+	else
+	{
+		WorldNormal = GBufferData.WorldNormal;
+	}
+	
+	float DepthSign = GBufferData.HasDistanceFieldRepresentation > 0 ? 1 : -1;
+	// Note: Encoding whether the pixel has a distance field representation in the sign bit of the alpha
+	OutColor = float4(EncodeNormalForAO(WorldNormal), DepthSign * SceneDepth);
+}
+
+#define THREADS_PER_PIXEL 1
+
+#if THREADS_PER_PIXEL != 1
+groupshared int SharedClosestDistanceFieldObject[THREADGROUP_TOTALSIZE * THREADS_PER_PIXEL];
+groupshared float SharedClosestDistanceFieldSurface[THREADGROUP_TOTALSIZE * THREADS_PER_PIXEL];
+#endif
+
+RWTexture2D<float4> RWDistanceFieldNormal;
+
+[numthreads(THREADGROUP_SIZEX * THREADS_PER_PIXEL, THREADGROUP_SIZEY, 1)]
+void ComputeDistanceFieldNormalCS(
+	uint3 GroupId : SV_GroupID,
+	uint3 DispatchThreadId : SV_DispatchThreadID,
+    uint3 GroupThreadId : SV_GroupThreadID) 
+{
+	uint2 PixelCoordinate = uint2(DispatchThreadId.x / THREADS_PER_PIXEL, DispatchThreadId.y);
+	uint ThreadIndex = GroupThreadId.y * THREADGROUP_SIZEX * THREADS_PER_PIXEL + GroupThreadId.x;
+	uint PixelThreadIndex = GroupThreadId.x % THREADS_PER_PIXEL;
+	uint2 TileCoordinate = PixelCoordinate / uint2(THREADGROUP_SIZEX, THREADGROUP_SIZEY);
+	uint4 TileHead = GetTileHead(TileCoordinate);
+
+	float2 ScreenUV = float2((PixelCoordinate * DOWNSAMPLE_FACTOR + .5f) * View.ViewSizeAndSceneTexelSize.zw);
+	float2 ScreenPosition = (ScreenUV.xy - View.ScreenPositionScaleBias.wz) / View.ScreenPositionScaleBias.xy;
+
+	float SceneDepth = CalcSceneDepth(ScreenUV);
+	float4 HomogeneousWorldPosition = mul(float4(ScreenPosition * SceneDepth, SceneDepth, 1), View.ScreenToWorld);
+	float3 OpaqueWorldPosition = HomogeneousWorldPosition.xyz / HomogeneousWorldPosition.w;
+	
+	int ClosestDistanceFieldObject = -1;
+	float ClosestDistanceFieldSurface = 10000;
+	
+	LOOP 
+	// Iterate over objects whose surface we might be on
+	for (uint ListObjectIndex = PixelThreadIndex; ListObjectIndex < TileHead.y; ListObjectIndex += THREADS_PER_PIXEL)
+	{
+#if COHERENT_OBJECT_LIST_WRITES
+		uint ListIndex = 0;
+		uint ObjectIndex = TileArrayData.Load((ListObjectIndex * TileListGroupSize.x * TileListGroupSize.y + TileHead.x) * NUM_CULLED_OBJECT_LISTS + ListIndex);
+#else
+		uint ObjectIndex = TileArrayData.Load(TileHead.x + ListObjectIndex);
+#endif
+		float3 LocalPositionExtent = LoadObjectLocalPositionExtent(ObjectIndex);
+		float4x4 WorldToVolume = LoadObjectWorldToVolume(ObjectIndex);
+		float4 UVScaleAndVolumeScale = LoadObjectUVScale(ObjectIndex);
+		float3 UVAdd = LoadObjectUVAdd(ObjectIndex);
+		float3 VolumeShadingPosition = mul(float4(OpaqueWorldPosition, 1), WorldToVolume).xyz;
+
+		BRANCH
+		if (all(VolumeShadingPosition > -LocalPositionExtent) && all(VolumeShadingPosition < LocalPositionExtent))
+		{
+			float3 LocalShadingUV = DistanceFieldVolumePositionToUV(VolumeShadingPosition, UVScaleAndVolumeScale.xyz, UVAdd);
+			float DistanceFromSurface = Texture3DSampleLevel(DistanceFieldTexture, DistanceFieldSampler, LocalShadingUV, 0).x;
+			float WorldSpaceDistanceFromSurface = DistanceFromSurface * UVScaleAndVolumeScale.w;
+
+			if (WorldSpaceDistanceFromSurface < ClosestDistanceFieldSurface)
+			{
+				ClosestDistanceFieldSurface = WorldSpaceDistanceFromSurface;
+				ClosestDistanceFieldObject = ObjectIndex;
+			}
+		}
+	}
+
+#if THREADS_PER_PIXEL != 1
+	SharedClosestDistanceFieldObject[ThreadIndex] = ClosestDistanceFieldObject;
+	SharedClosestDistanceFieldSurface[ThreadIndex] = ClosestDistanceFieldSurface;
+
+	GroupMemoryBarrierWithGroupSync();
+#endif
+
+	BRANCH
+	if (PixelThreadIndex == 0)
+	{
+		FGBufferData GBufferData = GetGBufferData(ScreenUV);
+
+		float3 WorldNormal = GBufferData.WorldNormal;
+	
+		BRANCH
+		if (GBufferData.HasDistanceFieldRepresentation > 0)
+		{
+#if THREADS_PER_PIXEL == 4
+			float OtherThreadDistance1 = SharedClosestDistanceFieldSurface[ThreadIndex + 1];
+
+			if (OtherThreadDistance1 < ClosestDistanceFieldSurface)
+			{
+				ClosestDistanceFieldSurface = OtherThreadDistance1;
+				ClosestDistanceFieldObject = SharedClosestDistanceFieldObject[ThreadIndex + 1];
+			}
+
+			float OtherThreadDistance2 = SharedClosestDistanceFieldSurface[ThreadIndex + 2];
+
+			if (OtherThreadDistance2 < ClosestDistanceFieldSurface)
+			{
+				ClosestDistanceFieldSurface = OtherThreadDistance2;
+				ClosestDistanceFieldObject = SharedClosestDistanceFieldObject[ThreadIndex + 2];
+			}
+
+			float OtherThreadDistance3 = SharedClosestDistanceFieldSurface[ThreadIndex + 3];
+
+			if (OtherThreadDistance3 < ClosestDistanceFieldSurface)
+			{
+				ClosestDistanceFieldSurface = OtherThreadDistance3;
+				ClosestDistanceFieldObject = SharedClosestDistanceFieldObject[ThreadIndex + 3];
+			}
+#endif
+
+			BRANCH
+			if (ClosestDistanceFieldObject >= 0)
+			{
+				uint ObjectIndex = ClosestDistanceFieldObject;
+
+				float3 LocalPositionExtent = LoadObjectLocalPositionExtent(ObjectIndex);
+				float4x4 WorldToVolume = LoadObjectWorldToVolume(ObjectIndex);
+				float4 UVScaleAndVolumeScale = LoadObjectUVScale(ObjectIndex);
+				float3 UVAdd = LoadObjectUVAdd(ObjectIndex);
+
+				float4x4 VolumeToWorld = LoadObjectVolumeToWorld(ObjectIndex);
+
+				float3 VolumeShadingPosition = mul(float4(OpaqueWorldPosition, 1), WorldToVolume).xyz;
+				float3 LocalShadingUV = DistanceFieldVolumePositionToUV(VolumeShadingPosition, UVScaleAndVolumeScale.xyz, UVAdd);
+
+				float R = Texture3DSampleLevel(DistanceFieldTexture, DistanceFieldSampler, LocalShadingUV + float3(DistanceFieldAtlasTexelSize.x, 0, 0), 0).x;
+				float L = Texture3DSampleLevel(DistanceFieldTexture, DistanceFieldSampler, LocalShadingUV + float3(-DistanceFieldAtlasTexelSize.x, 0, 0), 0).x;
+				float F = Texture3DSampleLevel(DistanceFieldTexture, DistanceFieldSampler, LocalShadingUV + float3(0, DistanceFieldAtlasTexelSize.y, 0), 0).x;
+				float B = Texture3DSampleLevel(DistanceFieldTexture, DistanceFieldSampler, LocalShadingUV + float3(0, -DistanceFieldAtlasTexelSize.y, 0), 0).x;
+				float U = Texture3DSampleLevel(DistanceFieldTexture, DistanceFieldSampler, LocalShadingUV + float3(0, 0, DistanceFieldAtlasTexelSize.z), 0).x;
+				float D = Texture3DSampleLevel(DistanceFieldTexture, DistanceFieldSampler, LocalShadingUV + float3(0, 0, -DistanceFieldAtlasTexelSize.z), 0).x;
+		
+				float3 Gradient = .5f * float3(R - L, F - B, U - D);
+				float3 LocalNormal = normalize(Gradient);
+				WorldNormal = normalize(mul(LocalNormal, (float3x3)VolumeToWorld));
+				float DistanceAlpha = saturate(abs(ClosestDistanceFieldSurface) / 2);
+				//WorldNormal = normalize(lerp(WorldNormal, GBufferData.WorldNormal, .5f * DistanceAlpha));
+			}
+		}
+	
+		float DepthSign = GBufferData.HasDistanceFieldRepresentation > 0 ? 1 : -1;
+		float4 OutValue = float4(EncodeNormalForAO(WorldNormal), DepthSign * SceneDepth);
+		RWDistanceFieldNormal[PixelCoordinate] = OutValue;
+	}
+}
+
+// For some reason gives innaccurate results at lower resolutions
+#define USE_SCREENVECTOR_WORLD_POSITION FINAL_INTERPOLATION_PASS
+
+struct FIrradianceCacheSplatVertexOutput
+{
+	float4 Position : SV_POSITION;
+	nointerpolation float4 PositionRadius : TEXCOORD0;
+	nointerpolation float3 Normal : TEXCOORD1;
+	nointerpolation float3 BentNormal : TEXCOORD2;
+	nointerpolation uint RecordIndex : TEXCOORD3;
+#if USE_SCREENVECTOR_WORLD_POSITION
+	float3 ScreenVector : TEXCOORD4;
+#endif
+};
+
+Buffer<uint> CompactedRedirect;
+Buffer<uint> IrradianceCacheUsed;
+Buffer<float4> IrradianceCachePositionRadius;
+Buffer<float> IrradianceCacheOccluderRadius;
+Buffer<float4> IrradianceCacheNormal;
+Buffer<float4> IrradianceCacheBentNormal;
+float InterpolationRadiusScale;
+float2 NormalizedOffsetToPixelCenter;
+float HackExpand;
+
+/** */
+void IrradianceCacheSplatVS(
+	float2 InPosition : ATTRIBUTE0,
+	float2 InUV       : ATTRIBUTE1,
+	uint JobIndex : SV_InstanceID,
+	out FIrradianceCacheSplatVertexOutput Output
+	)
+{
+	float4 PositionAndRadius = IrradianceCachePositionRadius.Load(JobIndex);
+	float OccluderRadius = IrradianceCacheOccluderRadius.Load(JobIndex);
+	PositionAndRadius.w = min(PositionAndRadius.w, OccluderRadius);
+
+	float ViewToCenterLength = length(PositionAndRadius.xyz - View.ViewOrigin.xyz);
+	float3 NormalizedViewToCenter = (PositionAndRadius.xyz - View.ViewOrigin.xyz) / ViewToCenterLength;
+	PositionAndRadius.w = max(PositionAndRadius.w, ViewToCenterLength * .01f) * InterpolationRadiusScale;
+
+	// Construct a virtual sample position that won't be near plane clipped and will have a positive z after projection
+	float3 VirtualPosition = PositionAndRadius.xyz + NormalizedViewToCenter * PositionAndRadius.w;
+	// The circle verts were placed directly on the circle, so edges will go inside
+	// Also clipping the edge of the circle a bit can save perf
+	float RadiusSlack = 1.1f;
+	// Compute new radius since we moved the center away from the camera, approximate as similar triangles
+	float VirtualRadius = RadiusSlack * PositionAndRadius.w * (PositionAndRadius.w + ViewToCenterLength) / ViewToCenterLength;
+
+#if !FINAL_INTERPOLATION_PASS
+	VirtualRadius += HackExpand * ViewToCenterLength * 4;
+#endif
+
+	float2 CornerUVs = InUV;
+	float3 CornerPosition = VirtualPosition + (View.ViewRight * CornerUVs.x + View.ViewUp * CornerUVs.y) * VirtualRadius;
+
+	Output.Position = mul(float4(CornerPosition.xyz, 1), View.WorldToClip);
+
+#if USE_SCREENVECTOR_WORLD_POSITION
+	Output.ScreenVector = mul(float4(Output.Position.xy / Output.Position.w, 1, 0), View.ScreenToWorld).xyz;
+#endif
+
+	Output.Position.xy += NormalizedOffsetToPixelCenter * Output.Position.w;
+
+	Output.PositionRadius = PositionAndRadius; 
+	Output.Normal = DecodeNormalForAO(IrradianceCacheNormal.Load(JobIndex).xyz);
+	Output.BentNormal = DecodeNormalForAO(IrradianceCacheBentNormal.Load(JobIndex).xyz);
+	Output.RecordIndex = JobIndex;
+} 
+
+float InterpolationAngleNormalization; 
+float InvMinCosPointBehindPlane;
+
+void IrradianceCacheSplatPS(
+	FIrradianceCacheSplatVertexOutput Input, 
+	in float4 SVPos : SV_POSITION,
+	out float4 OutColor : COLOR0)
+{
+	OutColor = 0;
+
+	float2 BaseLevelScreenUV = float2((((uint2)SVPos.xy) * DownsampleFactorToBaseLevel + .5f) * BaseLevelTexelSize);
+	float4 DistanceFieldNormalValue = Texture2DSampleLevel(DistanceFieldNormalTexture, DistanceFieldNormalSampler, BaseLevelScreenUV, 0);
+
+#if USE_SCREENVECTOR_WORLD_POSITION
+
+	float SceneW = abs(DistanceFieldNormalValue.w);
+	float3 OpaqueWorldPosition = View.ViewOrigin.xyz + Input.ScreenVector * SceneW;
+
+#else
+
+	float2 ScreenUV = float2((((uint2)SVPos.xy) * CurrentLevelDownsampleFactor + float2(.5f, .5f)) * View.ViewSizeAndSceneTexelSize.zw);
+	float SceneW = CalcSceneDepth(ScreenUV); 
+	float2 ScreenPosition = (ScreenUV.xy - View.ScreenPositionScaleBias.wz) / View.ScreenPositionScaleBias.xy;
+	float4 HomogeneousWorldPosition = mul(float4(ScreenPosition.xy * SceneW, SceneW, 1), View.ScreenToWorld);
+	float3 OpaqueWorldPosition = HomogeneousWorldPosition.xyz / HomogeneousWorldPosition.w;
+
+#endif
+	
+	float Distance = length(OpaqueWorldPosition - Input.PositionRadius.xyz);
+	float DistanceError = saturate(Distance / Input.PositionRadius.w);
+
+	BRANCH
+	if (DistanceError < 1 && SceneW < AOMaxViewDistance)
+	{ 
+		float3 WorldNormal = DecodeNormalForAO(DistanceFieldNormalValue.xyz);
+		float NormalError = InterpolationAngleNormalization * sqrt(saturate(1 - dot(WorldNormal, Input.Normal)));
+
+		// Don't use a lighting record if it's in front of the query point.
+		// Query points behind the lighting record may have nearby occluders that the lighting record does not see.
+		// Offset the comparison point along the negative normal to prevent self-occlusion
+		float3 RecordToVertexVector = OpaqueWorldPosition - (Input.PositionRadius.xyz - 1 * Input.Normal.xyz);
+		float DistanceToVertex = length(RecordToVertexVector);
+		float PlaneDistance = dot(Input.Normal.xyz, RecordToVertexVector) / DistanceToVertex;
+		
+		// Setup an error metric that goes from 0 if the points are coplanar, to 1 if the point being shaded is at the angle corresponding to MinCosPointBehindPlane behind the plane
+		float PointBehindPlaneError = min(max(PlaneDistance * InvMinCosPointBehindPlane, 0.0f), DistanceToVertex / 3.0f);
+
+		float PrecisionScale = .1f;
+		float OverlapWeightScale = 1;
+		float Weight = OverlapWeightScale * saturate(PrecisionScale * (1 - max(DistanceError, max(NormalError, PointBehindPlaneError))));
+		//float Weight = saturate(PrecisionScale * (1 - max(DistanceError, NormalError)));
+		//float Weight = saturate(PrecisionScale * (1 - DistanceError));
+		//float Weight = saturate(PrecisionScale * (1 - NormalError));
+
+		FLATTEN
+		if (DistanceFieldNormalValue.w < 0)
+		{
+			// Don't splat on pixels that don't have a distance field representation
+			// Otherwise these will be spread to neighbors in the gap filling pass
+			Weight = 0;
+		}
+
+		float VisualizePlacement = Distance < .1f * Input.PositionRadius.w;
+		OutColor.rgb = Input.BentNormal * Weight;
+		OutColor.a = Weight;
+	}
+
+	#if FINAL_INTERPOLATION_PASS
+		// Visualize cache record points
+		//OutColor.a = (DistanceError < .05f) * .1f;
+	#endif
+}
+
+Texture2D BentNormalAOTexture3;
+SamplerState BentNormalAOSampler3;
+
+void AOCombinePS2(
+	in float4 UVAndScreenPos : TEXCOORD0, 
+	in float4 SVPos : SV_POSITION,
+	out float4 OutColor : SV_Target0)
+{
+	float2 ScreenUV = float2((((uint2)SVPos.xy) * DOWNSAMPLE_FACTOR + .5f) * View.ViewSizeAndSceneTexelSize.zw);
+	float SceneDepth = CalcSceneDepth(ScreenUV);
+
+#if 1
+
+	float4 IrradianceCacheAccumulation = Texture2DSampleLevel(BentNormalAOTexture3, BentNormalAOSampler3, UVAndScreenPos.xy, 0);
+	float3 BentNormalAO = float3(0, 0, 0);
+
+	if (IrradianceCacheAccumulation.w > 0)
+	{
+		OutColor.rgb = IrradianceCacheAccumulation.xyz / IrradianceCacheAccumulation.w;
+		OutColor.a = SceneDepth;
+	}
+	else
+	{
+		OutColor.rgb = float3(0, 0, .1f);
+		// Sign bit stores whether texel is valid
+		OutColor.a = -SceneDepth;
+	}
+
+	float DefaultVisibility = .7f;
+
+	// Fade to unoccluded in the distance
+	OutColor.rgb = lerp(DefaultVisibility.xxx, OutColor.rgb, saturate(.003f * (AOMaxViewDistance - SceneDepth)));
+
+	if (SceneDepth > AOMaxViewDistance)
+	{
+		OutColor.a = abs(OutColor.a);
+	}
+
+#else
+	float3 BentNormalAO = Texture2DSampleLevel(BentNormalAOTexture3, BentNormalAOSampler3, UVAndScreenPos.xy, 0).aaa;
+	OutColor = float4(BentNormalAO, SceneDepth); 
+#endif
+}
+
+float2 BentNormalAOTexelSize;
+Texture2D BentNormalAOTexture;
+SamplerState BentNormalAOSampler;
+
+#define HALF_FILL_KERNEL_SIZE 2
+
+void FillGapsPS(
+	in float4 UVAndScreenPos : TEXCOORD0, 
+	in float4 SVPos : SV_POSITION,
+	out float4 OutColor : SV_Target0)
+{
+	float4 CenterValue = Texture2DSampleLevel(BentNormalAOTexture, BentNormalAOSampler, UVAndScreenPos.xy, 0);
+	float SceneDepth = abs(CenterValue.w);
+
+	float4 Accumulation = 0;
+
+	for (float y = -HALF_FILL_KERNEL_SIZE; y <= HALF_FILL_KERNEL_SIZE; y++)
+	{
+		for (float x = -HALF_FILL_KERNEL_SIZE; x <= HALF_FILL_KERNEL_SIZE; x++)
+		{
+			float2 UVOffset = BentNormalAOTexelSize * float2(x, y);
+			float4 TextureValue = Texture2DSampleLevel(BentNormalAOTexture, BentNormalAOSampler, UVAndScreenPos.xy + UVOffset, 0);
+			float SampleSceneDepth = abs(TextureValue.w);
+			float ValidMask = TextureValue.w > 0;
+
+			float DepthWeight = 1 / exp2(abs(SceneDepth - SampleSceneDepth) * .01f);
+
+			float2 Weight2D = 1 / exp2(abs(float2(x, y) * 10.0f / HALF_FILL_KERNEL_SIZE));
+			float ScreenSpaceSpatialWeight = max(Weight2D.x, Weight2D.y);
+
+			float Weight = ValidMask * ScreenSpaceSpatialWeight * DepthWeight;
+
+			Accumulation.rgb += TextureValue.rgb * Weight;
+			Accumulation.a += Weight;
+		}
+	}
+	
+	OutColor = float4(Accumulation.rgb / max(Accumulation.a, .00001f), CenterValue.w);
+}
+
+float4 CameraMotion[5];
+Texture2D HistoryTexture;
+SamplerState HistorySampler;
+float HistoryWeight;
+
+void UpdateHistoryPS(
+	in float4 UVAndScreenPos : TEXCOORD0, 
+	out float4 OutColor : SV_Target0)
+{
+	float3 xyd;
+    xyd.xy = UVAndScreenPos.zw * float2(0.5, -0.5) + 0.5;
+    xyd.z = Texture2DSampleLevel(SceneDepthTexture, SceneDepthTextureSampler, UVAndScreenPos.xy, 0).r;
+
+    float scaleM = 1.0 / (dot(xyd, CameraMotion[0].xyz) + CameraMotion[0].w);
+    float2 mv;
+    mv.x = ((xyd.x * ((CameraMotion[1].x * xyd.y) + (CameraMotion[1].y * xyd.z) + CameraMotion[1].z)) + (CameraMotion[1].w * xyd.y) + (CameraMotion[2].x * xyd.x * xyd.x) + (CameraMotion[2].y * xyd.z) + CameraMotion[2].z) * scaleM;
+    mv.y = ((xyd.y * ((CameraMotion[3].x * xyd.x) + (CameraMotion[3].y * xyd.z) + CameraMotion[3].z)) + (CameraMotion[3].w * xyd.x) + (CameraMotion[4].x * xyd.y * xyd.y) + (CameraMotion[4].y * xyd.z) + CameraMotion[4].z) * scaleM;
+
+    float2 OldUV = UVAndScreenPos.xy + mv / float2(0.5, -0.5) * View.ScreenPositionScaleBias.xy;
+	float2 OldScreenPosition = (OldUV.xy - View.ScreenPositionScaleBias.wz) / View.ScreenPositionScaleBias.xy;
+
+	float4 NewValue = Texture2DSampleLevel(BentNormalAOTexture, BentNormalAOSampler, UVAndScreenPos.xy, 0);
+	float4 HistoryValue = Texture2DSampleLevel(HistoryTexture, HistorySampler, OldUV, 0);
+	
+	//@todo - this doesn't take camera movement into account
+	float DepthDelta = abs(abs(NewValue.a) - abs(HistoryValue.a));
+	float EffectiveHistoryWeight = HistoryWeight;
+
+	FLATTEN
+	if (any(OldScreenPosition > 1) 
+		|| any(OldScreenPosition < -1)
+		|| DepthDelta > 100
+		// Toss history value from pixels that don't have a distance field representation
+		|| HistoryValue.a < 0)
+	{
+		EffectiveHistoryWeight = 0;
+	}
+
+	OutColor.rgb = lerp(NewValue.rgb, HistoryValue.rgb, EffectiveHistoryWeight);
+	OutColor.a = NewValue.a;
+}
+
+void AOUpsamplePS(in float4 UVAndScreenPos : TEXCOORD0, out float4 OutColor : SV_Target0)
+{
+#define BILATERAL_UPSAMPLE 1
+#if BILATERAL_UPSAMPLE
+	float2 LowResBufferSize = floor(View.RenderTargetSize / DOWNSAMPLE_FACTOR);
+	float2 LowResTexelSize = 1.0f / LowResBufferSize;
+	float2 Corner00UV = floor(UVAndScreenPos.xy * LowResBufferSize - .5f) / LowResBufferSize + .5f * LowResTexelSize;
+	float2 BilinearWeights = (UVAndScreenPos.xy - Corner00UV) * LowResBufferSize;
+
+	float4 TextureValues00 = Texture2DSampleLevel(BentNormalAOTexture, BentNormalAOSampler, Corner00UV, 0);
+	float4 TextureValues10 = Texture2DSampleLevel(BentNormalAOTexture, BentNormalAOSampler, Corner00UV + float2(LowResTexelSize.x, 0), 0);
+	float4 TextureValues01 = Texture2DSampleLevel(BentNormalAOTexture, BentNormalAOSampler, Corner00UV + float2(0, LowResTexelSize.y), 0);
+	float4 TextureValues11 = Texture2DSampleLevel(BentNormalAOTexture, BentNormalAOSampler, Corner00UV + LowResTexelSize, 0);
+
+	float4 CornerWeights = float4(
+		(1 - BilinearWeights.y) * (1 - BilinearWeights.x), 
+		(1 - BilinearWeights.y) * BilinearWeights.x,
+		BilinearWeights.y * (1 - BilinearWeights.x),
+		BilinearWeights.y * BilinearWeights.x);
+
+	float Epsilon = .0001f;
+
+	float4 CornerDepths = abs(float4(TextureValues00.w, TextureValues10.w, TextureValues01.w, TextureValues11.w));
+	float SceneDepth = CalcSceneDepth(UVAndScreenPos.xy);
+	float4 DepthWeights = 1.0f / (abs(CornerDepths - SceneDepth.xxxx) + Epsilon);
+	float4 FinalWeights = CornerWeights * DepthWeights;
+
+	float3 InterpolatedResult = 
+		(FinalWeights.x * TextureValues00.xyz 
+		+ FinalWeights.y * TextureValues10.xyz
+		+ FinalWeights.z * TextureValues01.xyz 
+		+ FinalWeights.w * TextureValues11.xyz)
+		/ dot(FinalWeights, 1);
+
+	float3 BentNormal = InterpolatedResult.xyz;
+
+#else
+	float3 BentNormal = Texture2DSampleLevel(BentNormalAOTexture, BentNormalAOSampler, UVAndScreenPos.xy, 0).xyz;
+#endif
+
+	OutColor = float4(BentNormal, 1);
+}
+
+float BentNormalUVScale;
+float ApplyShadowing;
+
+void SkyLightDiffusePS(
+	in float4 UVAndScreenPos : TEXCOORD0, 
+	out float4 OutColor : SV_Target0)
+{
+	float2 UV = UVAndScreenPos.xy;
+
+	FScreenSpaceData ScreenSpaceData = GetScreenSpaceData(UV);
+	FGBufferData InGBufferData = ScreenSpaceData.GBuffer;
+	
+	float3 Lighting = 0; 
+
+	BRANCH
+	// Only light pixels marked as lit
+	if (InGBufferData.ShadingModelID > 0)
+	{
+		float SkyVisibility = 1;
+		float DotProductFactor = 1;
+		float3 SkyLightingNormal = InGBufferData.WorldNormal;
+
+		BRANCH
+		if (ApplyShadowing > 0)
+		{
+			float3 BentNormal = Texture2DSampleLevel(BentNormalAOTexture, BentNormalAOSampler, UV * BentNormalUVScale, 0).xyz;
+			SkyVisibility = BentNormal.x;
+
+			// Disabled for now as it needs tweaking
+			#define USE_DYNAMIC_OCCLUSION_BENT_NORMAL 0
+			#if USE_DYNAMIC_OCCLUSION_BENT_NORMAL
+				SkyVisibility = length(BentNormal);
+				float3 NormalizedBentNormal = BentNormal / (max(SkyVisibility, .00001f));
+
+				// Weight toward the material normal to increase directionality
+				float BentNormalWeightFactor = 1 - (1 - SkyVisibility) * (1 - SkyVisibility);
+
+				// We are lerping between the inputs of two lighting scenarios based on occlusion
+				// In the mostly unoccluded case, evaluate sky lighting with the material normal, because it has higher detail
+				// In the mostly occluded case, evaluate sky lighting with the bent normal, because it is a better representation of the incoming lighting
+				// Then treat the lighting evaluated along the bent normal as an area light, so we must apply the lambert term
+				SkyLightingNormal = lerp(NormalizedBentNormal, InGBufferData.WorldNormal, BentNormalWeightFactor);
+				DotProductFactor = lerp(dot(NormalizedBentNormal, InGBufferData.WorldNormal), 1, BentNormalWeightFactor);
+			#endif
+
+			if (InGBufferData.HasDistanceFieldRepresentation == 0)
+			{
+				SkyVisibility = .5f;
+			}
+		}
+			
+		// Compute the preconvolved incoming lighting with the bent normal direction
+		float3 DiffuseLookup = GetSkySHDiffuse(SkyLightingNormal) * View.SkyLightColor.rgb;
+
+		float ScalarFactors = SkyVisibility * DotProductFactor * InGBufferData.GBufferAO * ScreenSpaceData.AmbientOcclusion;
+
+		// Apply AO to the sky diffuse and account for darkening due to the geometry term
+		Lighting = ScalarFactors * DiffuseLookup * InGBufferData.DiffuseColor;
+
+		// For debugging
+		//Lighting = Texture2DSampleLevel(BentNormalAOTexture, BentNormalAOSampler, UV * BentNormalUVScale, 0).xyz;
+	}
+
+	OutColor = float4(Lighting, 0);
+}
\ No newline at end of file
diff --git a/Engine/Shaders/DistanceFieldSurfaceCacheLightingCompute.usf b/Engine/Shaders/DistanceFieldSurfaceCacheLightingCompute.usf
new file mode 100644
--- /dev/null
+++ b/Engine/Shaders/DistanceFieldSurfaceCacheLightingCompute.usf
@@ -0,0 +1,713 @@
+// Copyright 1998-2013 Epic Games, Inc. All Rights Reserved.
+
+/*=============================================================================
+
+=============================================================================*/
+
+#include "Common.usf"
+#include "DeferredShadingCommon.usf"
+#include "SHCommon.usf"
+#include "ReflectionEnvironmentShared.usf"
+#include "DistanceFieldLightingShared.usf"
+
+/** Min and Max depth for this tile. */
+groupshared uint IntegerTileMinZ;
+groupshared uint IntegerTileMaxZ;
+
+/** Inner Min and Max depth for this tile. */
+groupshared uint IntegerTileMinZ2;
+groupshared uint IntegerTileMaxZ2;
+
+/** View rect min in xy, max in zw. */
+uint4 ViewDimensions;
+float2 NumGroups;
+
+RWBuffer<float4> RWTileConeAxisAndCos;
+RWBuffer<float4> RWTileConeDepthRanges;
+RWBuffer<uint> RWTileHeadDataUnpacked;
+
+#ifndef MAX_OBJECTS_PER_TILE
+#define MAX_OBJECTS_PER_TILE 1
+#endif
+
+[numthreads(THREADGROUP_SIZEX, THREADGROUP_SIZEY, 1)]
+void BuildTileConesMain(
+	uint3 GroupId : SV_GroupID,
+	uint3 DispatchThreadId : SV_DispatchThreadID,
+    uint3 GroupThreadId : SV_GroupThreadID) 
+{
+    uint ThreadIndex = GroupThreadId.y * THREADGROUP_SIZEX + GroupThreadId.x;
+	float2 ScreenUV = (DispatchThreadId.xy * DOWNSAMPLE_FACTOR + float2(.5f, .5f)) * View.ViewSizeAndSceneTexelSize.zw;
+	float SceneDepth = CalcSceneDepth(ScreenUV);
+
+	// Initialize per-tile variables
+    if (ThreadIndex == 0) 
+	{
+        IntegerTileMinZ = 0x7F7FFFFF;     
+        IntegerTileMaxZ = 0;
+		IntegerTileMinZ2 = 0x7F7FFFFF;  
+		IntegerTileMaxZ2 = 0;
+    }
+
+    GroupMemoryBarrierWithGroupSync();
+    
+	// Use shared memory atomics to build the depth bounds for this tile
+	// Each thread is assigned to a pixel at this point
+    InterlockedMin(IntegerTileMinZ, asuint(SceneDepth));
+    InterlockedMax(IntegerTileMaxZ, asuint(SceneDepth));
+
+    GroupMemoryBarrierWithGroupSync();
+
+    float MinTileZ = asfloat(IntegerTileMinZ);
+    float MaxTileZ = asfloat(IntegerTileMaxZ);
+
+	float HalfZ = .5f * (MinTileZ + MaxTileZ);
+
+	// Compute a second min and max Z, clipped by HalfZ, so that we get two depth bounds per tile
+	// This results in more conservative tile depth bounds and fewer intersections
+	if (SceneDepth >= HalfZ)
+	{
+		InterlockedMin(IntegerTileMinZ2, asuint(SceneDepth));
+	}
+
+	if (SceneDepth <= HalfZ)
+	{
+		InterlockedMax(IntegerTileMaxZ2, asuint(SceneDepth));
+	}
+
+	GroupMemoryBarrierWithGroupSync();
+	
+	float MinTileZ2 = asfloat(IntegerTileMinZ2);
+	float MaxTileZ2 = asfloat(IntegerTileMaxZ2);
+
+	if (ThreadIndex == 0)
+	{
+		float3 TileConeVertex;
+		float3 TileConeAxis;
+		float TileConeAngleCos;
+		float TileConeAngleSin;
+		float4 ConeAxisDepthRanges;
+
+		{
+			float2 ViewSize = float2(1 / View.ViewToClip[0][0], 1 / View.ViewToClip[1][1]);
+			float3 TileCorner00 = normalize(float3((GroupId.x + 0) / NumGroups.x * ViewSize.x * 2 - ViewSize.x, ViewSize.y - (GroupId.y + 0) / NumGroups.y * ViewSize.y * 2, 1));
+			float3 TileCorner10 = normalize(float3((GroupId.x + 1) / NumGroups.x * ViewSize.x * 2 - ViewSize.x, ViewSize.y - (GroupId.y + 0) / NumGroups.y * ViewSize.y * 2, 1));
+			float3 TileCorner01 = normalize(float3((GroupId.x + 0) / NumGroups.x * ViewSize.x * 2 - ViewSize.x, ViewSize.y - (GroupId.y + 1) / NumGroups.y * ViewSize.y * 2, 1));
+			float3 TileCorner11 = normalize(float3((GroupId.x + 1) / NumGroups.x * ViewSize.x * 2 - ViewSize.x, ViewSize.y - (GroupId.y + 1) / NumGroups.y * ViewSize.y * 2, 1));
+
+			TileConeAxis = normalize(TileCorner00 + TileCorner10 + TileCorner01 + TileCorner11);
+			TileConeAngleCos = dot(TileConeAxis, TileCorner00);
+			TileConeAngleSin = sqrt(1 - TileConeAngleCos * TileConeAngleCos);
+			float TileConeAngleTan = TileConeAngleSin / TileConeAngleCos; 
+
+			float3 ViewSpaceSampleDirection = mul(float3(0, 0, 1), (float3x3)View.TranslatedWorldToView);
+			float ConeExpandDistance = 0;
+
+			float VertexPullbackLength = ConeExpandDistance / TileConeAngleTan;
+			float DistanceToNearPlane = length(TileConeAxis / TileConeAxis.z * View.NearPlane);
+			// 1 / cos(AngleBetweenTileCenterAndViewForward)
+			float InvCosTileAngle = 1.0f / TileConeAxis.z;
+			float ConeAxisDistanceMultiply = InvCosTileAngle;
+			float ConeAxisDistanceAdd = VertexPullbackLength + DistanceToNearPlane;
+			ConeAxisDepthRanges.x = ConeAxisDistanceMultiply * (MinTileZ - ConeExpandDistance) + ConeAxisDistanceAdd;
+			ConeAxisDepthRanges.y = ConeAxisDistanceMultiply * (MaxTileZ2 + ConeExpandDistance) + ConeAxisDistanceAdd;
+			ConeAxisDepthRanges.z = ConeAxisDistanceMultiply * (MinTileZ2 - ConeExpandDistance) + ConeAxisDistanceAdd;
+			ConeAxisDepthRanges.w = ConeAxisDistanceMultiply * (MaxTileZ + ConeExpandDistance) + ConeAxisDistanceAdd;
+
+			// Pull back cone vertex to contain potential samples
+			//@todo - only expand in sky direction
+			TileConeVertex = float3(0, 0, 0) - TileConeAxis * VertexPullbackLength;
+		}
+
+		uint TileIndex = GroupId.y * NumGroups.x + GroupId.x;
+		RWTileConeAxisAndCos[TileIndex] = float4(TileConeAxis, TileConeAngleCos);
+		RWTileConeDepthRanges[TileIndex] = ConeAxisDepthRanges;
+
+#if COHERENT_OBJECT_LIST_WRITES
+		RWTileHeadDataUnpacked[TileIndex * 4 + 0] = TileIndex;
+#else
+		RWTileHeadDataUnpacked[TileIndex * 4 + 0] = TileIndex * MAX_OBJECTS_PER_TILE;
+#endif
+		RWTileHeadDataUnpacked[TileIndex * 4 + 1] = 0;
+		RWTileHeadDataUnpacked[TileIndex * 4 + 2] = 0;
+		RWTileHeadDataUnpacked[TileIndex * 4 + 3] = 0;
+	}
+}
+
+groupshared uint SmallTileObjectIndices[MAX_OBJECTS_PER_TILE];
+groupshared uint MediumTileObjectIndices[MAX_OBJECTS_PER_TILE];
+groupshared uint LargeTileObjectIndices[MAX_OBJECTS_PER_TILE];
+
+groupshared uint SmallTileNumObjects;
+groupshared uint MediumTileNumObjects;
+groupshared uint LargeTileNumObjects;
+groupshared uint TileArrayDataStart;
+
+RWBuffer<uint> RWTileArrayData;
+RWBuffer<uint> RWTileArrayNextAllocation;
+RWBuffer<uint4> RWTileHeadData;
+
+[numthreads(THREADGROUP_SIZEX, THREADGROUP_SIZEY, 1)]
+void DistanceFieldAOBuildTileListMain(
+	uint3 GroupId : SV_GroupID,
+	uint3 DispatchThreadId : SV_DispatchThreadID,
+    uint3 GroupThreadId : SV_GroupThreadID) 
+{
+    uint ThreadIndex = GroupThreadId.y * THREADGROUP_SIZEX + GroupThreadId.x;
+	float2 ScreenUV = (DispatchThreadId.xy * DOWNSAMPLE_FACTOR + float2(.5f, .5f)) * View.ViewSizeAndSceneTexelSize.zw;
+	float SceneDepth = CalcSceneDepth(ScreenUV);
+
+	// Initialize per-tile variables
+    if (ThreadIndex == 0) 
+	{
+        IntegerTileMinZ = 0x7F7FFFFF;     
+        IntegerTileMaxZ = 0;
+		IntegerTileMinZ2 = 0x7F7FFFFF;  
+		IntegerTileMaxZ2 = 0;
+		SmallTileNumObjects = 0;
+		MediumTileNumObjects = 0;
+		LargeTileNumObjects = 0;
+    }
+
+    GroupMemoryBarrierWithGroupSync();
+    
+	// Use shared memory atomics to build the depth bounds for this tile
+	// Each thread is assigned to a pixel at this point
+    InterlockedMin(IntegerTileMinZ, asuint(SceneDepth));
+    InterlockedMax(IntegerTileMaxZ, asuint(SceneDepth));
+
+    GroupMemoryBarrierWithGroupSync();
+
+    float MinTileZ = asfloat(IntegerTileMinZ);
+    float MaxTileZ = asfloat(IntegerTileMaxZ);
+
+	float HalfZ = .5f * (MinTileZ + MaxTileZ);
+
+	// Compute a second min and max Z, clipped by HalfZ, so that we get two depth bounds per tile
+	// This results in more conservative tile depth bounds and fewer intersections
+	if (SceneDepth >= HalfZ)
+	{
+		InterlockedMin(IntegerTileMinZ2, asuint(SceneDepth));
+	}
+
+	if (SceneDepth <= HalfZ)
+	{
+		InterlockedMax(IntegerTileMaxZ2, asuint(SceneDepth));
+	}
+
+	GroupMemoryBarrierWithGroupSync();
+	
+	float MinTileZ2 = asfloat(IntegerTileMinZ2);
+	float MaxTileZ2 = asfloat(IntegerTileMaxZ2);
+
+	float3 TileConeVertex;
+	float3 TileConeAxis;
+	float TileConeAngleCos;
+	float TileConeAngleSin;
+	float4 ConeAxisDepthRanges;
+
+	{
+		float2 ViewSize = float2(1 / View.ViewToClip[0][0], 1 / View.ViewToClip[1][1]);
+		float3 TileCorner00 = normalize(float3((GroupId.x + 0) / NumGroups.x * ViewSize.x * 2 - ViewSize.x, ViewSize.y - (GroupId.y + 0) / NumGroups.y * ViewSize.y * 2, 1));
+		float3 TileCorner10 = normalize(float3((GroupId.x + 1) / NumGroups.x * ViewSize.x * 2 - ViewSize.x, ViewSize.y - (GroupId.y + 0) / NumGroups.y * ViewSize.y * 2, 1));
+		float3 TileCorner01 = normalize(float3((GroupId.x + 0) / NumGroups.x * ViewSize.x * 2 - ViewSize.x, ViewSize.y - (GroupId.y + 1) / NumGroups.y * ViewSize.y * 2, 1));
+		float3 TileCorner11 = normalize(float3((GroupId.x + 1) / NumGroups.x * ViewSize.x * 2 - ViewSize.x, ViewSize.y - (GroupId.y + 1) / NumGroups.y * ViewSize.y * 2, 1));
+
+		TileConeAxis = normalize(TileCorner00 + TileCorner10 + TileCorner01 + TileCorner11);
+		TileConeAngleCos = dot(TileConeAxis, TileCorner00);
+		TileConeAngleSin = sqrt(1 - TileConeAngleCos * TileConeAngleCos);
+		float TileConeAngleTan = TileConeAngleSin / TileConeAngleCos; 
+
+		float3 ViewSpaceSampleDirection = mul(float3(0, 0, 1), (float3x3)View.TranslatedWorldToView);
+		float ConeExpandDistance = 0;
+
+		float VertexPullbackLength = ConeExpandDistance / TileConeAngleTan;
+		float DistanceToNearPlane = length(TileConeAxis / TileConeAxis.z * View.NearPlane);
+		// 1 / cos(AngleBetweenTileCenterAndViewForward)
+		float InvCosTileAngle = 1.0f / TileConeAxis.z;
+		float ConeAxisDistanceMultiply = InvCosTileAngle;
+		float ConeAxisDistanceAdd = VertexPullbackLength + DistanceToNearPlane;
+		ConeAxisDepthRanges.x = ConeAxisDistanceMultiply * (MinTileZ - ConeExpandDistance) + ConeAxisDistanceAdd;
+		ConeAxisDepthRanges.y = ConeAxisDistanceMultiply * (MaxTileZ2 + ConeExpandDistance) + ConeAxisDistanceAdd;
+		ConeAxisDepthRanges.z = ConeAxisDistanceMultiply * (MinTileZ2 - ConeExpandDistance) + ConeAxisDistanceAdd;
+		ConeAxisDepthRanges.w = ConeAxisDistanceMultiply * (MaxTileZ + ConeExpandDistance) + ConeAxisDistanceAdd;
+
+		// Pull back cone vertex to contain potential samples
+		//@todo - only expand in sky direction
+		TileConeVertex = float3(0, 0, 0) - TileConeAxis * VertexPullbackLength;
+	}
+	
+	// A value of 1 is conservative, but has huge impact on performance
+	float RadiusScale = .5f;
+	float SmallGroupMaxSampleRadius;
+	{
+		uint StartIndex;
+		uint EndIndex;
+		float StepScale = GetPhaseParameters(0, StartIndex, EndIndex);
+		SmallGroupMaxSampleRadius = GetStepOffset(EndIndex, StepScale) * 2 * RadiusScale;
+	}
+
+	float MediumGroupMaxSampleRadius;
+	{
+		uint StartIndex;
+		uint EndIndex;
+		float StepScale = GetPhaseParameters(1, StartIndex, EndIndex);
+		MediumGroupMaxSampleRadius = GetStepOffset(EndIndex, StepScale) * 2 * RadiusScale;
+	}
+
+	float LargeGroupMaxSampleRadius;
+	{
+		uint StartIndex;
+		uint EndIndex;
+		float StepScale = GetPhaseParameters(2, StartIndex, EndIndex);
+		LargeGroupMaxSampleRadius = GetStepOffset(EndIndex, StepScale) * 2 * RadiusScale;
+	}
+
+	// Compute per-tile lists of affecting objects through bounds culling
+	// Each thread now operates on a sample instead of a pixel
+	LOOP
+	for (uint ObjectIndex = ThreadIndex; ObjectIndex < NumObjects; ObjectIndex += THREADGROUP_TOTALSIZE)
+	{
+		float4 SphereCenterAndRadius = LoadObjectPositionAndRadius(ObjectIndex);
+		float3 ViewSpaceSphereCenter = mul(float4(SphereCenterAndRadius.xyz + View.PreViewTranslation.xyz, 1), View.TranslatedWorldToView).xyz;
+		
+		if (SphereIntersectConeWithDepthRanges(float4(ViewSpaceSphereCenter, SphereCenterAndRadius.w + SmallGroupMaxSampleRadius), TileConeVertex, TileConeAxis, TileConeAngleCos, TileConeAngleSin, ConeAxisDepthRanges))
+		{
+			uint ListIndex;
+			InterlockedAdd(SmallTileNumObjects, 1U, ListIndex);
+			SmallTileObjectIndices[ListIndex] = ObjectIndex; 
+		}
+		else if (SphereIntersectConeWithDepthRanges(float4(ViewSpaceSphereCenter, SphereCenterAndRadius.w + MediumGroupMaxSampleRadius), TileConeVertex, TileConeAxis, TileConeAngleCos, TileConeAngleSin, ConeAxisDepthRanges))
+		{
+			uint ListIndex;
+			InterlockedAdd(MediumTileNumObjects, 1U, ListIndex);
+			MediumTileObjectIndices[ListIndex] = ObjectIndex; 
+		}
+		else if (SphereIntersectConeWithDepthRanges(float4(ViewSpaceSphereCenter, SphereCenterAndRadius.w + LargeGroupMaxSampleRadius), TileConeVertex, TileConeAxis, TileConeAngleCos, TileConeAngleSin, ConeAxisDepthRanges))
+		{
+			uint ListIndex;
+			InterlockedAdd(LargeTileNumObjects, 1U, ListIndex);
+			LargeTileObjectIndices[ListIndex] = ObjectIndex; 
+		}
+	}
+	
+	GroupMemoryBarrierWithGroupSync();
+
+	if (ThreadIndex == 0) 
+	{
+		SmallTileNumObjects = min(SmallTileNumObjects, MAX_OBJECTS_PER_TILE);
+		MediumTileNumObjects = min(MediumTileNumObjects, MAX_OBJECTS_PER_TILE);
+		LargeTileNumObjects = min(LargeTileNumObjects, MAX_OBJECTS_PER_TILE);
+
+		uint ArrayStart;
+		uint NumObjectsIntersecting = SmallTileNumObjects + MediumTileNumObjects + LargeTileNumObjects;
+		InterlockedAdd(RWTileArrayNextAllocation[0], NumObjectsIntersecting, ArrayStart);
+		TileArrayDataStart = ArrayStart;
+		RWTileHeadData[GroupId.y * NumGroups.x + GroupId.x] = uint4(TileArrayDataStart, SmallTileNumObjects, MediumTileNumObjects, LargeTileNumObjects);
+	}
+	
+	GroupMemoryBarrierWithGroupSync();
+
+	uint ArrayDataStart = TileArrayDataStart;
+	
+	LOOP
+	for (uint SmallListIndex = ThreadIndex; SmallListIndex < SmallTileNumObjects; SmallListIndex += THREADGROUP_TOTALSIZE)
+	{
+		RWTileArrayData[ArrayDataStart + SmallListIndex] = SmallTileObjectIndices[SmallListIndex];
+	}
+
+	ArrayDataStart += SmallTileNumObjects;
+
+	LOOP
+	for (uint MediumListIndex = ThreadIndex; MediumListIndex < MediumTileNumObjects; MediumListIndex += THREADGROUP_TOTALSIZE)
+	{
+		RWTileArrayData[ArrayDataStart + MediumListIndex] = MediumTileObjectIndices[MediumListIndex];
+	}
+
+	ArrayDataStart += MediumTileNumObjects;
+
+	LOOP
+	for (uint LargeListIndex = ThreadIndex; LargeListIndex < LargeTileNumObjects; LargeListIndex += THREADGROUP_TOTALSIZE)
+	{
+		RWTileArrayData[ArrayDataStart + LargeListIndex] = LargeTileObjectIndices[LargeListIndex];
+	}
+}
+
+/** View rect min in xy, max in zw. */
+float2 ThreadToCulledTile;
+float ConeHalfAngle;
+float BentNormalNormalizeFactor;
+float RecordRadiusScale;
+
+RWBuffer<float4> RWIrradianceCachePositionRadius;
+RWBuffer<float4> RWIrradianceCacheNormal;
+RWBuffer<uint> RWScatterDrawParameters;
+RWBuffer<uint2> RWIrradianceCacheTileCoordinate;
+
+Texture2D IrradianceCacheSplatTexture;
+SamplerState IrradianceCacheSplatSampler;
+
+void FindBestAxisVectors2(float3 InZAxis, out float3 OutXAxis, out float3 OutYAxis )
+{
+	float3 UpVector = abs(InZAxis.z) < 0.999 ? float3(0,0,1) : float3(1,0,0);
+	OutXAxis = normalize( cross( UpVector, InZAxis ) );
+	OutYAxis = cross( InZAxis, OutXAxis );
+}
+
+groupshared float4 CachedIrradianceCachePositionRadius[THREADGROUP_TOTALSIZE];
+groupshared float4 CachedIrradianceCacheNormal[THREADGROUP_TOTALSIZE];
+groupshared uint2 CachedIrradianceCacheTileCoordinate[THREADGROUP_TOTALSIZE];
+groupshared uint NumQueuedIrradianceCacheRecords;
+groupshared uint BaseRecordIndex;
+
+[numthreads(THREADGROUP_SIZEX, THREADGROUP_SIZEY, 1)]
+void PopulateCacheCS(
+	uint3 GroupId : SV_GroupID,
+	uint3 DispatchThreadId : SV_DispatchThreadID,
+    uint3 GroupThreadId : SV_GroupThreadID) 
+{
+	uint ThreadIndex = GroupThreadId.y * THREADGROUP_SIZEX + GroupThreadId.x;
+
+	if (ThreadIndex == 0)
+	{
+		NumQueuedIrradianceCacheRecords = 0;
+	}
+
+	GroupMemoryBarrierWithGroupSync();
+
+	float2 ScreenUV = (DispatchThreadId.xy * CurrentLevelDownsampleFactor + float2(.5f, .5f)) * View.ViewSizeAndSceneTexelSize.zw;
+	float2 ScreenPosition = (ScreenUV.xy - View.ScreenPositionScaleBias.wz) / View.ScreenPositionScaleBias.xy;
+	float2 DownsampledScreenUV = (DispatchThreadId.xy + float2(.5f, .5f)) / AOBufferSize;
+
+	FGBufferData GBufferData = GetGBufferData(ScreenUV);
+
+	float4 IrradianceCacheSplat = Texture2DSampleLevel(IrradianceCacheSplatTexture, IrradianceCacheSplatSampler, DownsampledScreenUV, 0);
+
+	BRANCH
+	if (GBufferData.ShadingModelID > 0
+		&& GBufferData.HasDistanceFieldRepresentation > 0
+		&& IrradianceCacheSplat.w < .0001f 
+		&& all(DispatchThreadId.xy < AOBufferSize))
+	{
+		float SceneDepth = CalcSceneDepth(ScreenUV);
+
+		BRANCH
+		if (SceneDepth < AOMaxViewDistance)
+		{
+			uint2 TileCoordinate = DispatchThreadId.xy * DownsampleFactorToBaseLevel / uint2(THREADGROUP_SIZEX, THREADGROUP_SIZEY);
+
+			float4 HomogeneousWorldPosition = mul(float4(ScreenPosition * SceneDepth, SceneDepth, 1), View.ScreenToWorld);
+			float3 OpaqueWorldPosition = HomogeneousWorldPosition.xyz / HomogeneousWorldPosition.w;
+
+			float2 BaseLevelScreenUV = (DispatchThreadId.xy * DownsampleFactorToBaseLevel + float2(.5f, .5f)) * BaseLevelTexelSize;
+			float3 WorldNormal = DecodeNormalForAO(Texture2DSampleLevel(DistanceFieldNormalTexture, DistanceFieldNormalSampler, BaseLevelScreenUV, 0).xyz);
+
+			//@todo - offset shading position along normal to avoid incorrect self-occlusion?
+			float3 WorldShadingPosition = OpaqueWorldPosition;
+
+			//if (all(DispatchThreadId.xy == uint2(2,1)))
+			{
+				int NextSampleIndex;
+				InterlockedAdd(NumQueuedIrradianceCacheRecords, 1, NextSampleIndex);
+				CachedIrradianceCachePositionRadius[NextSampleIndex] = float4(WorldShadingPosition, SceneDepth * CurrentLevelDownsampleFactor * .005f);
+				CachedIrradianceCacheNormal[NextSampleIndex] = float4(WorldNormal, .01f * CurrentLevelDownsampleFactor / 2);
+				CachedIrradianceCacheTileCoordinate[NextSampleIndex] = TileCoordinate;
+			}
+		}
+	}
+
+	GroupMemoryBarrierWithGroupSync();
+	
+	if (ThreadIndex == 0)
+	{
+		InterlockedAdd(RWScatterDrawParameters[1], NumQueuedIrradianceCacheRecords, BaseRecordIndex);
+	}
+
+	GroupMemoryBarrierWithGroupSync();
+
+	LOOP
+	for (uint LocalRecordIndex = ThreadIndex; LocalRecordIndex < NumQueuedIrradianceCacheRecords; LocalRecordIndex += THREADGROUP_TOTALSIZE)
+	{
+		int SampleIndex = BaseRecordIndex + LocalRecordIndex;
+		RWIrradianceCachePositionRadius[SampleIndex] = CachedIrradianceCachePositionRadius[LocalRecordIndex];
+		RWIrradianceCacheNormal[SampleIndex] = float4(EncodeNormalForAO(CachedIrradianceCacheNormal[LocalRecordIndex].xyz), 0);
+		uint2 TileCoordinate = CachedIrradianceCacheTileCoordinate[LocalRecordIndex];
+		RWIrradianceCacheTileCoordinate[SampleIndex] = TileCoordinate;
+	}
+
+	if (all(DispatchThreadId == 0))
+	{
+		// VertexCountPerInstance
+		RWScatterDrawParameters[0] = 8 * 3;
+		// StartVertexLocation
+		RWScatterDrawParameters[2] = 0;
+		// StartInstanceLocation
+		RWScatterDrawParameters[3] = 0;
+	}
+}
+
+Buffer<float4> IrradianceCachePositionRadius;
+Buffer<uint2> IrradianceCacheTileCoordinate;
+Buffer<float4> IrradianceCacheNormal;
+Buffer<uint> ScatterDrawParameters;
+Buffer<uint> SavedStartIndex;
+
+RWBuffer<float> RWOccluderRadius;
+RWBuffer<float4> RWIrradianceCacheBentNormal;
+
+#define FINAL_GATHER_THREADGROUP_SIZE 64
+#define SIMULTANEOUSLY_TRACED_OBJECTS 4
+#define THREADS_PER_OBJECT (FINAL_GATHER_THREADGROUP_SIZE / SIMULTANEOUSLY_TRACED_OBJECTS)
+
+groupshared uint SharedConeVisibility[NUM_SAMPLES];
+groupshared uint SharedMinOcclusionDistance;
+
+[numthreads(FINAL_GATHER_THREADGROUP_SIZE, 1, 1)]
+void FinalGatherCS(
+	uint3 GroupId : SV_GroupID,
+	uint3 DispatchThreadId : SV_DispatchThreadID,
+    uint3 GroupThreadId : SV_GroupThreadID) 
+{
+	uint ThreadIndex = GroupThreadId.x;
+	uint ObjectOffsetIndex = ThreadIndex / THREADS_PER_OBJECT;
+
+	float ShadingConeAngleScale = 1;
+	float ShadingConeTanTheta = tan(ShadingConeAngleScale * ConeHalfAngle);
+
+	if (ThreadIndex == 0)
+	{
+		for (uint ConeIndex = 0; ConeIndex < NUM_SAMPLES; ConeIndex++) 
+		{
+			SharedConeVisibility[ConeIndex] = asuint(1.0f);
+		}
+
+		SharedMinOcclusionDistance = asuint(AOMaxDistance);
+	}
+
+	GroupMemoryBarrierWithGroupSync();
+
+	uint StartIndex = SavedStartIndex[0];
+	uint NumRecords = ScatterDrawParameters[1];
+
+	uint RecordIndex = StartIndex + GroupId.x;
+
+	float3 TangentX;
+	float3 TangentY;
+	float3 WorldNormal;
+
+	{
+		WorldNormal = DecodeNormalForAO(IrradianceCacheNormal[RecordIndex].xyz);
+		float3 WorldShadingPosition = IrradianceCachePositionRadius[RecordIndex].xyz;
+		uint2 TileCoordinate = IrradianceCacheTileCoordinate[RecordIndex];
+
+		uint4 TileHead = GetTileHead(TileCoordinate);
+		uint NumObjectsAffectingTile = TileHead.y + TileHead.z + TileHead.w;
+
+		FindBestAxisVectors2(WorldNormal, TangentX, TangentY);
+
+		LOOP 
+		for (uint ListObjectIndex = 0; ListObjectIndex < NumObjectsAffectingTile; ListObjectIndex += SIMULTANEOUSLY_TRACED_OBJECTS)
+		{
+			uint EffectiveListObjectIndex = ListObjectIndex + ObjectOffsetIndex;
+
+			if (EffectiveListObjectIndex < NumObjectsAffectingTile)
+			{
+#if COHERENT_OBJECT_LIST_WRITES
+
+				uint ListIndex = 0;
+				uint ArrayIndex = EffectiveListObjectIndex;
+
+				FLATTEN
+				if (EffectiveListObjectIndex >= TileHead.y + TileHead.z)
+				{
+					ListIndex = 2;
+					ArrayIndex = EffectiveListObjectIndex - TileHead.y - TileHead.z;
+				}
+				else if (EffectiveListObjectIndex >= TileHead.y)
+				{
+					ListIndex = 1;
+					ArrayIndex = EffectiveListObjectIndex - TileHead.y;
+				}
+
+				uint ObjectIndex = TileArrayData.Load((ArrayIndex * TileListGroupSize.x * TileListGroupSize.y + TileHead.x) * NUM_CULLED_OBJECT_LISTS + ListIndex);
+#else
+				uint ObjectIndex = TileArrayData.Load(TileHead.x + EffectiveListObjectIndex);
+#endif
+
+				float3 LocalPositionExtent = LoadObjectLocalPositionExtent(ObjectIndex);
+				float4x4 WorldToVolume = LoadObjectWorldToVolume(ObjectIndex);
+				float4 UVScaleAndVolumeScale = LoadObjectUVScale(ObjectIndex);
+				float3 VolumeShadingPosition = mul(float4(WorldShadingPosition, 1), WorldToVolume).xyz;
+
+				float ObjectOccluderRadius = length(LocalPositionExtent) * .5f;
+				float BoxDistance = ComputeDistanceFromBoxToPoint(-LocalPositionExtent, LocalPositionExtent, VolumeShadingPosition) * UVScaleAndVolumeScale.w;
+
+				BRANCH
+				if (BoxDistance < AOMaxDistance)
+				{
+					float3 UVAdd = LoadObjectUVAdd(ObjectIndex);
+
+					uint StartStepIndex = 0;
+			
+					FLATTEN
+					if (EffectiveListObjectIndex >= TileHead.y + TileHead.z)
+					{
+						StartStepIndex = 8;
+					}
+					else if (EffectiveListObjectIndex >= TileHead.y)
+					{
+						StartStepIndex = 5;
+					}
+
+					uint NumConeSteps = NUM_CONE_STEPS - StartStepIndex;
+
+					LOOP
+					for (uint SampleIndex = ThreadIndex % THREADS_PER_OBJECT; SampleIndex < NUM_SAMPLES * NumConeSteps; SampleIndex += THREADS_PER_OBJECT)
+					{
+						// Avoid uint divide
+						uint ConeIndex = (uint)(SampleIndex / (float)NumConeSteps);
+						uint StepIndex = StartStepIndex + SampleIndex - ConeIndex * NumConeSteps;
+
+						float3 ConeDirection = AOSamples2.SampleDirections[ConeIndex].xyz;
+						float3 RotatedConeDirection = ConeDirection.x * TangentX + ConeDirection.y * TangentY + ConeDirection.z * WorldNormal;
+
+						float WorldStepOffset = AOStepScale * exp2(AOStepExponentScale * StepIndex);
+						float3 WorldSamplePosition = WorldShadingPosition + RotatedConeDirection * WorldStepOffset;
+						float3 StepSamplePosition = mul(float4(WorldSamplePosition, 1), WorldToVolume).xyz;
+						float3 ClampedSamplePosition = clamp(StepSamplePosition, -LocalPositionExtent, LocalPositionExtent);
+						float DistanceToClamped = length(StepSamplePosition - ClampedSamplePosition);
+
+						float LocalStepOffset = length(StepSamplePosition - VolumeShadingPosition);
+
+						float3 StepVolumeUV = DistanceFieldVolumePositionToUV(ClampedSamplePosition, UVScaleAndVolumeScale.xyz, UVAdd);
+						float DistanceToOccluder = Texture3DSampleLevel(DistanceFieldTexture, DistanceFieldSampler, StepVolumeUV, 0).x + DistanceToClamped;
+
+						float EffectiveDistanceToOccluder = DistanceToOccluder;
+						float SphereRadius = LocalStepOffset * ShadingConeTanTheta;
+						// Allows use of negative distances in the small steps to create full penumbra range, reduces over-occlusion
+						float MinOcclusionFraction = 0;//-.6f * (1 - saturate(4 * StepIndex / (float)NUM_CONE_STEPS));
+						float Visibility = saturate((EffectiveDistanceToOccluder / SphereRadius - MinOcclusionFraction) / (1 - MinOcclusionFraction));
+
+						// Don't allow small objects to fully occlude a cone step
+						Visibility = max(Visibility, 1 - saturate(ObjectOccluderRadius / SphereRadius));
+
+						InterlockedMin(SharedConeVisibility[ConeIndex], asuint(Visibility));
+
+						if (EffectiveDistanceToOccluder < .9f * SphereRadius)
+						{
+							// Assuming occluder is straight forward along the cone
+							float WorldDistanceToOccluder = (LocalStepOffset + DistanceToOccluder) * UVScaleAndVolumeScale.w;
+							InterlockedMin(SharedMinOcclusionDistance, asuint(WorldDistanceToOccluder));
+						}
+					}
+				}
+			}
+		}
+	}
+
+	GroupMemoryBarrierWithGroupSync();
+
+	if (ThreadIndex == 0)
+	{
+		float3 UnoccludedDirection = 0;
+		float Visibility = 0;
+
+		for (uint ConeIndex = 0; ConeIndex < NUM_SAMPLES; ConeIndex++)
+		{
+			float3 ConeDirection = AOSamples2.SampleDirections[ConeIndex].xyz;
+			float3 RotatedConeDirection = ConeDirection.x * TangentX + ConeDirection.y * TangentY + ConeDirection.z * WorldNormal;
+
+			float ConeVisibility = asfloat(SharedConeVisibility[ConeIndex]);
+			UnoccludedDirection += ConeVisibility * RotatedConeDirection;
+			Visibility += ConeVisibility;
+		}
+
+		UnoccludedDirection = UnoccludedDirection * BentNormalNormalizeFactor / (float)NUM_SAMPLES;
+		Visibility = Visibility / (float)NUM_SAMPLES;
+
+		RWOccluderRadius[RecordIndex] = RecordRadiusScale * asfloat(SharedMinOcclusionDistance);
+		RWIrradianceCacheBentNormal[RecordIndex] = float4(EncodeNormalForAO(Visibility.xxx), 0);
+	}
+}
+
+#define COMPACT_THREADGROUP_SIZEX 64
+
+Buffer<uint> DrawParameters;
+RWBuffer<uint> RWDispatchParameters;
+
+RWBuffer<uint> RWSavedStartIndex;
+
+[numthreads(1, 1, 1)]
+void SaveStartIndexCS(
+	uint3 GroupId : SV_GroupID,
+	uint3 DispatchThreadId : SV_DispatchThreadID,
+    uint3 GroupThreadId : SV_GroupThreadID) 
+{
+	RWSavedStartIndex[0] = DrawParameters[1];
+}
+
+[numthreads(1, 1, 1)]
+void SetupFinalGatherIndirectArgumentsCS(
+	uint3 GroupId : SV_GroupID,
+	uint3 DispatchThreadId : SV_DispatchThreadID,
+    uint3 GroupThreadId : SV_GroupThreadID) 
+{
+	uint StartIndex = SavedStartIndex[0];
+	uint NumRecords = DrawParameters[1];
+	uint WorkRange = NumRecords - StartIndex;
+
+	RWDispatchParameters[0] = WorkRange;
+	RWDispatchParameters[1] = 1;
+	RWDispatchParameters[2] = 1;
+}
+
+#define COPY_THREADGROUP_SIZE 256
+float TrimFraction;
+
+[numthreads(1, 1, 1)]
+void SetupCopyIndirectArgumentsCS(
+	uint3 GroupId : SV_GroupID,
+	uint3 DispatchThreadId : SV_DispatchThreadID,
+    uint3 GroupThreadId : SV_GroupThreadID) 
+{
+	uint NumRecords = (1 - TrimFraction) * DrawParameters[1];
+
+	RWDispatchParameters[0] = (NumRecords + COPY_THREADGROUP_SIZE - 1) / COPY_THREADGROUP_SIZE;
+	RWDispatchParameters[1] = 1;
+	RWDispatchParameters[2] = 1;
+}
+
+RWBuffer<float4> RWCopyIrradianceCachePositionRadius;
+RWBuffer<float4> RWCopyIrradianceCacheNormal;
+RWBuffer<float> RWCopyOccluderRadius;
+RWBuffer<float4> RWCopyIrradianceCacheBentNormal;
+RWBuffer<uint2> RWCopyIrradianceCacheTileCoordinate;
+
+Buffer<float> OccluderRadius;
+Buffer<float4> IrradianceCacheBentNormal;
+
+[numthreads(COPY_THREADGROUP_SIZE, 1, 1)]
+void CopyIrradianceCacheSamplesCS(
+	uint3 GroupId : SV_GroupID,
+	uint3 DispatchThreadId : SV_DispatchThreadID,
+    uint3 GroupThreadId : SV_GroupThreadID) 
+{
+	uint NumRecords = DrawParameters[1];
+	uint StartIndex = NumRecords * TrimFraction;
+	uint SourceIndex = StartIndex + DispatchThreadId.x;
+	uint DestIndex = DispatchThreadId.x;
+
+	if (SourceIndex < NumRecords)
+	{
+		RWCopyIrradianceCachePositionRadius[DestIndex] = IrradianceCachePositionRadius[SourceIndex];
+		RWCopyIrradianceCacheNormal[DestIndex] = IrradianceCacheNormal[SourceIndex];
+		RWCopyOccluderRadius[DestIndex] = OccluderRadius[SourceIndex];
+		RWCopyIrradianceCacheBentNormal[DestIndex] = IrradianceCacheBentNormal[SourceIndex];
+		RWCopyIrradianceCacheTileCoordinate[DestIndex] = IrradianceCacheTileCoordinate[SourceIndex];
+	}
+
+	if (DispatchThreadId.x == 0)
+	{
+		RWScatterDrawParameters[1] = NumRecords - StartIndex;
+	}
+}
\ No newline at end of file
diff --git a/Engine/Shaders/LPVWriteVplCommon.usf b/Engine/Shaders/LPVWriteVplCommon.usf
--- a/Engine/Shaders/LPVWriteVplCommon.usf
+++ b/Engine/Shaders/LPVWriteVplCommon.usf
@@ -16,11 +16,11 @@
 
 //-------------------------------------------------------------------------------------------------
 
-RWStructuredBuffer<VplListEntry>				gGvListBufferRW;		// 0
-RWByteAddressBuffer								gGvListHeadBufferRW;	// 1
+RWStructuredBuffer<VplListEntry>				RWGvListBuffer;		// 0
+RWByteAddressBuffer								RWGvListHeadBuffer;	// 1
 
-RWStructuredBuffer<VplListEntry>				gVplListBufferRW;		// 2
-RWByteAddressBuffer								gVplListHeadBufferRW;	// 3
+RWStructuredBuffer<VplListEntry>				RWVplListBuffer;		// 2
+RWByteAddressBuffer								RWVplListHeadBuffer;	// 3
 
 //-------------------------------------------------------------------------------------------------
 // Functions
@@ -32,9 +32,9 @@ int AddToVplList( float3 worldPos, float3 flux, float3 direction, const bool emi
 	if ( gridIndex >= 0 )
 	{
 		// Allocate a new list entry 
-		uint newHeadIndex = gVplListBufferRW.IncrementCounter();
+		uint newHeadIndex = RWVplListBuffer.IncrementCounter();
 		uint oldHeadIndex;
-		gVplListHeadBufferRW.InterlockedExchange( gridIndex*4, newHeadIndex, oldHeadIndex );
+		RWVplListHeadBuffer.InterlockedExchange( gridIndex*4, newHeadIndex, oldHeadIndex );
 
 		if ( emissiveInjection )
 		{
@@ -51,7 +51,7 @@ int AddToVplList( float3 worldPos, float3 flux, float3 direction, const bool emi
 
 		VplListEntry listNode = PackVpl( vpl );
 		listNode.nextIndex = oldHeadIndex;
-		gVplListBufferRW[ newHeadIndex ] = listNode;
+		RWVplListBuffer[ newHeadIndex ] = listNode;
 	}
 	return gridIndex;
 }
@@ -67,9 +67,9 @@ int AddToGvList( float3 worldPos, float3 flux, float3 direction )
 	if ( gridIndex >= 0 )
 	{
 		// Allocate a new list entry 
-		uint newHeadIndex = gGvListBufferRW.IncrementCounter();
+		uint newHeadIndex = RWGvListBuffer.IncrementCounter();
 		uint oldHeadIndex;
-		gGvListHeadBufferRW.InterlockedExchange( gridIndex*4, newHeadIndex, oldHeadIndex );
+		RWGvListHeadBuffer.InterlockedExchange( gridIndex*4, newHeadIndex, oldHeadIndex );
 
 		// Fill in the new list node 
 		Vpl vpl;
@@ -78,7 +78,7 @@ int AddToGvList( float3 worldPos, float3 flux, float3 direction )
 
 		VplListEntry listNode = PackVpl( vpl );
 		listNode.nextIndex = oldHeadIndex;
-		gGvListBufferRW[ newHeadIndex ] = listNode;
+		RWGvListBuffer[ newHeadIndex ] = listNode;
 	}
 	return gridIndex;
 }
diff --git a/Engine/Shaders/ReflectionEnvironmentShared.usf b/Engine/Shaders/ReflectionEnvironmentShared.usf
--- a/Engine/Shaders/ReflectionEnvironmentShared.usf
+++ b/Engine/Shaders/ReflectionEnvironmentShared.usf
@@ -45,3 +45,26 @@ float3 GetSkyLightReflection(float3 ReflectionVector, float Roughness, bool Norm
 
 	return Reflection * View.SkyLightColor.rgb;
 }
+/** 
+ * Computes sky diffuse lighting from the SH irradiance map.  
+ * This has the SH basis evaluation and diffuse convolusion weights combined for minimal ALU's - see "Stupid Spherical Harmonics (SH) Tricks" 
+ */
+float3 GetSkySHDiffuse(float3 Normal)
+{
+	float4 NormalVector = float4(Normal, 1);
+
+	float3 Intermediate0, Intermediate1, Intermediate2;
+	Intermediate0.x = dot(View.SkyIrradianceEnvironmentMap[0], NormalVector);
+	Intermediate0.y = dot(View.SkyIrradianceEnvironmentMap[1], NormalVector);
+	Intermediate0.z = dot(View.SkyIrradianceEnvironmentMap[2], NormalVector);
+
+	float4 vB = NormalVector.xyzz * NormalVector.yzzx;
+	Intermediate1.x = dot(View.SkyIrradianceEnvironmentMap[3], vB);
+	Intermediate1.y = dot(View.SkyIrradianceEnvironmentMap[4], vB);
+	Intermediate1.z = dot(View.SkyIrradianceEnvironmentMap[5], vB);
+
+	float vC = NormalVector.x * NormalVector.x - NormalVector.y * NormalVector.y;
+	Intermediate2 = View.SkyIrradianceEnvironmentMap[6].xyz * vC;
+
+	return Intermediate0 + Intermediate1 + Intermediate2;
+}
diff --git a/Engine/Source/Developer/MeshUtilities/Private/MeshUtilities.cpp b/Engine/Source/Developer/MeshUtilities/Private/MeshUtilities.cpp
--- a/Engine/Source/Developer/MeshUtilities/Private/MeshUtilities.cpp
+++ b/Engine/Source/Developer/MeshUtilities/Private/MeshUtilities.cpp
@@ -15,6 +15,11 @@
 #include "MaterialExportUtils.h"
 #include "Textures/TextureAtlas.h"
 
+//@todo - implement required vector intrinsics for other implementations
+#if PLATFORM_ENABLE_VECTORINTRINSICS
+#include "kDOP.h"
+#endif
+
 /*------------------------------------------------------------------------------
 	MeshUtilities module.
 ------------------------------------------------------------------------------*/
@@ -23,10 +28,422 @@
 // causes meshes to be rebuilt you MUST generate a new GUID and replace this
 // string with it.
 
-#define MESH_UTILITIES_VER TEXT("359a029847e84730ba516769d0f19427")
+#define MESH_UTILITIES_VER TEXT("359a039847e84730ba516769d0f19427")
 
 DEFINE_LOG_CATEGORY_STATIC(LogMeshUtilities,Verbose,All);
 
+//@todo - implement required vector intrinsics for other implementations
+#if PLATFORM_ENABLE_VECTORINTRINSICS
+
+class FMeshBuildDataProvider
+{
+public:
+
+	/** Initialization constructor. */
+	FMeshBuildDataProvider(
+		const TkDOPTree<const FMeshBuildDataProvider,uint32>& InkDopTree):
+		kDopTree(InkDopTree)
+	{}
+
+	// kDOP data provider interface.
+
+	FORCEINLINE const TkDOPTree<const FMeshBuildDataProvider,uint32>& GetkDOPTree(void) const
+	{
+		return kDopTree;
+	}
+
+	FORCEINLINE const FMatrix& GetLocalToWorld(void) const
+	{
+		return FMatrix::Identity;
+	}
+
+	FORCEINLINE const FMatrix& GetWorldToLocal(void) const
+	{
+		return FMatrix::Identity;
+	}
+
+	FORCEINLINE FMatrix GetLocalToWorldTransposeAdjoint(void) const
+	{
+		return FMatrix::Identity;
+	}
+
+	FORCEINLINE float GetDeterminant(void) const
+	{
+		return 1.0f;
+	}
+
+private:
+
+	const TkDOPTree<const FMeshBuildDataProvider,uint32>& kDopTree;
+};
+
+/** Generates unit length, stratified and uniformly distributed direction samples in a hemisphere. */
+void GenerateStratifiedUniformHemisphereSamples(int32 NumThetaSteps, int32 NumPhiSteps, FRandomStream& RandomStream, TArray<FVector4>& Samples)
+{
+	Samples.Empty(NumThetaSteps * NumPhiSteps);
+	for (int32 ThetaIndex = 0; ThetaIndex < NumThetaSteps; ThetaIndex++)
+	{
+		for (int32 PhiIndex = 0; PhiIndex < NumPhiSteps; PhiIndex++)
+		{
+			const float U1 = RandomStream.GetFraction();
+			const float U2 = RandomStream.GetFraction();
+
+			const float Fraction1 = (ThetaIndex + U1) / (float)NumThetaSteps;
+			const float Fraction2 = (PhiIndex + U2) / (float)NumPhiSteps;
+
+			const float R = FMath::Sqrt(1.0f - Fraction1 * Fraction1);
+
+			const float Phi = 2.0f * (float)PI * Fraction2;
+			// Convert to Cartesian
+			Samples.Add(FVector4(FMath::Cos(Phi) * R, FMath::Sin(Phi) * R, Fraction1));
+		}
+	}
+}
+
+class FMeshDistanceFieldThreadRunnable : public FRunnable
+{
+private:
+	FRunnableThread* Thread;
+
+	// Readonly inputs
+	TkDOPTree<const FMeshBuildDataProvider,uint32>& kDopTree;
+	const TArray<FVector4>& SampleDirections;
+	FBox VolumeBounds;
+	FIntVector VolumeDimensions;
+	float VolumeMaxDistance;
+
+	volatile int32* SharedZIndex;
+	volatile int32* NegativeAtBorder;
+
+	// Output
+	TArray<FFloat16>& OutDistanceFieldVolume;
+
+	static int32 NextThreadIndex;
+
+public:
+	/** Initialization constructor. */
+	FMeshDistanceFieldThreadRunnable(
+		TkDOPTree<const FMeshBuildDataProvider,uint32>& InkDopTree, 
+		const TArray<FVector4>& InSampleDirections,
+		FBox InVolumeBounds,
+		FIntVector InVolumeDimensions,
+		float InVolumeMaxDistance,
+		volatile int32* InSharedZIndex,
+		volatile int32* InNegativeAtBorder,
+		TArray<FFloat16>& DistanceFieldVolume)
+		:
+		kDopTree(InkDopTree),
+		SampleDirections(InSampleDirections),
+		VolumeBounds(InVolumeBounds),
+		VolumeDimensions(InVolumeDimensions),
+		VolumeMaxDistance(InVolumeMaxDistance),
+		SharedZIndex(InSharedZIndex),
+		NegativeAtBorder(InNegativeAtBorder),
+		OutDistanceFieldVolume(DistanceFieldVolume)
+	{
+		Thread = FRunnableThread::Create(this, *FString::Printf(TEXT("MeshDistanceFieldThread%u"), NextThreadIndex), 0, TPri_Normal);
+		NextThreadIndex++;
+	}
+
+	// FRunnable interface.
+	virtual bool Init(void) { return true; }
+	virtual void Exit(void) {}
+	virtual void Stop(void) {}
+	virtual uint32 Run(void);
+
+	void WaitForCompletion()
+	{
+		Thread->WaitForCompletion();
+	}
+};
+
+int32 FMeshDistanceFieldThreadRunnable::NextThreadIndex = 0;
+
+uint32 FMeshDistanceFieldThreadRunnable::Run()
+{
+	int32 ZIndex = FPlatformAtomics::InterlockedAdd(SharedZIndex, 1);;
+	FMeshBuildDataProvider kDOPDataProvider(kDopTree);
+
+	while (ZIndex < VolumeDimensions.Z)
+	{
+		const FVector4 DistanceFieldVoxelSize(VolumeBounds.GetSize() / FVector(VolumeDimensions.X, VolumeDimensions.Y, VolumeDimensions.Z));
+
+		for (int32 YIndex = 0; YIndex < VolumeDimensions.Y; YIndex++)
+		{
+			for (int32 XIndex = 0; XIndex < VolumeDimensions.X; XIndex++)
+			{
+				const FVector4 VoxelPosition = FVector4(XIndex + .5f, YIndex + .5f, ZIndex + .5f) * DistanceFieldVoxelSize + VolumeBounds.Min;
+				const int32 Index = (ZIndex * VolumeDimensions.Y * VolumeDimensions.X + YIndex * VolumeDimensions.X + XIndex);
+
+				float MinDistance = VolumeMaxDistance;
+				int32 Hit = 0;
+				int32 HitFront = 0;
+
+				for (int32 SampleIndex = 0; SampleIndex < SampleDirections.Num(); SampleIndex++)
+				{
+					const FVector RayDirection = SampleDirections[SampleIndex];
+					const FVector ReflectedRayDirection = RayDirection;
+
+					//@todo - worth it anymore?
+					if (FMath::LineBoxIntersection(VolumeBounds, VoxelPosition, VoxelPosition + ReflectedRayDirection * VolumeMaxDistance, ReflectedRayDirection, FVector(1.0f) / ReflectedRayDirection))
+					{
+						FkHitResult Result;
+
+						TkDOPLineCollisionCheck<const FMeshBuildDataProvider,uint32> kDOPCheck(
+							VoxelPosition,
+							VoxelPosition + ReflectedRayDirection * VolumeMaxDistance,
+							true,
+							kDOPDataProvider,
+							&Result);
+
+						bool bHit = kDopTree.LineCheck(kDOPCheck);
+
+						if (bHit)
+						{
+							Hit++;
+
+							const FVector HitNormal = kDOPCheck.GetHitNormal();
+
+							if (FVector::DotProduct(ReflectedRayDirection, HitNormal) < 0 
+								// MaterialIndex on the build triangles was set to 1 if two-sided, or 0 if one-sided
+								|| kDOPCheck.Result->Item != 0)
+							{
+								HitFront++;
+							}
+
+							const float CurrentDistance = VolumeMaxDistance * Result.Time;
+
+							if (CurrentDistance < MinDistance)
+							{
+								MinDistance = CurrentDistance;
+							}
+						}
+					}
+				}
+
+				// Consider this voxel 'outside' an object if more than 50% of the rays hit front faces
+				MinDistance *= (Hit == 0 || HitFront > Hit * .5f) ? 1 : -1;
+
+				const float VolumeSpaceDistance = MinDistance / VolumeBounds.GetExtent().GetMax();
+
+				if (MinDistance < 0 && 
+					(XIndex == 0 || XIndex == VolumeDimensions.X - 1 || 
+						YIndex == 0 || YIndex == VolumeDimensions.Y - 1 || 
+						ZIndex == 0 || ZIndex == VolumeDimensions.Z - 1))
+				{
+					*NegativeAtBorder = 1;
+				}
+
+				OutDistanceFieldVolume[Index] = FFloat16(VolumeSpaceDistance);
+			}
+		}
+
+		ZIndex = FPlatformAtomics::InterlockedAdd(SharedZIndex, 1);
+	}
+
+	return 0;
+}
+
+static void GenerateSignedDistanceFieldVolumeData(
+	FStaticMeshLODResources& LODModel,
+	const TArray<UMaterialInterface*>& Materials,
+	const FBoxSphereBounds& Bounds,
+	float DistanceFieldResolutionScale)
+{
+	static const auto CVar = IConsoleManager::Get().FindTConsoleVariableDataInt(TEXT("r.AllowMeshDistanceFieldRepresentations"));
+
+	if (DistanceFieldResolutionScale > 0 
+		&& CVar 
+		&& CVar->GetValueOnGameThread() != 0)
+	{
+		const FPositionVertexBuffer& PositionVertexBuffer = LODModel.PositionVertexBuffer;
+		FIndexArrayView Indices = LODModel.IndexBuffer.GetArrayView();
+		FDistanceFieldVolumeData& OutData = LODModel.DistanceFieldData;
+
+		TArray<bool> CachedTwoSided;
+		TArray<bool> CachedOpaqueOrMasked;
+
+		CachedTwoSided.Empty(Materials.Num());
+		CachedTwoSided.AddZeroed(Materials.Num());
+
+		CachedOpaqueOrMasked.Empty(Materials.Num());
+		CachedOpaqueOrMasked.AddZeroed(Materials.Num());
+
+		for (int32 MaterialIndex = 0; MaterialIndex < Materials.Num(); MaterialIndex++)
+		{
+			UMaterialInterface* Material = Materials[MaterialIndex];
+
+			if (Material)
+			{
+				CachedTwoSided[MaterialIndex] = Material->IsTwoSided();
+				CachedOpaqueOrMasked[MaterialIndex] = !IsTranslucentBlendMode(Material->GetBlendMode());
+			}
+			else
+			{
+				// Default material properties
+				CachedOpaqueOrMasked[MaterialIndex] = true;
+			}
+		}
+
+		TArray<FkDOPBuildCollisionTriangle<uint32> > BuildTriangles;
+
+		for (int32 i = 0; i < Indices.Num(); i += 3)
+		{
+			const FVector V0 = PositionVertexBuffer.VertexPosition(Indices[i + 0]);
+			const FVector V1 = PositionVertexBuffer.VertexPosition(Indices[i + 1]);
+			const FVector V2 = PositionVertexBuffer.VertexPosition(Indices[i + 2]);
+
+			const FVector LocalNormal = ((V1 - V2) ^ (V0 - V2)).SafeNormal();
+
+			// No degenerates
+			if (LocalNormal.IsUnit())
+			{
+				uint32 StoreAsTwoSided = 0;
+				bool bTriangleIsOpaqueOrMasked = false;
+
+				for (int32 SectionIndex = 0; SectionIndex < LODModel.Sections.Num(); SectionIndex++)
+				{
+					const FStaticMeshSection& Section = LODModel.Sections[SectionIndex];
+
+					if ((uint32)i >= Section.FirstIndex && (uint32)i < Section.FirstIndex + Section.NumTriangles * 3)
+					{
+						if (CachedTwoSided.IsValidIndex(Section.MaterialIndex))
+						{
+							StoreAsTwoSided = CachedTwoSided[Section.MaterialIndex] ? 1 : 0;
+						}
+
+						if (CachedOpaqueOrMasked.IsValidIndex(Section.MaterialIndex))
+						{
+							bTriangleIsOpaqueOrMasked = CachedOpaqueOrMasked[Section.MaterialIndex];
+						}
+						
+						break;
+					}
+				}
+
+				if (bTriangleIsOpaqueOrMasked)
+				{
+					BuildTriangles.Add(FkDOPBuildCollisionTriangle<uint32>(
+						StoreAsTwoSided, 
+						V0,
+						V1,
+						V2));
+				}
+			}
+		}
+
+		TkDOPTree<const FMeshBuildDataProvider,uint32> kDopTree;
+		kDopTree.Build(BuildTriangles);
+
+		//@todo - project setting
+		const int32 NumVoxelDistanceSamples = 1200;
+		TArray<FVector4> SampleDirections;
+		const int32 NumThetaSteps = FMath::TruncToInt(FMath::Sqrt(NumVoxelDistanceSamples / (2.0f * (float)PI)));
+		const int32 NumPhiSteps = FMath::TruncToInt(NumThetaSteps * (float)PI);
+		FRandomStream RandomStream(0);
+		GenerateStratifiedUniformHemisphereSamples(NumThetaSteps, NumPhiSteps, RandomStream, SampleDirections);
+		TArray<FVector4> OtherHemisphereSamples;
+		GenerateStratifiedUniformHemisphereSamples(NumThetaSteps, NumPhiSteps, RandomStream, OtherHemisphereSamples);
+
+		for (int32 i = 0; i < OtherHemisphereSamples.Num(); i++)
+		{
+			FVector4 Sample = OtherHemisphereSamples[i];
+			Sample.Z *= -1;
+			SampleDirections.Add(Sample);
+		}
+
+		// Meshes with explicit artist-specified scale can go higher
+		//@todo - project setting
+		const int32 MaxNumVoxelsOneDim = DistanceFieldResolutionScale <= 1 ? 64 : 128;
+		const int32 MinNumVoxelsOneDim = 8;
+
+		//@todo - project setting
+		const float NumVoxelsPerLocalSpaceUnit = .1f * DistanceFieldResolutionScale;
+		FBox MeshBounds(Bounds.GetBox());
+
+		{
+			const float MaxOriginalExtent = MeshBounds.GetExtent().GetMax();
+			// Expand so that the edges of the volume are guaranteed to be outside of the mesh
+			const FVector NewExtent(MeshBounds.GetExtent() + FVector(.2f * MaxOriginalExtent));
+			FBox DistanceFieldVolumeBounds = FBox(MeshBounds.GetCenter() - NewExtent, MeshBounds.GetCenter() + NewExtent);
+			const float DistanceFieldVolumeMaxDistance = DistanceFieldVolumeBounds.GetExtent().Size();
+
+			const FVector DesiredDimensions(DistanceFieldVolumeBounds.GetSize() * FVector(NumVoxelsPerLocalSpaceUnit));
+
+			const FIntVector VolumeDimensions(
+				FMath::Clamp(FMath::TruncToInt(DesiredDimensions.X), MinNumVoxelsOneDim, MaxNumVoxelsOneDim), 
+				FMath::Clamp(FMath::TruncToInt(DesiredDimensions.Y), MinNumVoxelsOneDim, MaxNumVoxelsOneDim),
+				FMath::Clamp(FMath::TruncToInt(DesiredDimensions.Z), MinNumVoxelsOneDim, MaxNumVoxelsOneDim));
+
+			OutData.Size = VolumeDimensions;
+			OutData.LocalBoundingBox = DistanceFieldVolumeBounds;
+			OutData.DistanceFieldVolume.AddZeroed(VolumeDimensions.X * VolumeDimensions.Y * VolumeDimensions.Z);
+
+			UE_LOG(LogMeshUtilities,Log,TEXT("Beginning build of mesh - %ux%ux%u distance field, %u triangles"), 
+				VolumeDimensions.X, 
+				VolumeDimensions.Y, 
+				VolumeDimensions.Z, 
+				Indices.Num() / 3);
+
+			const int32 NumThreads = FMath::Max<int32>(FPlatformMisc::NumberOfCoresIncludingHyperthreads() - 2, 1);
+			TArray<FMeshDistanceFieldThreadRunnable*> Threads;
+
+			volatile int32 SharedVolumeZIndex = 0;
+			volatile int32 NegativeAtBorder = 0;
+
+			for (int32 ThreadIndex = 0; ThreadIndex < NumThreads; ThreadIndex++)
+			{
+				Threads.Add(new FMeshDistanceFieldThreadRunnable(
+					kDopTree,
+					SampleDirections,
+					DistanceFieldVolumeBounds,
+					VolumeDimensions,
+					DistanceFieldVolumeMaxDistance,
+					&SharedVolumeZIndex,
+					&NegativeAtBorder,
+					OutData.DistanceFieldVolume));
+			}
+
+			for (int32 ThreadIndex = 0; ThreadIndex < NumThreads; ThreadIndex++)
+			{
+				Threads[ThreadIndex]->WaitForCompletion();
+				delete Threads[ThreadIndex];
+			}
+
+			OutData.bMeshWasClosed = NegativeAtBorder == 0;
+
+			// Toss distance field if mesh was not closed
+			if (NegativeAtBorder != 0)
+			{
+				OutData.Size = FIntVector(0, 0, 0);
+				OutData.DistanceFieldVolume.Empty();
+
+				UE_LOG(LogMeshUtilities,Log,TEXT("Discarded distance field as mesh was not closed!  Assign a two-sided material to fix."));
+			}
+		}
+	}
+}
+
+#else
+
+static void GenerateSignedDistanceFieldVolumeData(
+	FStaticMeshLODResources& LODModel,
+	const TArray<UMaterialInterface*>& Materials,
+	const FBoxSphereBounds& Bounds,
+	float DistanceFieldResolutionScale)
+{
+	static const auto CVar = IConsoleManager::Get().FindTConsoleVariableDataInt(TEXT("r.AllowMeshDistanceFieldRepresentations"));
+
+	if (DistanceFieldResolutionScale > 0 
+		&& CVar 
+		&& CVar->GetValueOnGameThread() != 0)
+	{
+		UE_LOG(LogMeshUtilities,Error,TEXT("Couldn't generate distance field for mesh, platform is missing required Vector intrinsics."));
+	}
+}
+
+#endif
+
 class FMeshUtilities : public IMeshUtilities
 {
 public:
@@ -56,6 +473,7 @@ private:
 	virtual bool BuildStaticMesh(
 		FStaticMeshRenderData& OutRenderData,
 		TArray<FStaticMeshSourceModel>& SourceModels,
+		const TArray<UMaterialInterface*>& Materials,
 		const FStaticMeshLODGroup& LODGroup
 		) OVERRIDE;
 
@@ -1821,6 +2239,7 @@ static void ApplyScaling( FRawMesh& Mesh, float BuildScale )
 bool FMeshUtilities::BuildStaticMesh(
 	FStaticMeshRenderData& OutRenderData,
 	TArray<FStaticMeshSourceModel>& SourceModels,
+	const TArray<UMaterialInterface*>& Materials,
 	const FStaticMeshLODGroup& LODGroup
 	)
 {
@@ -2111,6 +2530,16 @@ bool FMeshUtilities::BuildStaticMesh(
 			OutRenderData.Bounds.SphereRadius
 			);
 	}
+	
+	{
+		FStaticMeshLODResources& LODModel = OutRenderData.LODResources[0];
+
+		GenerateSignedDistanceFieldVolumeData(
+			LODModel,
+			Materials,
+			OutRenderData.Bounds,
+			LODBuildSettings[0].DistanceFieldResolutionScale);
+	}
 
 	return true;
 }
diff --git a/Engine/Source/Developer/MeshUtilities/Public/MeshUtilities.h b/Engine/Source/Developer/MeshUtilities/Public/MeshUtilities.h
--- a/Engine/Source/Developer/MeshUtilities/Public/MeshUtilities.h
+++ b/Engine/Source/Developer/MeshUtilities/Public/MeshUtilities.h
@@ -132,6 +132,7 @@ public:
 	virtual bool BuildStaticMesh(
 		class FStaticMeshRenderData& OutRenderData,
 		TArray<struct FStaticMeshSourceModel>& SourceModels,
+		const TArray<UMaterialInterface*>& Materials,
 		const class FStaticMeshLODGroup& LODGroup
 		) = 0;
 
diff --git a/Engine/Source/Developer/Windows/ShaderFormatD3D/Private/D3D11ShaderCompiler.cpp b/Engine/Source/Developer/Windows/ShaderFormatD3D/Private/D3D11ShaderCompiler.cpp
--- a/Engine/Source/Developer/Windows/ShaderFormatD3D/Private/D3D11ShaderCompiler.cpp
+++ b/Engine/Source/Developer/Windows/ShaderFormatD3D/Private/D3D11ShaderCompiler.cpp
@@ -604,19 +604,15 @@ void CompileD3D11Shader(const FShaderCompilerInput& Input,FShaderCompilerOutput&
 						 BindDesc.Type == D3D11_SIT_UAV_RWBYTEADDRESS || BindDesc.Type == D3D11_SIT_UAV_RWSTRUCTURED_WITH_COUNTER ||
 						 BindDesc.Type == D3D11_SIT_UAV_APPEND_STRUCTURED)
 				{
-					// in D3D11 UAVs for PS are bound with the SetRenderTarget call, so we don't want error messages that the shader is not setting those
-					if(Input.Target.Frequency == SF_Compute)
-					{
-						TCHAR OfficialName[1024];
-						FCString::Strcpy(OfficialName, ANSI_TO_TCHAR(BindDesc.Name));
-
-						Output.ParameterMap.AddParameterAllocation(
-							OfficialName,
-							0,
-							BindDesc.BindPoint,
-							1
-							);
-					}
+					TCHAR OfficialName[1024];
+					FCString::Strcpy(OfficialName, ANSI_TO_TCHAR(BindDesc.Name));
+
+					Output.ParameterMap.AddParameterAllocation(
+						OfficialName,
+						0,
+						BindDesc.BindPoint,
+						1
+						);
 				}
 				else if (BindDesc.Type == D3D11_SIT_STRUCTURED || BindDesc.Type == D3D11_SIT_BYTEADDRESS)
 				{
diff --git a/Engine/Source/Editor/StaticMeshEditor/Private/StaticMeshEditorTools.cpp b/Engine/Source/Editor/StaticMeshEditor/Private/StaticMeshEditorTools.cpp
--- a/Engine/Source/Editor/StaticMeshEditor/Private/StaticMeshEditorTools.cpp
+++ b/Engine/Source/Editor/StaticMeshEditor/Private/StaticMeshEditorTools.cpp
@@ -358,6 +358,25 @@ void FMeshBuildSettingsLayout::GenerateChildContent( IDetailChildrenBuilder& Chi
 			.Font(IDetailLayoutBuilder::GetDetailFont())
 		];
 	}
+
+	{
+		ChildrenBuilder.AddChildContent( LOCTEXT("DistanceFieldResolutionScale", "Distance Field Resolution Scale").ToString() )
+		.NameContent()
+		[
+			SNew(STextBlock)
+			.Font( IDetailLayoutBuilder::GetDetailFont() )
+			.Text(LOCTEXT("DistanceFieldResolutionScale", "Distance Field Resolution Scale").ToString())
+		]
+		.ValueContent()
+		[
+			SNew(SSpinBox<float>)
+			.Font( IDetailLayoutBuilder::GetDetailFont() )
+			.MinValue(0.0f)
+			.MaxValue(100.0f)
+			.Value(this, &FMeshBuildSettingsLayout::GetDistanceFieldResolutionScale)
+			.OnValueChanged(this, &FMeshBuildSettingsLayout::OnDistanceFieldResolutionScaleChanged)
+		];
+	}
 		
 	{
 		ChildrenBuilder.AddChildContent( LOCTEXT("ApplyChanges", "Apply Changes").ToString() )
@@ -425,6 +444,11 @@ TOptional<float> FMeshBuildSettingsLayout::GetBuildScaleZ() const
 	return BuildSettings.BuildScale3D.Z;
 }
 
+float FMeshBuildSettingsLayout::GetDistanceFieldResolutionScale() const
+{
+	return BuildSettings.DistanceFieldResolutionScale;
+}
+
 void FMeshBuildSettingsLayout::OnRecomputeNormalsChanged(ESlateCheckBoxState::Type NewState)
 {
 	BuildSettings.bRecomputeNormals = (NewState == ESlateCheckBoxState::Checked) ? true : false;
@@ -469,6 +493,11 @@ void FMeshBuildSettingsLayout::OnBuildScaleZChanged( float NewScaleZ, ETextCommi
 	}
 }
 
+void FMeshBuildSettingsLayout::OnDistanceFieldResolutionScaleChanged(float NewValue)
+{
+	BuildSettings.DistanceFieldResolutionScale = NewValue;
+}
+
 FMeshReductionSettingsLayout::FMeshReductionSettingsLayout( TSharedRef<FLevelOfDetailSettingsLayout> InParentLODSettings )
 	: ParentLODSettings( InParentLODSettings )
 {
diff --git a/Engine/Source/Editor/StaticMeshEditor/Private/StaticMeshEditorTools.h b/Engine/Source/Editor/StaticMeshEditor/Private/StaticMeshEditorTools.h
--- a/Engine/Source/Editor/StaticMeshEditor/Private/StaticMeshEditorTools.h
+++ b/Engine/Source/Editor/StaticMeshEditor/Private/StaticMeshEditorTools.h
@@ -140,6 +140,7 @@ private:
 	TOptional<float> GetBuildScaleX() const;
 	TOptional<float> GetBuildScaleY() const;
 	TOptional<float> GetBuildScaleZ() const;
+	float GetDistanceFieldResolutionScale() const;
 
 	void OnRecomputeNormalsChanged(ESlateCheckBoxState::Type NewState);
 	void OnRecomputeTangentsChanged(ESlateCheckBoxState::Type NewState);
@@ -149,6 +150,7 @@ private:
 	void OnBuildScaleYChanged( float NewScaleY, ETextCommit::Type TextCommitType );
 	void OnBuildScaleZChanged( float NewScaleZ, ETextCommit::Type TextCommitType );
 
+	void OnDistanceFieldResolutionScaleChanged(float NewValue);
 private:
 	TWeakPtr<FLevelOfDetailSettingsLayout> ParentLODSettings;
 	FMeshBuildSettings BuildSettings;
@@ -406,4 +408,4 @@ private:
 
 	/** Modified by a spinbox, user setting for the min spacing between charts. */
 	float MinSpacingBetweenCharts;
-};
\ No newline at end of file
+};
diff --git a/Engine/Source/Editor/StaticMeshEditor/Private/StaticMeshEditorViewportClient.cpp b/Engine/Source/Editor/StaticMeshEditor/Private/StaticMeshEditorViewportClient.cpp
--- a/Engine/Source/Editor/StaticMeshEditor/Private/StaticMeshEditorViewportClient.cpp
+++ b/Engine/Source/Editor/StaticMeshEditor/Private/StaticMeshEditorViewportClient.cpp
@@ -515,6 +515,40 @@ void FStaticMeshEditorViewportClient::DrawCanvas( FViewport& InViewport, FSceneV
 		);
 	YPos += 18;
 
+	const FDistanceFieldVolumeData& VolumeData = StaticMesh->RenderData->LODResources[0].DistanceFieldData;
+
+	if (VolumeData.Size.GetMax() > 0)
+	{
+		float MemoryMb = (VolumeData.Size.X * VolumeData.Size.Y * VolumeData.Size.Z * VolumeData.DistanceFieldVolume.GetTypeSize()) / (1024.0f * 1024.0f);
+
+		FNumberFormattingOptions NumberOptions;
+		NumberOptions.MinimumFractionalDigits = 2;
+		NumberOptions.MaximumFractionalDigits = 2;
+
+		if (VolumeData.bMeshWasClosed)
+		{
+			Canvas.DrawShadowedString(
+				6,
+				YPos,
+				*FText::Format( NSLOCTEXT("UnrealEd", "DistanceFieldRes_F", "Distance Field:  {0}x{1}x{2} = {3}mb"), FText::AsNumber(VolumeData.Size.X), FText::AsNumber(VolumeData.Size.Y), FText::AsNumber(VolumeData.Size.Z), FText::AsNumber(MemoryMb, &NumberOptions) ).ToString(),
+				GEngine->GetSmallFont(),
+				FLinearColor::White
+				);
+		}
+		else
+		{
+			Canvas.DrawShadowedString(
+				6,
+				YPos,
+				*NSLOCTEXT("UnrealEd", "DistanceFieldClosed_F", "Distance Field:  Mesh was not closed and material was one-sided").ToString(),
+				GEngine->GetSmallFont(),
+				FLinearColor::Red
+				);
+		}
+		
+		YPos += 18;
+	}
+
 	Canvas.DrawShadowedString(
 		6,
 		YPos,
diff --git a/Engine/Source/Editor/UnrealEd/Private/StaticLightingSystem/StaticLightingSystem.cpp b/Engine/Source/Editor/UnrealEd/Private/StaticLightingSystem/StaticLightingSystem.cpp
--- a/Engine/Source/Editor/UnrealEd/Private/StaticLightingSystem/StaticLightingSystem.cpp
+++ b/Engine/Source/Editor/UnrealEd/Private/StaticLightingSystem/StaticLightingSystem.cpp
@@ -1810,7 +1810,7 @@ void FStaticLightingSystem::GatherScene()
 		ULightComponentBase* LightBase = Lights[LightIndex];
 		USkyLightComponent* SkyLight = Cast<USkyLightComponent>(LightBase);
 
-		if (SkyLight)
+		if (SkyLight && (SkyLight->Mobility == EComponentMobility::Static || SkyLight->Mobility == EComponentMobility::Stationary))
 		{
 			LightmassExporter->AddLight(SkyLight);
 		}
diff --git a/Engine/Source/Runtime/Core/Private/HAL/ConsoleManager.cpp b/Engine/Source/Runtime/Core/Private/HAL/ConsoleManager.cpp
--- a/Engine/Source/Runtime/Core/Private/HAL/ConsoleManager.cpp
+++ b/Engine/Source/Runtime/Core/Private/HAL/ConsoleManager.cpp
@@ -1826,6 +1826,12 @@ void CreateConsoleVariables()
 		TEXT("Games that only use dynamic lighting should set this to 0 to save some static lighting overhead."),
 		ECVF_ReadOnly | ECVF_RenderThreadSafe);
 
+	IConsoleManager::Get().RegisterConsoleVariable(TEXT("r.AllowMeshDistanceFieldRepresentations"),
+		0,	
+		TEXT("Whether to build distance fields of static meshes, needed for distance field AO, which is used to implement Movable SkyLight shadows.\n")
+		TEXT("Enabling will increase mesh build times and memory usage.  Changing this value will cause a rebuild of all static meshes."),
+		ECVF_ReadOnly);
+
 	IConsoleManager::Get().RegisterConsoleVariable(TEXT("con.MinLogVerbosity"),
 		0,
 		TEXT("Allows to see the log in the in game console (by default deactivated to avoid spam and minor performance loss).\n")
diff --git a/Engine/Source/Runtime/Core/Public/Math/IntPoint.h b/Engine/Source/Runtime/Core/Public/Math/IntPoint.h
--- a/Engine/Source/Runtime/Core/Public/Math/IntPoint.h
+++ b/Engine/Source/Runtime/Core/Public/Math/IntPoint.h
@@ -224,6 +224,7 @@ public:
 	 * @return A new divided int point.
 	 */
 	static FIntPoint DivideAndRoundUp( FIntPoint lhs, int32 Divisor );
+	static FIntPoint DivideAndRoundDown( FIntPoint lhs, int32 Divisor );
 
 	/**
 	 * Gets number of components point has.
@@ -366,6 +367,10 @@ FORCEINLINE FIntPoint FIntPoint::DivideAndRoundUp( FIntPoint lhs, int32 Divisor
 	return FIntPoint(FMath::DivideAndRoundUp(lhs.X, Divisor), FMath::DivideAndRoundUp(lhs.Y, Divisor));
 }	
 
+FORCEINLINE FIntPoint FIntPoint::DivideAndRoundDown( FIntPoint lhs, int32 Divisor )
+{
+	return FIntPoint(FMath::DivideAndRoundDown(lhs.X, Divisor), FMath::DivideAndRoundDown(lhs.Y, Divisor));
+}	
 
 FORCEINLINE FIntPoint FIntPoint::operator+( const FIntPoint& Other ) const
 {
