Delete extraneous MainVS left-over from recent merge that was causing hlslcc to fail.

diff --git a/Engine/Shaders/PostProcessHMD.usf b/Engine/Shaders/PostProcessHMD.usf
--- a/Engine/Shaders/PostProcessHMD.usf
+++ b/Engine/Shaders/PostProcessHMD.usf
@@ -48,19 +48,6 @@ void MainVS(in float2 Position : ATTRIBUTE0, in float2 TexCoord0 : ATTRIBUTE1, i
 	OutColor = Color;              // Used for vignette fade.
 }
 
-
-// vertex shader entry point
-void MainVS(
-	in float4 InPosition : ATTRIBUTE0,
-	in float2 InTexCoord : ATTRIBUTE1,
-	out float4 OutTexCoord : TEXCOORD0,
-	out float4 OutPosition : SV_POSITION
-	)
-{
-	OutPosition = InPosition;
-	OutTexCoord = float4(InTexCoord, InPosition.xy);
-}
-
 void MainPS(in float4 Position : SV_Position, in float4 Color : COLOR,
 			in float2 TexCoord0 : TEXCOORD0, in float2 TexCoord1 : TEXCOORD1, in float2 TexCoord2 : TEXCOORD2, out float4 OutColor : SV_Target0)
 {
