Add casts and hash-define's to make new DistanceField shaders compile on hlslcc, Shader Model 4 platforms like OS X.

diff --git a/Engine/Shaders/DistanceFieldLightingShared.usf b/Engine/Shaders/DistanceFieldLightingShared.usf
--- a/Engine/Shaders/DistanceFieldLightingShared.usf
+++ b/Engine/Shaders/DistanceFieldLightingShared.usf
@@ -99,9 +99,9 @@ uint4 GetTileHead(uint2 TileCoordinate)
 #if UNPACKED_TILE_HEAD
 	return uint4(
 		TileHeadDataUnpacked[TileIndex * 4 + 0], 
-		min(TileHeadDataUnpacked[TileIndex * 4 + 1], MAX_OBJECTS_PER_TILE), 
-		min(TileHeadDataUnpacked[TileIndex * 4 + 2], MAX_OBJECTS_PER_TILE), 
-		min(TileHeadDataUnpacked[TileIndex * 4 + 3], MAX_OBJECTS_PER_TILE));
+		min(TileHeadDataUnpacked[TileIndex * 4 + 1], (uint)MAX_OBJECTS_PER_TILE),
+		min(TileHeadDataUnpacked[TileIndex * 4 + 2], (uint)MAX_OBJECTS_PER_TILE),
+		min(TileHeadDataUnpacked[TileIndex * 4 + 3], (uint)MAX_OBJECTS_PER_TILE));
 #else
 	return TileHeadData[TileIndex];
 #endif
diff --git a/Engine/Shaders/DistanceFieldSurfaceCacheLighting.usf b/Engine/Shaders/DistanceFieldSurfaceCacheLighting.usf
--- a/Engine/Shaders/DistanceFieldSurfaceCacheLighting.usf
+++ b/Engine/Shaders/DistanceFieldSurfaceCacheLighting.usf
@@ -36,6 +36,7 @@ void ObjectCullVS(
 	Output.ObjectIndex = ObjectIndex;
 } 
 
+#if FEATURE_LEVEL >= FEATURE_LEVEL_SM5
 RWBuffer<uint> RWTileHeadDataUnpacked;
 RWBuffer<uint> RWTileArrayData;
 
@@ -100,6 +101,7 @@ void ObjectCullPS(
 		}
 	}
 }
+#endif
 
 
 float2 UVToTileHead;
@@ -195,6 +197,8 @@ void ComputeDistanceFieldNormalPS(
 	OutColor = float4(EncodeNormalForAO(WorldNormal), DepthSign * SceneDepth);
 }
 
+#if FEATURE_LEVEL >= FEATURE_LEVEL_SM5
+
 #define THREADS_PER_PIXEL 1
 
 #if THREADS_PER_PIXEL != 1
@@ -335,6 +339,7 @@ void ComputeDistanceFieldNormalCS(
 		RWDistanceFieldNormal[PixelCoordinate] = OutValue;
 	}
 }
+#endif
 
 // For some reason gives innaccurate results at lower resolutions
 #define USE_SCREENVECTOR_WORLD_POSITION FINAL_INTERPOLATION_PASS
