Improvements to quality, sharpness, and performance for temporal AA.

diff --git a/Engine/Shaders/PostProcessTemporalAA.usf b/Engine/Shaders/PostProcessTemporalAA.usf
--- a/Engine/Shaders/PostProcessTemporalAA.usf
+++ b/Engine/Shaders/PostProcessTemporalAA.usf
@@ -20,6 +20,86 @@ float2 RandomOffset;
 // TODO: This can be removed.
 float ClampFeebackMix;
 
+
+
+// Faster but less accurate luma computation. 
+// Luma includes a scaling by 4.
+float Luma4(float3 Color)
+{
+	return (Color.g * 2.0) + (Color.r + Color.b);
+}
+
+// Optimized HDR weighting function.
+float HdrWeight4(float3 Color, float Exposure) 
+{
+	return rcp(Luma4(Color) * Exposure + 4.0);
+}
+
+float HdrWeightG(float3 Color, float Exposure) 
+{
+	return rcp(Color.g * Exposure + 1.0);
+}
+
+float HdrWeightG_(float Color, float Exposure) 
+{
+	return rcp(Color * Exposure + 1.0);
+}
+
+
+// Optimized HDR weighting function.
+float HdrWeight4_(float Color, float Exposure) 
+{
+	return rcp(Color * Exposure + 4.0);
+}
+
+// Optimized HDR weighting inverse.
+float HdrWeightInv4(float3 Color, float Exposure) 
+{
+	return 4.0 * rcp(Luma4(Color) * (-Exposure) + 1.0);
+}
+
+float HdrWeightInvG(float3 Color, float Exposure) 
+{
+	return rcp(Color.g * (-Exposure) + 1.0);
+}
+
+float HdrWeightInv4_(float Color, float Exposure) 
+{
+	return 4.0 * rcp(Color * (-Exposure) + 1.0);
+}
+
+float HdrWeightInvG_(float Color, float Exposure) 
+{
+	return rcp(Color * (-Exposure) + 1.0);
+}
+
+
+// This returns exposure normalized linear luma from a PerceptualLuma4().
+float LinearLuma4(float Channel, float Exposure) 
+{
+	return Channel * HdrWeightInv4_(Channel, Exposure);
+}
+
+// This returns exposure normalized linear luma from a PerceptualLuma4().
+float LinearLumaG(float Channel, float Exposure) 
+{
+	return Channel * HdrWeightInvG_(Channel, Exposure);
+}
+
+
+float PerceptualLuma4(float3 Color, float Exposure) 
+{
+	float L = Luma4(Color);
+	return L * HdrWeight4_(L, Exposure);
+}
+
+float PerceptualLumaG(float3 Color, float Exposure) 
+{
+	return Color.g * HdrWeightG_(Color.g, Exposure);
+}
+
+
+
 float Luma(float3 Color) 
 {
 	#if 1
@@ -60,9 +140,9 @@ float LinearLuma(float Channel)
 // Returns distance on line segment.
 float IntersectAABB(float3 Dir, float3 Org, float3 Box)
 {
-	#if 1 // PS4_PROFILE
-	    // Turning this on for all platforms in case of bad AABB test.
-		if(min(min(abs(Dir.x), abs(Dir.y)), abs(Dir.z)) < (1.0/16384.0)) return 1.0;
+	#if PS4_PROFILE
+	    // This causes flicker, it should only be used on PS4 until proper fix is in.
+		if(min(min(abs(Dir.x), abs(Dir.y)), abs(Dir.z)) < (1.0/65536.0)) return 1.0;
 	#endif
 	float3 RcpDir = rcp(Dir);
 	float3 TNeg = (  Box  - Org) * RcpDir;
@@ -70,7 +150,7 @@ float IntersectAABB(float3 Dir, float3 Org, float3 Box)
 	return max(max(min(TNeg.x, TPos.x), min(TNeg.y, TPos.y)), min(TNeg.z, TPos.z));
 }
 
-float HistoryClamp(float3 History, float3 Filtered, float3 NeighborMin, float3 NeighborMax, float Exposure)
+float HistoryClamp(float3 History, float3 Filtered, float3 NeighborMin, float3 NeighborMax)
 {
 	float3 Min = min(Filtered, min(NeighborMin, NeighborMax));
 	float3 Max = max(Filtered, max(NeighborMin, NeighborMax));	
@@ -102,6 +182,7 @@ void SSRTemporalAAPS( float4 UVAndScreenPos : TEXCOORD0, float3 InExposureScaleV
 {
 #if ENABLE_TEMPORAL_AA
 	float InExposureScale = InExposureScaleVignette.x;
+	#define AA_NEW 1
 	#define AA_ALPHA 0
 	#define AA_CROSS 0
 	#define AA_DYNAMIC 0
@@ -118,6 +199,7 @@ void DOFTemporalAAPS( float4 UVAndScreenPos : TEXCOORD0, float3 InExposureScaleV
 {
 #if ENABLE_TEMPORAL_AA
 	float InExposureScale = InExposureScaleVignette.x;
+	#define AA_NEW 1
 	#define AA_ALPHA 0
 	#define AA_CROSS 4
 	#define AA_DYNAMIC 0
@@ -134,6 +216,7 @@ void LightShaftTemporalAAPS( float4 UVAndScreenPos : TEXCOORD0, float3 InExposur
 {
 #if ENABLE_TEMPORAL_AA
 	float InExposureScale = InExposureScaleVignette.x;
+	#define AA_NEW 1
 	#define AA_ALPHA 0
 	#define AA_CROSS 0
 	#define AA_DYNAMIC 0
@@ -149,8 +232,21 @@ void MainTemporalAAPS( float4 UVAndScreenPos : TEXCOORD0, float3 InExposureScale
 {
 #if ENABLE_TEMPORAL_AA
 	float InExposureScale = InExposureScaleVignette.x;
-	#define AA_BORDER 1
-	#define AA_GRAIN 1
+	#if 1
+		#define AA_NEW 1
+		#define AA_FILTERED 1
+		#define AA_BORDER 1
+		#define AA_CROSS 2
+		#define AA_GREEN_AS_LUMA 0
+		#define AA_AABB 1
+		#define AA_LOWPASS 1
+		#define AA_ONE_DYNAMIC_SAMPLE 1
+		#define AA_DEBUG 0
+	#else
+		#define AA_NEW 0
+		#define AA_BORDER 1
+		#define AA_GRAIN 1
+	#endif
 	#include "PostProcessTemporalCommon.usf"
 #else 
 	OutColor = PostprocessInput0.SampleLevel(PostprocessInput0Sampler, UVAndScreenPos.xy, 0);
@@ -161,12 +257,25 @@ void MainFastTemporalAAPS( float4 UVAndScreenPos : TEXCOORD0, float3 InExposureS
 {
 #if ENABLE_TEMPORAL_AA
 	float InExposureScale = InExposureScaleVignette.x;
-	#define AA_BORDER 1
-	#define AA_GRAIN 1
-	#define AA_HDR 0
-	#define AA_HDR_HISTORY 0
-	#define AA_LOWPASS 0
-	#define AA_ONE_DYNAMIC_SAMPLE 1
+	#if 1
+		#define AA_NEW 1
+		#define AA_FILTERED 1
+		#define AA_BORDER 1
+		#define AA_CROSS 2
+		#define AA_GREEN_AS_LUMA 1
+		#define AA_AABB 1
+		#define AA_LOWPASS 0
+		#define AA_ONE_DYNAMIC_SAMPLE 1
+		#define AA_DEBUG 0
+	#else
+		#define AA_NEW 0
+		#define AA_BORDER 1
+		#define AA_GRAIN 1
+		#define AA_HDR 0
+		#define AA_HDR_HISTORY 0
+		#define AA_LOWPASS 0
+		#define AA_ONE_DYNAMIC_SAMPLE 1
+	#endif
 	#include "PostProcessTemporalCommon.usf"
 #else 
 	OutColor = PostprocessInput0.SampleLevel(PostprocessInput0Sampler, UVAndScreenPos.xy, 0);
diff --git a/Engine/Shaders/PostProcessTemporalCommon.usf b/Engine/Shaders/PostProcessTemporalCommon.usf
--- a/Engine/Shaders/PostProcessTemporalCommon.usf
+++ b/Engine/Shaders/PostProcessTemporalCommon.usf
@@ -110,6 +110,403 @@ which can be configured via defines for various temporal AA passes.
 	#define AA_ONE_DYNAMIC_SAMPLE 0
 #endif
 
+// Use green as luma.
+#ifndef AA_GREEN_AS_LUMA
+	#define AA_GREEN_AS_LUMA 0
+#endif
+
+// Try the new AA algorithm.
+#ifndef AA_NEW 
+	#define AA_NEW 0
+#endif
+
+
+
+#if AA_NEW
+
+	// FIND MOTION OF PIXEL AND NEAREST IN NEIGHBORHOOD
+    // ------------------------------------------------
+	float3 PosP; // Position of this pixel.
+	PosP.xy = UVAndScreenPos.zw * float2(0.5, -0.5) + 0.5; // View position [0 to 1] flipped in Y.
+	PosP.z = SceneDepthTexture.SampleLevel(SceneDepthTextureSampler, UVAndScreenPos.xy, 0).r;
+	float3 PosN; // Position of closest pixel in neighborhood.
+	PosN = PosP;
+	// Screen position of minimum depth.
+	float2 VelocityOffset = float2(0.0, 0.0);
+	#if AA_CROSS
+		// For motion vector, use camera/dynamic motion from min depth pixel in pattern around pixel.
+		// This enables better quality outline on foreground against different motion background.
+		// Larger 2 pixel distance "x" works best (because AA diolates surface).
+		float4 Depths;
+		#if FEATURE_LEVEL >= FEATURE_LEVEL_SM5
+			Depths = SceneDepthTexture.GatherRed(SceneDepthTextureSampler, UVAndScreenPos.xy, int2(-AA_CROSS, -AA_CROSS), int2(AA_CROSS, -AA_CROSS), int2(-AA_CROSS, AA_CROSS), int2(AA_CROSS, AA_CROSS));
+		#else
+			Depths.x = SceneDepthTexture.SampleLevel(SceneDepthTextureSampler, UVAndScreenPos.xy, 0, int2(-AA_CROSS, -AA_CROSS)).r;
+			Depths.y = SceneDepthTexture.SampleLevel(SceneDepthTextureSampler, UVAndScreenPos.xy, 0, int2( AA_CROSS, -AA_CROSS)).r;
+			Depths.z = SceneDepthTexture.SampleLevel(SceneDepthTextureSampler, UVAndScreenPos.xy, 0, int2(-AA_CROSS,  AA_CROSS)).r;
+			Depths.w = SceneDepthTexture.SampleLevel(SceneDepthTextureSampler, UVAndScreenPos.xy, 0, int2( AA_CROSS,  AA_CROSS)).r;
+		#endif
+		float2 DepthOffset = float2(AA_CROSS, AA_CROSS);
+		float DepthOffsetXx = float(AA_CROSS);
+		// Nearest depth is the largest depth (depth surface 0=far, 1=near).
+		if(Depths.x > Depths.y) 
+		{
+			DepthOffsetXx = -AA_CROSS;
+		}
+		if(Depths.z > Depths.w) 
+		{
+			DepthOffset.x = -AA_CROSS;
+		}
+		float DepthsXY = max(Depths.x, Depths.y);
+		float DepthsZW = max(Depths.z, Depths.w);
+		if(DepthsXY > DepthsZW) 
+		{
+			DepthOffset.y = -AA_CROSS;
+			DepthOffset.x = DepthOffsetXx; 
+		}
+		float DepthsXYZW = max(DepthsXY, DepthsZW);
+		if(DepthsXYZW > PosP.z) 
+		{
+			// This is offset for reading from velocity texture.
+			// This supports half or fractional resolution velocity textures.
+			// With the assumption that UV position scales between velocity and color.
+			VelocityOffset = DepthOffset * PostprocessInput0Size.zw;
+			// This is [0 to 1] flipped in Y.
+			PosN.xy = (UVAndScreenPos.zw + DepthOffset * ViewportSize.zw * 2.0) * float2(0.5, -0.5) + 0.5;
+			PosN.z = DepthsXYZW;
+		}
+	#endif
+	// Camera motion for pixel (in ScreenPos space).
+	float ScaleM = 1.0 / (dot(PosP, CameraMotion[0].xyz) + CameraMotion[0].w);
+	float2 BackP;
+	BackP.x = -2.0 * ((PosP.x * ((CameraMotion[1].x * PosP.y) + (CameraMotion[1].y * PosP.z) + CameraMotion[1].z)) + (CameraMotion[1].w * PosP.y) + (CameraMotion[2].x * PosP.x * PosP.x) + (CameraMotion[2].y * PosP.z) + CameraMotion[2].z) * ScaleM;
+	BackP.y =  2.0 * ((PosP.y * ((CameraMotion[3].x * PosP.x) + (CameraMotion[3].y * PosP.z) + CameraMotion[3].z)) + (CameraMotion[3].w * PosP.x) + (CameraMotion[4].x * PosP.y * PosP.y) + (CameraMotion[4].y * PosP.z) + CameraMotion[4].z) * ScaleM;
+	// Convert BackP back into pixels.
+	float2 BackTemp = BackP * ViewportSize.xy;
+	// Save the amount of pixel offset of just camera motion, used later as the amount of blur introduced by history.
+	float HistoryBlurAmp = 2.0;
+	float HistoryBlur = saturate(abs(BackTemp.x) * HistoryBlurAmp + abs(BackTemp.y) * HistoryBlurAmp);
+	#if AA_CROSS
+		// Camera motion for neighborhood.
+		ScaleM = 1.0 / (dot(PosN, CameraMotion[0].xyz) + CameraMotion[0].w);
+		float2 BackN;
+		BackN.x = -2.0 * ((PosN.x * ((CameraMotion[1].x * PosN.y) + (CameraMotion[1].y * PosN.z) + CameraMotion[1].z)) + (CameraMotion[1].w * PosN.y) + (CameraMotion[2].x * PosN.x * PosN.x) + (CameraMotion[2].y * PosN.z) + CameraMotion[2].z) * ScaleM;
+		BackN.y =  2.0 * ((PosN.y * ((CameraMotion[3].x * PosN.x) + (CameraMotion[3].y * PosN.z) + CameraMotion[3].z)) + (CameraMotion[3].w * PosN.x) + (CameraMotion[4].x * PosN.y * PosN.y) + (CameraMotion[4].y * PosN.z) + CameraMotion[4].z) * ScaleM;
+		BackTemp = BackN * ViewportSize.xy;
+	#endif
+	#if AA_DYNAMIC
+		// Dynamic motion of pixel.
+		#if (AA_CROSS == 0) || (AA_ONE_DYNAMIC_SAMPLE == 0)
+			float2 VelocityP = PostprocessInput3.SampleLevel(PostprocessInput3Sampler, UVAndScreenPos.xy, 0).xy;
+			bool DynamicP = VelocityP.x > 0.0;
+			if(DynamicP)
+			{
+				BackP = DecodeVelocityFromTexture(VelocityP);
+			}
+		#endif
+		BackTemp = BackP * ViewportSize.xy;
+		// At this point DoublePixelMotionDynamic is kept squared until later.
+		float DoublePixelMotionDynamic = dot(BackTemp, BackTemp);
+		#if AA_CROSS
+			// Dynamic motion of nearest in neighborhood.
+			float2 VelocityN = PostprocessInput3.SampleLevel(PostprocessInput3Sampler, UVAndScreenPos.xy + VelocityOffset, 0).xy;
+			bool DynamicN = VelocityN.x > 0.0;
+			if(DynamicN)
+			{
+				BackN = DecodeVelocityFromTexture(VelocityN);
+			}
+			#if (AA_ONE_DYNAMIC_SAMPLE == 0)
+				bool DynamicMotion = DynamicN || DynamicP;
+			#else
+				bool DynamicMotion = DynamicN;
+			#endif
+			BackTemp = BackN * ViewportSize.xy;
+			DoublePixelMotionDynamic = max(DoublePixelMotionDynamic, dot(BackTemp, BackTemp));
+		#else
+			bool DynamicMotion = DynamicP;
+		#endif
+		DoublePixelMotionDynamic = sqrt(DoublePixelMotionDynamic);
+	#else
+		bool DynamicMotion = false;
+		float DoublePixelMotionDynamic = 0.0;
+	#endif
+	#if (AA_CROSS == 0)
+		float2 BackN = BackP;
+	#endif
+	// Easier to do off screen check before conversion.
+	// BackN is in units of 2pixels/viewportWidthInPixels
+	// This converts back projection vector to [-1 to 1] offset in viewport.
+	BackN = UVAndScreenPos.zw - BackN;
+	bool OffScreen = max(abs(BackN.x), abs(BackN.y)) >= 1.0;
+	// Also clamp to be on screen (fixes problem with DOF).
+	// The .z and .w is the 1/width and 1/height.
+	// This clamps to be a pixel inside the viewport.
+	BackN.x = clamp(BackN.x, -1.0 + ViewportSize.z, 1.0 - ViewportSize.z);
+	BackN.y = clamp(BackN.y, -1.0 + ViewportSize.w, 1.0 - ViewportSize.w);
+	// Convert from [-1 to 1] to view rectangle which is somewhere in [0 to 1].
+	// The extra +0.5 factor is because ScreenPosToPixel.zw is incorrectly computed
+	// as the upper left of the pixel instead of the center of the pixel.
+	BackN = (BackN * ScreenPosToPixel.xy + ScreenPosToPixel.zw + 0.5) * PostprocessInput0Size.zw;
+
+	// FILTER PIXEL (RESAMPLE TO REMOVE JITTER OFFSET) AND GET NEIGHBORHOOD
+    // --------------------------------------------------------------------
+	// 012
+	// 345
+	// 678
+	float4 Neighbor0 = PostprocessInput0.SampleLevel(PostprocessInput0Sampler, UVAndScreenPos.xy, 0, int2(-1, -1));
+	float4 Neighbor1 = PostprocessInput0.SampleLevel(PostprocessInput0Sampler, UVAndScreenPos.xy, 0, int2( 0, -1));
+	float4 Neighbor2 = PostprocessInput0.SampleLevel(PostprocessInput0Sampler, UVAndScreenPos.xy, 0, int2( 1, -1));
+	float4 Neighbor3 = PostprocessInput0.SampleLevel(PostprocessInput0Sampler, UVAndScreenPos.xy, 0, int2(-1,  0));
+	float4 Neighbor4 = PostprocessInput0.SampleLevel(PostprocessInput0Sampler, UVAndScreenPos.xy, 0);
+	float4 Neighbor5 = PostprocessInput0.SampleLevel(PostprocessInput0Sampler, UVAndScreenPos.xy, 0, int2( 1,  0));
+	float4 Neighbor6 = PostprocessInput0.SampleLevel(PostprocessInput0Sampler, UVAndScreenPos.xy, 0, int2(-1,  1));
+	float4 Neighbor7 = PostprocessInput0.SampleLevel(PostprocessInput0Sampler, UVAndScreenPos.xy, 0, int2( 0,  1));
+	float4 Neighbor8 = PostprocessInput0.SampleLevel(PostprocessInput0Sampler, UVAndScreenPos.xy, 0, int2( 1,  1));
+	#if AA_GREEN_AS_LUMA
+		Neighbor0.rgb *= HdrWeightG(Neighbor0.rgb, InExposureScale);
+		Neighbor1.rgb *= HdrWeightG(Neighbor1.rgb, InExposureScale);
+		Neighbor2.rgb *= HdrWeightG(Neighbor2.rgb, InExposureScale);
+		Neighbor3.rgb *= HdrWeightG(Neighbor3.rgb, InExposureScale);
+		Neighbor4.rgb *= HdrWeightG(Neighbor4.rgb, InExposureScale);
+		Neighbor5.rgb *= HdrWeightG(Neighbor5.rgb, InExposureScale);
+		Neighbor6.rgb *= HdrWeightG(Neighbor6.rgb, InExposureScale);
+		Neighbor7.rgb *= HdrWeightG(Neighbor7.rgb, InExposureScale);
+		Neighbor8.rgb *= HdrWeightG(Neighbor8.rgb, InExposureScale);
+	#else
+		Neighbor0.rgb *= HdrWeight4(Neighbor0.rgb, InExposureScale);
+		Neighbor1.rgb *= HdrWeight4(Neighbor1.rgb, InExposureScale);
+		Neighbor2.rgb *= HdrWeight4(Neighbor2.rgb, InExposureScale);
+		Neighbor3.rgb *= HdrWeight4(Neighbor3.rgb, InExposureScale);
+		Neighbor4.rgb *= HdrWeight4(Neighbor4.rgb, InExposureScale);
+		Neighbor5.rgb *= HdrWeight4(Neighbor5.rgb, InExposureScale);
+		Neighbor6.rgb *= HdrWeight4(Neighbor6.rgb, InExposureScale);
+		Neighbor7.rgb *= HdrWeight4(Neighbor7.rgb, InExposureScale);
+		Neighbor8.rgb *= HdrWeight4(Neighbor8.rgb, InExposureScale);
+	#endif
+	#if AA_FILTERED
+		float4 Filtered = 
+			Neighbor0 * SampleWeights[0] +
+			Neighbor1 * SampleWeights[1] +
+			Neighbor2 * SampleWeights[2] +
+			Neighbor3 * SampleWeights[3] +
+			Neighbor4 * SampleWeights[4] +
+			Neighbor5 * SampleWeights[5] +
+			Neighbor6 * SampleWeights[6] +
+			Neighbor7 * SampleWeights[7] +
+			Neighbor8 * SampleWeights[8];
+		#if AA_LOWPASS
+			float4 FilteredLow = 
+				Neighbor0 * LowpassWeights[0] +
+				Neighbor1 * LowpassWeights[1] +
+				Neighbor2 * LowpassWeights[2] +
+				Neighbor3 * LowpassWeights[3] +
+				Neighbor4 * LowpassWeights[4] +
+				Neighbor5 * LowpassWeights[5] +
+				Neighbor6 * LowpassWeights[6] +
+				Neighbor7 * LowpassWeights[7] +
+				Neighbor8 * LowpassWeights[8];
+		#else
+			float4 FilteredLow = Filtered;
+		#endif
+		#if AA_BORDER
+			// Use unfiltered for 1 pixel border.
+			float2 TestPos = abs(UVAndScreenPos.zw);
+			// Add 1 pixel and check if off screen.
+			TestPos += ViewportSize.zw * 2.0;
+			bool FilteredOffScreen = max(TestPos.x, TestPos.y) >= 1.0;
+			if(FilteredOffScreen) 
+			{
+				Filtered = Neighbor4;
+				FilteredLow = Neighbor4;
+			}
+		#endif
+	#else
+		// Unfiltered.
+		float4 Filtered = Neighbor4;
+		float4 FilteredLow = Neighbor4;
+	#endif
+	#if AA_ROUND
+		float4 NeighborMin2 = min(min(Neighbor0, Neighbor2), min(Neighbor6, Neighbor8));		
+		float4 NeighborMax2 = max(max(Neighbor0, Neighbor2), max(Neighbor6, Neighbor8));		
+		float4 NeighborMin = min(min(min(Neighbor1, Neighbor3), min(Neighbor4, Neighbor5)), Neighbor7);		
+		float4 NeighborMax = max(max(max(Neighbor1, Neighbor3), max(Neighbor4, Neighbor5)), Neighbor7);		
+		NeighborMin2 = min(NeighborMin2, NeighborMin);
+		NeighborMax2 = max(NeighborMax2, NeighborMax);
+		NeighborMin = NeighborMin * 0.5 + NeighborMin2 * 0.5;
+		NeighborMax = NeighborMax * 0.5 + NeighborMax2 * 0.5;
+	#else
+		float4 NeighborMin = min(min(
+			min(min(Neighbor0, Neighbor1), min(Neighbor2, Neighbor3)),
+			min(min(Neighbor4, Neighbor5), min(Neighbor6, Neighbor7))), Neighbor8);
+		float4 NeighborMax = max(max(
+			max(max(Neighbor0, Neighbor1), max(Neighbor2, Neighbor3)),
+			max(max(Neighbor4, Neighbor5), max(Neighbor6, Neighbor7))), Neighbor8);
+	#endif
+
+	// FETCH HISTORY
+	// -------------
+	OutColor = PostprocessInput1.SampleLevel(PostprocessInput1Sampler, BackN.xy, 0);
+	#if AA_DEBUG
+		Neighbor4.rg = float2(0.0, 0.0);
+		NeighborMin.rg = float2(0.0, 0.0);
+		NeighborMax.rg = float2(0.0, 0.0);
+		Filtered.rg = float2(0.0, 0.0);
+		FilteredLow.rg = float2(0.0, 0.0);
+		float DebugDiffCurrent = Filtered.b;
+	#endif
+	#if AA_GREEN_AS_LUMA
+		OutColor.rgb *= HdrWeightG(OutColor.rgb, InExposureScale);
+	#else
+		OutColor.rgb *= HdrWeight4(OutColor.rgb, InExposureScale);
+	#endif
+	#if AA_DEBUG
+		OutColor.rg = float2(0.0, 0.0);
+		float DebugDiffPrior = OutColor.b;
+	#endif
+
+	// FIND LUMA OF CLAMPED HISTORY
+	// ----------------------------
+	// Save off luma of history before the clamp.
+	#if AA_GREEN_AS_LUMA
+		float LumaMin = NeighborMin.g;
+		float LumaMax = NeighborMax.g;
+		float LumaHistory = OutColor.g;
+	#else
+		float LumaMin = Luma4(NeighborMin.rgb);
+		float LumaMax = Luma4(NeighborMax.rgb);
+		float LumaHistory = Luma4(OutColor.rgb);
+	#endif
+	float LumaContrast = LumaMax - LumaMin;
+	#if AA_AABB
+		// Clamp history, this uses color AABB intersection for tighter fit.
+		// Clamping works with the low pass (if available) to reduce flicker.
+		float ClampBlend = HistoryClamp(OutColor.rgb, FilteredLow.rgb, NeighborMin.rgb, NeighborMax.rgb);
+		#if AA_ALPHA
+			OutColor.rgb = lerp(OutColor.rgb, FilteredLow.rgb, ClampBlend);
+		#else
+			OutColor.rgba = lerp(OutColor.rgba, FilteredLow.rgba, ClampBlend);
+		#endif
+	#else
+		OutColor = clamp(OutColor, NeighborMin, NeighborMax);
+	#endif
+	#if AA_DEBUG
+		OutColor.rg = float2(0.0, 0.0);
+	#endif
+
+	// ADD BACK IN ALIASING TO SHARPEN
+    // -------------------------------
+	#if AA_FILTERED
+		// Blend in non-filtered based on the amount of sub-pixel motion.
+		float AddAliasing = saturate(HistoryBlur) * 0.5;
+		float LumaContrastFactor = 32.0;
+		#if AA_GREEN_AS_LUMA
+			// GREEN_AS_LUMA is 1/4 as bright.
+			LumaContrastFactor *= 4.0;
+		#endif
+		AddAliasing = saturate(AddAliasing + rcp(1.0 + LumaContrast * LumaContrastFactor));
+		Filtered.rgb = lerp(Filtered.rgb, Neighbor4.rgb, AddAliasing);
+	#endif
+	#if AA_GREEN_AS_LUMA
+		float LumaFiltered = Filtered.g;
+	#else
+		float LumaFiltered = Luma4(Filtered.rgb);
+	#endif
+		
+	// COMPUTE DYNAMIC MOTION BLEND FACTOR
+    // -----------------------------------
+	#if (AA_DYNAMIC == 1)
+		float BlendDynamicMotion = 0.0;
+	   	#if (AA_ALPHA == 1)
+			// Decay prior value.
+			OutColor.a *= 0.5;
+			BlendDynamicMotion = OutColor.a;
+		#endif
+		if(DynamicMotion) 
+		{
+			float BlendDynamicMotionNew = (1.0/2.0) * saturate(DoublePixelMotionDynamic * (1.0/8.0));
+			// Reduce dynamic motion factor if camera is also moving.
+			BlendDynamicMotionNew *= 1.0 - HistoryBlur;
+			#if (AA_ALPHA == 1)
+				// Feedback of 2.0 provides 1 extra frame of fast reconvergence.
+				// Feedback of 4.0 provides 2 extra frames of fast reconvergence.
+				// This enables fast reconverge after dynamic motion.
+				float FeedbackDynamic = 4.0;
+				OutColor.a = max(OutColor.a, BlendDynamicMotionNew * FeedbackDynamic);
+				BlendDynamicMotion = max(BlendDynamicMotion, BlendDynamicMotionNew);
+			#else
+				BlendDynamicMotion = BlendDynamicMotionNew;
+			#endif
+		}
+	#endif
+
+	// COMPUTE BLEND AMOUNT 
+    // --------------------
+	float HistoryAmount = (1.0/8.0) + HistoryBlur * (1.0/8.0);
+	float HistoryFactor = LumaHistory * HistoryAmount * (1.0 + HistoryBlur * HistoryAmount * 4.0);
+	float BlendFinal = saturate(HistoryFactor * rcp(LumaHistory + LumaContrast));
+	#if RESPONSIVE
+		// Responsive forces 1/4 of new frame.
+		BlendFinal = 1.0/4.0;
+	#endif
+	#if (AA_DYNAMIC == 1)
+		BlendFinal = max(BlendFinal, BlendDynamicMotion);
+		BlendFinal = min(1.0/2.0, BlendFinal);
+	#endif
+	// Offscreen feedback resets.
+	#if AA_LERP 
+		float FixedLerp = 1.0/float(AA_LERP);
+	#endif
+	if(OffScreen) 
+	{
+		OutColor = Filtered;
+		#if AA_ALPHA
+			OutColor.a = 0.0;
+		#endif
+		#if AA_LERP
+			FixedLerp = 1.0;
+		#endif
+	}
+	
+	// DO FINAL BLEND BETWEEN HISTORY AND FILTERED COLOR
+	// -------------------------------------------------
+	#if (AA_LERP == 0)
+		#if AA_ALPHA
+			// Blend in linear to hit luma target.
+			OutColor.rgb = lerp(OutColor.rgb, Filtered.rgb, BlendFinal);
+			#if RESPONSIVE
+				OutColor.a = max(OutColor.a, 1.0/2.0);
+			#endif
+		#else
+			OutColor = lerp(OutColor, Filtered, BlendFinal);
+			#if AA_FORCE_ALPHA_CLAMP
+				OutColor.a = clamp(OutColor.a, NeighborMin.a, NeighborMax.a);
+			#endif
+		#endif
+		// Convert back into linear.
+		#if AA_GREEN_AS_LUMA
+			OutColor.rgb *= HdrWeightInvG(OutColor.rgb, InExposureScale);
+		#else
+			OutColor.rgb *= HdrWeightInv4(OutColor.rgb, InExposureScale);
+		#endif
+	#else
+		OutColor = lerp(OutColor, Filtered, FixedLerp);
+		#if AA_GREEN_AS_LUMA
+			OutColor.rgb *= HdrWeightInvG(OutColor.rgb, InExposureScale);
+		#else
+			OutColor.rgb *= HdrWeightInv4(OutColor.rgb, InExposureScale);
+		#endif
+	#endif
+	#if AA_NAN
+		// Transform NaNs to black, transform negative colors to black.
+		OutColor.rgb = -min(-OutColor.rgb, 0.0);
+	#endif
+	#if AA_DEBUG
+		OutColor.g = abs(DebugDiffPrior - DebugDiffCurrent);
+		OutColor.r = OutColor.a;
+	#endif
+
+#else
+
 	// FIND MOTION OF PIXEL AND NEAREST IN NEIGHBORHOOD
     // ------------------------------------------------
 	float3 PosP; // Position of this pixel.
@@ -438,7 +835,7 @@ which can be configured via defines for various temporal AA passes.
 	#if AA_AABB
 		// Clamp history, this uses color AABB intersection for tighter fit.
 		// Clamping works with the low pass (if available) to reduce flicker.
-		float ClampBlend = HistoryClamp(OutColor.rgb, FilteredLow.rgb, NeighborMin.rgb, NeighborMax.rgb, InExposureScale);
+		float ClampBlend = HistoryClamp(OutColor.rgb, FilteredLow.rgb, NeighborMin.rgb, NeighborMax.rgb);
 		float4 Clamped = lerp(OutColor, FilteredLow, ClampBlend);
 	#else
 		float4 Clamped = clamp(OutColor, NeighborMin, NeighborMax);
@@ -590,6 +987,9 @@ which can be configured via defines for various temporal AA passes.
 		OutColor.r = OutColor.a;
 	#endif
 
+#endif
+
+
 #undef AA_AABB
 #undef AA_ALPHA
 #undef AA_CROSS
@@ -606,4 +1006,6 @@ which can be configured via defines for various temporal AA passes.
 #undef AA_NAN
 #undef AA_BORDER
 #undef AA_FORCE_ALPHA_CLAMP
-#undef AA_ONE_DYNAMIC_SAMPLE
\ No newline at end of file
+#undef AA_ONE_DYNAMIC_SAMPLE
+#undef AA_GREEN_AS_LUMA
+#undef AA_NEW
\ No newline at end of file
