Switch to decayed velocity weighting instead of old dynamic feedback. This has improvments for dithering also.

diff --git a/Engine/Shaders/PostProcessTemporalCommon.usf b/Engine/Shaders/PostProcessTemporalCommon.usf
--- a/Engine/Shaders/PostProcessTemporalCommon.usf
+++ b/Engine/Shaders/PostProcessTemporalCommon.usf
@@ -183,9 +183,6 @@ which can be configured via defines for various temporal AA passes.
 	BackP.y =  2.0 * ((PosP.y * ((CameraMotion[3].x * PosP.x) + (CameraMotion[3].y * PosP.z) + CameraMotion[3].z)) + (CameraMotion[3].w * PosP.x) + (CameraMotion[4].x * PosP.y * PosP.y) + (CameraMotion[4].y * PosP.z) + CameraMotion[4].z) * ScaleM;
 	// Convert BackP back into pixels.
 	float2 BackTemp = BackP * ViewportSize.xy;
-	// Save the amount of pixel offset of just camera motion, used later as the amount of blur introduced by history.
-	float HistoryBlurAmp = 2.0;
-	float HistoryBlur = saturate(abs(BackTemp.x) * HistoryBlurAmp + abs(BackTemp.y) * HistoryBlurAmp);
 	#if AA_CROSS
 		// Camera motion for neighborhood.
 		ScaleM = 1.0 / (dot(PosN, CameraMotion[0].xyz) + CameraMotion[0].w);
@@ -205,8 +202,6 @@ which can be configured via defines for various temporal AA passes.
 			}
 		#endif
 		BackTemp = BackP * ViewportSize.xy;
-		// At this point DoublePixelMotionDynamic is kept squared until later.
-		float DoublePixelMotionDynamic = dot(BackTemp, BackTemp);
 		#if AA_CROSS
 			// Dynamic motion of nearest in neighborhood.
 			float2 VelocityN = PostprocessInput3.SampleLevel(PostprocessInput3Sampler, UVAndScreenPos.xy + VelocityOffset, 0).xy;
@@ -215,21 +210,13 @@ which can be configured via defines for various temporal AA passes.
 			{
 				BackN = DecodeVelocityFromTexture(VelocityN);
 			}
-			#if (AA_ONE_DYNAMIC_SAMPLE == 0)
-				bool DynamicMotion = DynamicN || DynamicP;
-			#else
-				bool DynamicMotion = DynamicN;
-			#endif
 			BackTemp = BackN * ViewportSize.xy;
-			DoublePixelMotionDynamic = max(DoublePixelMotionDynamic, dot(BackTemp, BackTemp));
-		#else
-			bool DynamicMotion = DynamicP;
 		#endif
-		DoublePixelMotionDynamic = sqrt(DoublePixelMotionDynamic);
-	#else
-		bool DynamicMotion = false;
-		float DoublePixelMotionDynamic = 0.0;
 	#endif
+	// Save the amount of pixel offset of just camera motion, used later as the amount of blur introduced by history.
+	float HistoryBlurAmp = 2.0;
+	float HistoryBlur = saturate(abs(BackTemp.x) * HistoryBlurAmp + abs(BackTemp.y) * HistoryBlurAmp);
+	float Velocity = sqrt(dot(BackTemp, BackTemp));
 	#if (AA_CROSS == 0)
 		float2 BackN = BackP;
 	#endif
@@ -412,33 +399,6 @@ which can be configured via defines for various temporal AA passes.
 		float LumaFiltered = Luma4(Filtered.rgb);
 	#endif
 		
-	// COMPUTE DYNAMIC MOTION BLEND FACTOR
-    // -----------------------------------
-	#if (AA_DYNAMIC == 1)
-		float BlendDynamicMotion = 0.0;
-	   	#if (AA_ALPHA == 1)
-			// Decay prior value.
-			OutColor.a *= 0.5;
-			BlendDynamicMotion = OutColor.a;
-		#endif
-		if(DynamicMotion) 
-		{
-			float BlendDynamicMotionNew = (1.0/2.0) * saturate(DoublePixelMotionDynamic * (1.0/1.0));
-			// Reduce dynamic motion factor if camera is also moving.
-			BlendDynamicMotionNew *= 1.0 - HistoryBlur;
-			#if (AA_ALPHA == 1)
-				// Feedback of 2.0 provides 1 extra frame of fast reconvergence.
-				// Feedback of 4.0 provides 2 extra frames of fast reconvergence.
-				// This enables fast reconverge after dynamic motion.
-				float FeedbackDynamic = 4.0;
-				OutColor.a = max(OutColor.a, BlendDynamicMotionNew * FeedbackDynamic);
-				BlendDynamicMotion = max(BlendDynamicMotion, BlendDynamicMotionNew);
-			#else
-				BlendDynamicMotion = BlendDynamicMotionNew;
-			#endif
-		}
-	#endif
-
 	// COMPUTE BLEND AMOUNT 
     // --------------------
 	float HistoryAmount = (1.0/8.0) + HistoryBlur * (1.0/8.0);
@@ -448,10 +408,18 @@ which can be configured via defines for various temporal AA passes.
 		// Responsive forces 1/4 of new frame.
 		BlendFinal = 1.0/4.0;
 	#endif
-	#if (AA_DYNAMIC == 1)
-		BlendFinal = max(BlendFinal, BlendDynamicMotion);
+
+   	#if (AA_ALPHA == 1)
+		// Velocity weighting.
+		// Decay prior value.
+		float PriorVelocity = OutColor.a;
+		float VelocityDecay = 0.5;
+		OutColor.a = max(OutColor.a * VelocityDecay, Velocity * rcp(VelocityDecay));
+		float VelocityDiff = abs(PriorVelocity - Velocity) / max(1.0, max(PriorVelocity, Velocity));
+		BlendFinal = max(BlendFinal, VelocityDiff * (1.0/8.0));
 		BlendFinal = min(1.0/2.0, BlendFinal);
 	#endif
+
 	// Offscreen feedback resets.
 	#if AA_LERP 
 		float FixedLerp = 1.0/float(AA_LERP);
