statics not allowed in function body on some platforms

diff --git a/Engine/Shaders/HeightFogCommon.usf b/Engine/Shaders/HeightFogCommon.usf
--- a/Engine/Shaders/HeightFogCommon.usf
+++ b/Engine/Shaders/HeightFogCommon.usf
@@ -19,8 +19,8 @@ half4 GetExponentialHeightFog(float3 WorldPosition, float3 InCameraPosition)
 	/** FogDensity * exp2(-FogHeightFalloff * (CameraWorldPosition.z - FogHeight)) in x, FogHeightFalloff in y, CosTerminatorAngle in z, StartDistance in w. */
 	const float4 ExponentialFogParameters = SharedFogParameter0;
 	/** Color to use. */
-	static const half3 ExponentialFogColor = SharedFogParameter1.xyz;
-	static const half MinFogOpacity = SharedFogParameter1.w;
+	const half3 ExponentialFogColor = SharedFogParameter1.xyz;
+	const half MinFogOpacity = SharedFogParameter1.w;
 
 	float3 CameraToReceiver = WorldPosition - InCameraPosition;
 	float CameraToReceiverLengthSqr = dot( CameraToReceiver, CameraToReceiver );
