Per-instance transformations passed to shaders in primitive space

diff --git a/Engine/Shaders/LocalVertexFactory.usf b/Engine/Shaders/LocalVertexFactory.usf
--- a/Engine/Shaders/LocalVertexFactory.usf
+++ b/Engine/Shaders/LocalVertexFactory.usf
@@ -12,7 +12,6 @@ float4 LightMapCoordinateScaleBias;
 float4 ShadowMapCoordinateScaleBias;
 
 #if USE_INSTANCING
-float4 InstancedViewTranslation;
 float4 InstancingFadeOutParams;
 #endif
 
@@ -185,8 +184,8 @@ FMaterialVertexParameters GetMaterialVertexParameters( FVertexFactoryInput Input
 	Result.TangentToWorld = Intermediates.TangentToWorld;
 
 #if USE_INSTANCING
-	Result.InstanceLocalToWorld = transpose(float4x4(Input.InstanceTransform1, Input.InstanceTransform2, Input.InstanceTransform3, float4(0, 0, 0, 1)));
-	Result.InstanceWorldToLocal = float3x3(Input.InstanceInverseTransform1.xyz, Input.InstanceInverseTransform2.xyz, Input.InstanceInverseTransform3.xyz);
+	Result.InstanceLocalToWorld = mul(transpose(float4x4(Input.InstanceTransform1, Input.InstanceTransform2, Input.InstanceTransform3, float4(0, 0, 0, 1))), Primitive.LocalToWorld);
+	Result.InstanceWorldToLocal = mul(Primitive.WorldToLocal, transpose(float3x3(Input.InstanceInverseTransform1.xyz, Input.InstanceInverseTransform2.xyz, Input.InstanceInverseTransform3.xyz)));
 	Result.InstanceLocalPosition = Input.Position.xyz;
 	Result.PerInstanceParams = Intermediates.PerInstanceParams;
 #endif
@@ -315,8 +314,8 @@ float4 CalcWorldPosition(float4 Position)
 #endif
 {
 #if USE_INSTANCING
-	float4x4 InstanceTransform = float4x4(InstanceTransform1, InstanceTransform2, InstanceTransform3, float4(0, 0, 0, 1));
-	return mul(InstanceTransform, Position) + InstancedViewTranslation;
+	float4x4 InstanceTransform = transpose(float4x4(InstanceTransform1, InstanceTransform2, InstanceTransform3, float4(0, 0, 0, 1)));
+	return TransformLocalToTranslatedWorld(mul(Position, InstanceTransform));
 #elif USE_SPLINEDEFORM
 	// Make transform for this point along spline
 	float4x3 SliceTransform = CalcSliceTransform(dot(Position.xyz, SplineMeshDir));
@@ -363,8 +362,8 @@ half3x3 CalcTangentToLocal(FVertexFactoryInput Input)
 half3x3 CalcTangentToWorld(FVertexFactoryInput Input, half3x3 TangentToLocal)
 {
 #if USE_INSTANCING
-	half3x3 InstanceTransform = float3x3(Input.InstanceInverseTransform1.xyz, Input.InstanceInverseTransform2.xyz, Input.InstanceInverseTransform3.xyz);
-	half3x3 TangentToWorld = mul(TangentToLocal, InstanceTransform);
+	half3x3 InstanceTransform = transpose(float3x3(Input.InstanceInverseTransform1.xyz, Input.InstanceInverseTransform2.xyz, Input.InstanceInverseTransform3.xyz));
+	half3x3 TangentToWorld = CalcTangentToWorldNoScale(mul(TangentToLocal, InstanceTransform));
 #else
 	half3x3 TangentToWorld = CalcTangentToWorldNoScale(TangentToLocal);
 #endif
@@ -382,7 +381,7 @@ FVertexFactoryIntermediates GetVertexFactoryIntermediates(FVertexFactoryInput In
 	// Swizzle vertex color.
 	Intermediates.Color = Input.Color FCOLOR_COMPONENT_SWIZZLE;
 #if USE_INSTANCING
-	float3 InstanceLocation = float3(Input.InstanceTransform1.w, Input.InstanceTransform2.w, Input.InstanceTransform3.w) + InstancedViewTranslation.xyz;
+	float3 InstanceLocation = TransformLocalToTranslatedWorld(float3(Input.InstanceTransform1.w, Input.InstanceTransform2.w, Input.InstanceTransform3.w)).xyz;
 	// x = per-instance random, y = per-instance fade out factor
 	Intermediates.PerInstanceParams.x = Input.InstanceInverseTransform3.w;
 	Intermediates.PerInstanceParams.y = 1.0 - saturate((length(InstanceLocation) - InstancingFadeOutParams.x) * InstancingFadeOutParams.y);
@@ -495,9 +494,8 @@ float4 VertexFactoryGetWorldPosition(FPositionOnlyVertexFactoryInput Input)
 float4 VertexFactoryGetPreviousWorldPosition(FVertexFactoryInput Input, FVertexFactoryIntermediates Intermediates)
 {
 #if USE_INSTANCING
-	// @todo: THIS IS SO VERY WRONG, WE DON'T HAVE PRECALCED PreviousLocalToWorld for instance sigh
-	float4x4 InstanceTransform = float4x4(Input.InstanceTransform1, Input.InstanceTransform2, Input.InstanceTransform3, float4(0, 0, 0, 1));
-	return mul(InstanceTransform, Input.Position);
+	float4x4 InstanceTransform = transpose(float4x4(Input.InstanceTransform1, Input.InstanceTransform2, Input.InstanceTransform3, float4(0, 0, 0, 1)));
+	return mul(mul(Input.Position, InstanceTransform), PreviousLocalToWorld);
 #elif USE_SPLINEDEFORM
 	// Just like CalcWorldPosition...
 	float4x3 SliceTransform = CalcSliceTransform(dot(Input.Position.xyz, SplineMeshDir));
diff --git a/Engine/Source/Runtime/Engine/Classes/Components/InstancedStaticMeshComponent.h b/Engine/Source/Runtime/Engine/Classes/Components/InstancedStaticMeshComponent.h
--- a/Engine/Source/Runtime/Engine/Classes/Components/InstancedStaticMeshComponent.h
+++ b/Engine/Source/Runtime/Engine/Classes/Components/InstancedStaticMeshComponent.h
@@ -140,8 +140,7 @@ public:
 	virtual void GetStaticLightingInfo(FStaticLightingPrimitiveInfo& OutPrimitiveInfo,const TArray<ULightComponent*>& InRelevantLights,const FLightingBuildOptions& Options) override;
 #endif
 	virtual void GetLightAndShadowMapMemoryUsage( int32& LightMapMemoryUsage, int32& ShadowMapMemoryUsage ) const override;
-	virtual void ApplyWorldOffset(const FVector& InOffset, bool bWorldShift) override;
-
+	
 	virtual bool DoCustomNavigableGeometryExport(struct FNavigableGeometryExport* GeomExport) const override;
 	// End UPrimitiveComponent Interface
 
diff --git a/Engine/Source/Runtime/Engine/Private/InstancedStaticMesh.cpp b/Engine/Source/Runtime/Engine/Private/InstancedStaticMesh.cpp
--- a/Engine/Source/Runtime/Engine/Private/InstancedStaticMesh.cpp
+++ b/Engine/Source/Runtime/Engine/Private/InstancedStaticMesh.cpp
@@ -235,31 +235,22 @@ void FStaticMeshInstanceBuffer::Init(UInstancedStaticMeshComponent* InComponent,
 #endif
 		RawData.Add( FVector4( Instance.ShadowmapUVBias.X, Instance.ShadowmapUVBias.Y, Z, W ) );
 
-		// Grab the instance -> local matrix.  Every mesh instance has it's own transformation into
+		// Instance -> local matrix.  Every mesh instance has it's own transformation into
 		// the actor's coordinate space.
-		const FMatrix& InstanceToLocal = Instance.Transform;
-
-		// Create an instance -> world transform by combining the instance -> local transform with the
-		// local -> world transform
-		const FMatrix InstanceToWorld = InstanceToLocal * LocalToWorld;
-
-		// Instance to world transform matrix
 		{
-			const FMatrix Transpose = InstanceToWorld.GetTransposed();
+			const FMatrix Transpose = Instance.Transform.GetTransposed();
+						
 			RawData.Add( FVector4(Transpose.M[0][0], Transpose.M[0][1], Transpose.M[0][2], Transpose.M[0][3]) );
 			RawData.Add( FVector4(Transpose.M[1][0], Transpose.M[1][1], Transpose.M[1][2], Transpose.M[1][3]) );
 			RawData.Add( FVector4(Transpose.M[2][0], Transpose.M[2][1], Transpose.M[2][2], Transpose.M[2][3]) );
 		}
 
-		// World to instance rotation matrix (3x3)
+		// Instance -> local rotation matrix (3x3)
 		{
-			// Invert the instance -> world matrix
-			const FMatrix WorldToInstance = InstanceToWorld.Inverse();
-
 			const float RandomInstanceID = RandomInstanceIDBase + RandomStream.GetFraction() * RandomInstanceIDRange;
-
 			// hide the offset (bias) of the lightmap and the per-instance random id in the matrix's w
-			const FMatrix Transpose = WorldToInstance.GetTransposed();
+			const FMatrix Transpose = Instance.Transform.Inverse().GetTransposed();
+			
 			RawData.Add( FVector4(Transpose.M[0][0], Transpose.M[0][1], Transpose.M[0][2], Instance.LightmapUVBias.X) );
 			RawData.Add( FVector4(Transpose.M[1][0], Transpose.M[1][1], Transpose.M[1][2], Instance.LightmapUVBias.Y) );
 			RawData.Add( FVector4(Transpose.M[2][0], Transpose.M[2][1], Transpose.M[2][2], RandomInstanceID) );
@@ -509,23 +500,11 @@ class FInstancedStaticMeshVertexFactoryShaderParameters : public FVertexFactoryS
 {
 	virtual void Bind(const FShaderParameterMap& ParameterMap) override
 	{
-		InstancedViewTranslationParameter.Bind(ParameterMap,TEXT("InstancedViewTranslation"));
 		InstancingFadeOutParamsParameter.Bind(ParameterMap,TEXT("InstancingFadeOutParams"));
 	}
 
 	virtual void SetMesh(FRHICommandList& RHICmdList, FShader* VertexShader,const class FVertexFactory* VertexFactory,const class FSceneView& View,const struct FMeshBatchElement& BatchElement,uint32 DataFlags) const override
 	{
-		if( InstancedViewTranslationParameter.IsBound() )
-		{
-			FVector4 InstancedViewTranslation(View.ViewMatrices.PreViewTranslation, 0.f);
-			SetShaderValue(
-				RHICmdList, 
-				VertexShader->GetVertexShader(),
-				InstancedViewTranslationParameter,
-				InstancedViewTranslation
-				);
-		}
-
 		if( InstancingFadeOutParamsParameter.IsBound() )
 		{
 			FVector4 InstancingFadeOutParams(0.f,0.f,1.f,1.f);
@@ -569,14 +548,12 @@ class FInstancedStaticMeshVertexFactoryShaderParameters : public FVertexFactoryS
 
 	void Serialize(FArchive& Ar)
 	{
-		Ar << InstancedViewTranslationParameter;
 		Ar << InstancingFadeOutParamsParameter;
 	}
 
 	virtual uint32 GetSize() const { return sizeof(*this); }
 
 private:
-	FShaderParameter InstancedViewTranslationParameter;
 	FShaderParameter InstancingFadeOutParamsParameter;
 };
 
@@ -1457,13 +1434,6 @@ void UInstancedStaticMeshComponent::SetupNewInstanceData(FInstancedStaticMeshIns
 	}
 }
 
-void UInstancedStaticMeshComponent::ApplyWorldOffset(const FVector& InOffset, bool bWorldShift)
-{
-	Super::ApplyWorldOffset(InOffset, bWorldShift);
-	// We still need to recreate buffers with instance data, as they holds absolute transformations right now
-	MarkRenderStateDirty();
-}
-
 bool UInstancedStaticMeshComponent::DoCustomNavigableGeometryExport(struct FNavigableGeometryExport* GeomExport) const
 {
 	if (StaticMesh != NULL)
