Hopefully final fix for overbright tonemapping when min==max brightness.

diff --git a/Engine/Shaders/PostProcessEyeAdaptation.usf b/Engine/Shaders/PostProcessEyeAdaptation.usf
--- a/Engine/Shaders/PostProcessEyeAdaptation.usf
+++ b/Engine/Shaders/PostProcessEyeAdaptation.usf
@@ -32,7 +32,7 @@ void MainPS(float4 UVAndScreenPos : TEXCOORD0, out float OutColor : SV_Target0)
 	
 	float TargetExposure = ComputeEyeAdaptationExposure(InputNew0);
 	float OldExposureScale = InputNew0.Load(int3(0, 1, 0)).x;
-	float OldExposure = ExposureOffsetMultipler / (isfinite( OldExposureScale ) ? OldExposureScale : 1.0f);
+	float OldExposure = ExposureOffsetMultipler / ( OldExposureScale != 0 ? OldExposureScale : 1.0f );
 	float FrameTime = EyeAdaptationParams[1].y;
 
 	// eye adaptation changes over time
