Fix PS4 shader compiling

diff --git a/Engine/Shaders/LocalVertexFactory.usf b/Engine/Shaders/LocalVertexFactory.usf
--- a/Engine/Shaders/LocalVertexFactory.usf
+++ b/Engine/Shaders/LocalVertexFactory.usf
@@ -185,7 +185,7 @@ FMaterialVertexParameters GetMaterialVertexParameters( FVertexFactoryInput Input
 
 #if USE_INSTANCING
 	Result.InstanceLocalToWorld = mul(transpose(float4x4(Input.InstanceTransform1, Input.InstanceTransform2, Input.InstanceTransform3, float4(0, 0, 0, 1))), Primitive.LocalToWorld);
-	Result.InstanceWorldToLocal = mul(Primitive.WorldToLocal, transpose(float3x3(Input.InstanceInverseTransform1.xyz, Input.InstanceInverseTransform2.xyz, Input.InstanceInverseTransform3.xyz)));
+	Result.InstanceWorldToLocal = mul(float3x3(Primitive.WorldToLocal[0].xyz,Primitive.WorldToLocal[1].xyz,Primitive.WorldToLocal[2].xyz), transpose(float3x3(Input.InstanceInverseTransform1.xyz, Input.InstanceInverseTransform2.xyz, Input.InstanceInverseTransform3.xyz)));
 	Result.InstanceLocalPosition = Input.Position.xyz;
 	Result.PerInstanceParams = Intermediates.PerInstanceParams;
 #endif
