give access to particle world position, velocity, and relative time to the vertex shader even on <SM5; fixes TTP 337648 on Mac and DX10

diff --git a/Engine/Shaders/MeshParticleVertexFactory.usf b/Engine/Shaders/MeshParticleVertexFactory.usf
--- a/Engine/Shaders/MeshParticleVertexFactory.usf
+++ b/Engine/Shaders/MeshParticleVertexFactory.usf
@@ -120,7 +120,6 @@ struct FVertexFactoryIntermediates
 	float4 DynamicParameter;
 #endif	//USE_DYNAMIC_PARAMETERS
 
-#if USE_PARTICLE_NODES
 	/** The velocity of the particle, XYZ: direction, W: speed. */
 	float4 ParticleVelocity;
 
@@ -129,8 +128,6 @@ struct FVertexFactoryIntermediates
 	
 	/** Relative time. */
 	float RelativeTime;
-#endif //USE_PARTICLE_NODES
-
 };
 
 /** Converts from vertex factory specific interpolants to a FMaterialPixelParameters, which is used by material inputs. */
@@ -212,12 +209,10 @@ FMaterialVertexParameters GetMaterialVertexParameters( FVertexFactoryInput Input
 	Result.WorldPosition = WorldPosition - View.PreViewTranslation;
 	Result.VertexColor = Intermediates.VertexColor;
 
-#if USE_PARTICLE_NODES
 	Result.Particle.PositionAndSize.xyz = Intermediates.ParticleWorldPosition.xyz;
 	Result.Particle.PositionAndSize.w = 1;
 	Result.Particle.Velocity = Intermediates.ParticleVelocity;
 	Result.Particle.RelativeTime = Intermediates.RelativeTime;
-#endif
 
 #if USE_DYNAMIC_PARAMETERS
 	Result.Particle.DynamicParameter = Intermediates.DynamicParameter;
@@ -269,10 +264,8 @@ FVertexFactoryIntermediates GetVertexFactoryIntermediates(FVertexFactoryInput In
 	Intermediates.VertexColor = Input.VertexColor FCOLOR_COMPONENT_SWIZZLE;
 
 	// World position.
-#if USE_PARTICLE_NODES
 	float3 ParticleWorldPosition = float3( Input.Transform1.w, Input.Transform2.w, Input.Transform3.w);
 	Intermediates.ParticleWorldPosition = ParticleWorldPosition;
-#endif
 
 #if PARTICLE_MESH_INSTANCED
 	Intermediates.ParticleColor = Input.ParticleColor;
@@ -307,11 +300,9 @@ FVertexFactoryIntermediates GetVertexFactoryIntermediates(FVertexFactoryInput In
 	#endif
 #endif	//USE_DYNAMIC_PARAMETERS
 
-#if USE_PARTICLE_NODES
 	Intermediates.ParticleWorldPosition = ParticleWorldPosition.xyz;
 	Intermediates.ParticleVelocity = Input.ParticleVelocity;
 	Intermediates.RelativeTime = Input.SubUVLerpAndRelTime.y;
-#endif
 
 	return Intermediates;
 }
