Submitting Graeme.Thornton's fix for mesh shaders not compiling in ES2 mode.

diff --git a/Engine/Shaders/MeshParticleVertexFactory.usf b/Engine/Shaders/MeshParticleVertexFactory.usf
--- a/Engine/Shaders/MeshParticleVertexFactory.usf
+++ b/Engine/Shaders/MeshParticleVertexFactory.usf
@@ -264,7 +264,12 @@ FVertexFactoryIntermediates GetVertexFactoryIntermediates(FVertexFactoryInput In
 	Intermediates.VertexColor = Input.VertexColor FCOLOR_COMPONENT_SWIZZLE;
 
 	// World position.
+#if PARTICLE_MESH_INSTANCED
 	float3 ParticleWorldPosition = float3( Input.Transform1.w, Input.Transform2.w, Input.Transform3.w);
+#else
+	float3 ParticleWorldPosition = float3( Transform1.w, Transform2.w, Transform3.w);
+#endif
+
 	Intermediates.ParticleWorldPosition = ParticleWorldPosition;
 
 #if PARTICLE_MESH_INSTANCED
@@ -301,8 +306,14 @@ FVertexFactoryIntermediates GetVertexFactoryIntermediates(FVertexFactoryInput In
 #endif	//USE_DYNAMIC_PARAMETERS
 
 	Intermediates.ParticleWorldPosition = ParticleWorldPosition.xyz;
+
+#if PARTICLE_MESH_INSTANCED
 	Intermediates.ParticleVelocity = Input.ParticleVelocity;
 	Intermediates.RelativeTime = Input.SubUVLerpAndRelTime.y;
+#else
+	Intermediates.ParticleVelocity = float4(0,1,0,0); // TODO: We don't have this data when instancing is disabled. Should we? 
+	Intermediates.RelativeTime = 0;
+#endif
 
 	return Intermediates;
 }
