* fixed TTP 338435 Regression: CONTENT EXAMPLES: Owen material is very light

diff --git a/Engine/Shaders/BasePassPixelShader.usf b/Engine/Shaders/BasePassPixelShader.usf
--- a/Engine/Shaders/BasePassPixelShader.usf
+++ b/Engine/Shaders/BasePassPixelShader.usf
@@ -553,7 +553,12 @@ void Main(
 
 	float MaterialAO = GetMaterialAmbientOcclusion(MaterialParameters);
 	float Roughness = GetMaterialRoughness(MaterialParameters);
-	float3 SubsurfaceColor = GetMaterialSubsurfaceColor(MaterialParameters);
+
+	// If we don't use this shading model the color should be black (don't geneate shader code for unused data, don't do indirectlighting cache lighting with this color).
+	float3 SubsurfaceColor = 0;
+#if MATERIAL_SHADINGMODEL_SUBSURFACE
+	SubsurfaceColor = GetMaterialSubsurfaceColor(MaterialParameters);
+#endif
 
 #if MATERIAL_FULLY_ROUGH
 	Roughness = 1;
