* more proper fix, the SS color is still used in the other shading mode

diff --git a/Engine/Shaders/BasePassPixelShader.usf b/Engine/Shaders/BasePassPixelShader.usf
--- a/Engine/Shaders/BasePassPixelShader.usf
+++ b/Engine/Shaders/BasePassPixelShader.usf
@@ -554,9 +554,9 @@ void Main(
 	float MaterialAO = GetMaterialAmbientOcclusion(MaterialParameters);
 	float Roughness = GetMaterialRoughness(MaterialParameters);
 
-	// If we don't use this shading model the color should be black (don't geneate shader code for unused data, don't do indirectlighting cache lighting with this color).
+	// If we don't use this shading model the color should be black (don't generate shader code for unused data, don't do indirectlighting cache lighting with this color).
 	float3 SubsurfaceColor = 0;
-#if MATERIAL_SHADINGMODEL_SUBSURFACE
+#if MATERIAL_SHADINGMODEL_SUBSURFACE || MATERIAL_SHADINGMODEL_PREINTEGRATED_SKIN
 	SubsurfaceColor = GetMaterialSubsurfaceColor(MaterialParameters);
 #endif
 
