Added clear coat shading model. Fixed incorrect reflection capture filtering.

Needs some clean up still. Standard deferred reflections not implemented yet. Tiled deferred reflections needs some optimization. ClearCoat and ClearCoatRoughness needs to be exposed to the material editor.

diff --git a/Engine/Shaders/BRDF.usf b/Engine/Shaders/BRDF.usf
--- a/Engine/Shaders/BRDF.usf
+++ b/Engine/Shaders/BRDF.usf
@@ -92,7 +92,7 @@ float D_GGX( float Roughness, float NoH )
 	float m = Roughness * Roughness;
 	float m2 = m * m;
 	float d = ( NoH * m2 - NoH ) * NoH + 1;	// 2 mad
-	return m2 / ( PI*d*d );					// 3 mul, 1 rcp
+	return m2 / ( PI*d*d );					// 4 mul, 1 rcp
 }
 
 // Anisotropic GGX
@@ -119,9 +119,9 @@ float Vis_Neumann( float NoV, float NoL )
 }
 
 // [Kelemen 2001, "A microfacet based coupled specular-matte brdf model with importance sampling"]
-float Vis_Kelemen( float3 L, float3 V )
+float Vis_Kelemen( float VoH )
 {
-	return 1 / ( 2 + 2 * dot(L, V) );
+	return rcp( 4 * VoH * VoH );
 }
 
 // Tuned to match behavior of Vis_Smith
@@ -155,11 +155,13 @@ float3 F_None( float3 SpecularColor )
 // [Lagarde 2012, "Spherical Gaussian approximation for Blinn-Phong, Phong and Fresnel"]
 float3 F_Schlick( float3 SpecularColor, float VoH )
 {
-	// Anything less than 2% is physically impossible and is instead considered to be shadowing 
-	return SpecularColor + ( saturate( 50.0 * SpecularColor.g ) - SpecularColor ) * exp2( (-5.55473 * VoH - 6.98316) * VoH );
-
+	float Fc = pow( 1 - VoH, 5 );							// 1 sub, 3 mul
 	//float Fc = exp2( (-5.55473 * VoH - 6.98316) * VoH );	// 1 mad, 1 mul, 1 exp
 	//return Fc + (1 - Fc) * SpecularColor;					// 1 add, 3 mad
+	
+	// Anything less than 2% is physically impossible and is instead considered to be shadowing
+	return saturate( 50.0 * SpecularColor.g ) * Fc + (1 - Fc) * SpecularColor;
+	
 }
 
 float3 F_Fresnel( float3 SpecularColor, float VoH )
@@ -194,14 +196,14 @@ float Distribution( float Roughness, float NoH )
 }
 
 // Vis = G / (4*NoL*NoV)
-float GeometricVisibility( float Roughness, float NoV, float NoL, float VoH, float3 L, float3 V )
+float GeometricVisibility( float Roughness, float NoV, float NoL, float VoH )
 {
 #if   PHYSICAL_SPEC_G == 0
 	return Vis_Implicit();
 #elif PHYSICAL_SPEC_G == 1
 	return Vis_Neumann( NoV, NoL );
 #elif PHYSICAL_SPEC_G == 2
-	return Vis_Kelemen( L, V );
+	return Vis_Kelemen( VoH );
 #elif PHYSICAL_SPEC_G == 3
 	return Vis_Schlick( Roughness, NoV, NoL );
 #elif PHYSICAL_SPEC_G == 4
diff --git a/Engine/Shaders/BasePassPixelShader.usf b/Engine/Shaders/BasePassPixelShader.usf
--- a/Engine/Shaders/BasePassPixelShader.usf
+++ b/Engine/Shaders/BasePassPixelShader.usf
@@ -682,34 +682,72 @@ void Main(
 
 	#if FEATURE_LEVEL >= FEATURE_LEVEL_SM4 && (MATERIALBLENDING_SOLID || MATERIALBLENDING_MASKED)
 
-		FGBufferData Data = (FGBufferData)0;
+		FGBufferData GBuffer = (FGBufferData)0;
 
-		Data.WorldNormal = MaterialParameters.WorldNormal;
+		GBuffer.WorldNormal = MaterialParameters.WorldNormal;
 #if DIFFUSE_SPEC_INPUTS
-		Data.DiffuseColor = DiffuseColor;
-		Data.SpecularColor = SpecularColor;
+		GBuffer.DiffuseColor = DiffuseColor;
+		GBuffer.SpecularColor = SpecularColor;
 #else
-		Data.BaseColor = BaseColor;
-		Data.Metallic = Metallic;
-		Data.Specular = Specular;
+		GBuffer.BaseColor = BaseColor;
+		GBuffer.Metallic = Metallic;
+		GBuffer.Specular = Specular;
 #endif
-		Data.Roughness = Roughness;
-		Data.CustomData = EncodeSubsurfaceColor( SubsurfaceColor );
-		Data.IndirectIrradiance = IndirectIrradiance;
-		Data.PrecomputedShadowFactors = PrecomputedShadowFactors;
-		Data.GBufferAO = MaterialAO;
-		Data.Opacity = Opacity;
-		Data.DecalMask = Primitive.DecalReceiverMask;
-		Data.HasDistanceFieldRepresentation = Primitive.HasDistanceFieldRepresentation;
+		GBuffer.Roughness = Roughness;
+		GBuffer.IndirectIrradiance = IndirectIrradiance;
+		GBuffer.PrecomputedShadowFactors = PrecomputedShadowFactors;
+		GBuffer.GBufferAO = MaterialAO;
+		GBuffer.Opacity = Opacity;
+		GBuffer.DecalMask = Primitive.DecalReceiverMask;
+		GBuffer.HasDistanceFieldRepresentation = Primitive.HasDistanceFieldRepresentation;
 
 		#if MATERIAL_SHADINGMODEL_UNLIT
-			Data.ShadingModelID = SHADINGMODELID_UNLIT;
+			GBuffer.ShadingModelID = SHADINGMODELID_UNLIT;
 		#elif MATERIAL_SHADINGMODEL_DEFAULT_LIT
-			Data.ShadingModelID = SHADINGMODELID_DEFAULT_LIT;
+			GBuffer.ShadingModelID = SHADINGMODELID_DEFAULT_LIT;
 		#elif MATERIAL_SHADINGMODEL_SUBSURFACE
-			Data.ShadingModelID = SHADINGMODELID_SUBSURFACE;
+			GBuffer.ShadingModelID = SHADINGMODELID_SUBSURFACE;
+			GBuffer.CustomData = EncodeSubsurfaceColor( SubsurfaceColor );
 		#elif MATERIAL_SHADINGMODEL_PREINTEGRATED_SKIN
-			Data.ShadingModelID = SHADINGMODELID_PREINTEGRATED_SKIN;
+			GBuffer.ShadingModelID = SHADINGMODELID_PREINTEGRATED_SKIN;
+			GBuffer.CustomData = EncodeSubsurfaceColor( SubsurfaceColor );
+		#elif MATERIAL_SHADINGMODEL_CLEAR_COAT
+			GBuffer.ShadingModelID = SHADINGMODELID_CLEAR_COAT;
+
+			float ClearCoat = 1;
+			float ClearCoatRoughness = 0.1;
+			float Film = 1;
+			float MetalSpec = 0.9;
+
+			Film *= ClearCoat;
+
+			float NoV = saturate( dot( MaterialParameters.WorldNormal, MaterialParameters.CameraVector ) );
+
+			// Approximation of refraction's effect on EnvBRDF
+			float RefractionScale = saturate( (NoV - 1) * lerp( 1.3, 0.5, Roughness ) + 1 );
+
+#if DIFFUSE_SPEC_INPUTS
+			// Absorption doesn't work correctly with nonmetals with diffuse/spec
+			float3 AbsorptionColor = (1 - Film) + SpecularColor * ( Film * (1 / MetalSpec) );
+			GBuffer.SpecularColor = lerp( SpecularColor, MetalSpec, Film );
+#else
+			float3 AbsorptionColor = (1 - Film) + BaseColor * ( Film * (1 / MetalSpec) );
+			GBuffer.BaseColor = lerp( BaseColor, MetalSpec, Film );
+#endif
+
+			// Approximation of absorption integral, tuned for Roughness=0.4
+			float3 Absorption = AbsorptionColor * ( (NoV - 1) * 0.85 * ( 1 - lerp( AbsorptionColor, Square(AbsorptionColor), -0.78 ) ) + 1 );
+
+#if DIFFUSE_SPEC_INPUTS
+			GBuffer.DiffuseColor *= Absorption;
+			GBuffer.SpecularColor *= Absorption * RefractionScale;
+#else
+			GBuffer.BaseColor *= Absorption * lerp( 1, RefractionScale, Metallic );
+			GBuffer.Specular *= Luminance( Absorption ) * RefractionScale;
+#endif
+
+			GBuffer.CustomData.x = ClearCoat;
+			GBuffer.CustomData.y = ClearCoatRoughness;
 		#else
 			// missing shading model, compiler should report ShadingModelID is not set
 		#endif
@@ -719,7 +757,7 @@ void Main(
 		#endif
 
 		float QuantizationBias = PseudoRandom( MaterialParameters.SVPosition.xy ) * 0.5 - 0.5;
-		EncodeGBuffer(Data, OutGBufferA, OutGBufferB, OutGBufferC, OutGBufferD, OutGBufferE, QuantizationBias);
+		EncodeGBuffer(GBuffer, OutGBufferA, OutGBufferB, OutGBufferC, OutGBufferD, OutGBufferE, QuantizationBias);
 	#endif 
 	
 
diff --git a/Engine/Shaders/DeferredLightPixelShaders.usf b/Engine/Shaders/DeferredLightPixelShaders.usf
--- a/Engine/Shaders/DeferredLightPixelShaders.usf
+++ b/Engine/Shaders/DeferredLightPixelShaders.usf
@@ -31,7 +31,7 @@ FDeferredLightData SetupLightDataForStandardDeferred()
 	//@todo - permutation opportunity
 	LightData.bSpotLight = RADIAL_ATTENUATION;
 	//@todo - permutation opportunity
-	LightData.bShadowed = true;
+	LightData.bShadowed = DeferredLightUniforms.bShadowed;
 
 	return LightData;
 }
diff --git a/Engine/Shaders/DeferredLightingCommon.usf b/Engine/Shaders/DeferredLightingCommon.usf
--- a/Engine/Shaders/DeferredLightingCommon.usf
+++ b/Engine/Shaders/DeferredLightingCommon.usf
@@ -49,138 +49,249 @@ struct FSimpleDeferredLightData
 	bool bInverseSquared;
 };
 
-Texture2D		PreIntegratedBRDF;
-SamplerState	PreIntegratedBRDFSampler;
+#if 0
+void StandardShadingShared( float3 DiffuseColor, float3 SpecularColor, float Roughness, float3 V, half3 N )
+{
+	float NoV = saturate( dot(N, V) );
 
+	// Diffuse_Lambert
+	Shared.DiffuseMul = DiffuseColor * (1.0 / PI);
 
-#undef LIGHT_SOURCE_SHAPE
-#define LIGHT_SOURCE_SHAPE 1
+	// D_GGX, Vis_Schlick
+	float m = Roughness * Roughness;
+	Shared.m2 = m * m;
+	Shared.k = Square( Roughness ) * 0.5;
+	float Vis_SchlickV = NoV * (1 - Shared.k) + Shared.k;
+	Shared.SpecularMul = (0.25 / PI) * (Shared.m2 / Vis_SchlickV);
+	
+	// F_Schlick
+	Shared.SpecularMul *= saturate( 50.0 * SpecularColor.g );
+}
 
-// greatly reduces shadow mapping artifacts
-float BiasedNDotL(float NDotLWithoutSaturate )
+void StandardShadingPerLight( Shared, float3 L, float3 V, half3 N )
 {
-	return saturate(NDotLWithoutSaturate * 1.08f - 0.08f);
+	float3 H = normalize(V + L);			// 3 add, 2 mad, 4 mul, 1 rsqrt
+	float NoL = saturate( dot(N, L) );		// 2 mad, 1 mul
+	float NoH = saturate( dot(N, H) );		// 2 mad, 1 mul
+	float VoH = saturate( dot(V, H) );		// 2 mad, 1 mul
+
+	// D_GGX, Vis_Schlick
+	float d = ( NoH * Shared.m2 - NoH ) * NoH + 1;					// 2 mad
+	float Vis_SchlickL = NoL * (1 - Shared.k) + Shared.k;			// 1 mad, 1 add
+	float D_Vis = Shared.SpecularMul * rcp( d * d * Vis_SchlickL );	// 3 mul, 1 rcp
+
+	// F_Schlick
+	float Fc = pow( 1 - VoH, 5 );							// 1 sub, 3 mul
+	//float Fc = exp2( (-5.55473 * VoH - 6.98316) * VoH );	// 1 mad, 1 mul, 1 exp
+	float3 F = Fc + (1 - Fc) * SpecularColor;				// 1 add, 3 mad
+
+	return Shared.DiffuseMul + D_Vis * F;					// 3 mad
 }
+#endif
 
-
-// @param VectorToLight, L = normalized(VectorToLight)
-float3 PointLightDiffuse( FScreenSpaceData ScreenSpaceData, float3 VectorToLight, float3 V, half3 N )
+float3 StandardShading( FGBufferData GBuffer, float3 LobeRoughness, float3 LobeEnergy, float3 L, float3 V, half3 N )
 {
-	FGBufferData InGBufferData = ScreenSpaceData.GBuffer;
-
-	float3 L = normalize( VectorToLight );
-
 	float3 H = normalize(V + L);
 	float NoL = saturate( dot(N, L) );
 	float NoV = saturate( dot(N, V) );
+	float NoH = saturate( dot(N, H) );
 	float VoH = saturate( dot(V, H) );
+	
+	// Generalized microfacet specular
+	float D = D_GGX( LobeRoughness[1], NoH ) * LobeEnergy[1];
+	float Vis = Vis_Schlick( LobeRoughness[1], NoV, NoL );
+	float3 F = F_Schlick( GBuffer.SpecularColor, VoH );
 
-	return Diffuse( InGBufferData.DiffuseColor, InGBufferData.Roughness, NoV, NoL, VoH );
-} 
-
-float3 SimplePointLightDiffuse( FScreenSpaceData ScreenSpaceData )
-{
-	return Diffuse_Lambert(ScreenSpaceData.GBuffer.DiffuseColor);
+	return Diffuse_Lambert( GBuffer.DiffuseColor ) * LobeEnergy[2] + (D * Vis) * F;
 }
 
-// @param VectorToLight, L = normalized(VectorToLight)
-float3 PointLightSpecular( FScreenSpaceData ScreenSpaceData, FDeferredLightData LightData, float3 VectorToLight, float3 V, half3 N )
+float3 SimpleShading( FGBufferData GBuffer, float Roughness, float3 L, float3 V, half3 N )
 {
-	FGBufferData InGBufferData = ScreenSpaceData.GBuffer;
-	float Roughness = InGBufferData.Roughness;
-	float Energy = 1;
-
-	Roughness = max( Roughness, LightData.MinRoughness );
-	float a = Roughness * Roughness;
-	
-	const float SourceRadius = LightData.SpotAnglesAndSourceRadius.z;
-	const float SourceLength = LightData.SpotAnglesAndSourceRadius.w;
+	float3 H = normalize(V + L);
+	float NoH = saturate( dot(N, H) );
 	
-	float3 R = reflect( -V, N );
-	float RLengthL = rsqrt( dot( VectorToLight, VectorToLight ) );
+	// Generalized microfacet specular
+	float D = D_GGX( Roughness, NoH );
+	float Vis = Vis_Implicit();
+	float3 F = F_None( GBuffer.SpecularColor );
 
-	BRANCH
-	if( SourceLength > 0 )
-	{
-		// Energy conservation
-		// asin(x) is angle to sphere, atan(x) is angle to disk, saturate(x) is free and in the middle
-		float LineAngle = saturate( SourceLength * RLengthL );
-		Energy *= a / saturate( a + 0.5 * LineAngle );
+	return Diffuse_Lambert( GBuffer.DiffuseColor ) + (D * Vis) * F;
+}
 
-		// Closest point on line segment to ray
-		float3 Ld = LightData.LightDirection * SourceLength;
-		float3 L0 = VectorToLight - 0.5 * Ld;
-		float3 L1 = VectorToLight + 0.5 * Ld;
+float3 ClearCoatShading( FGBufferData GBuffer, float3 LobeRoughness, float3 LobeEnergy, float3 L, float3 V, half3 N )
+{
+	const float ClearCoat			= GBuffer.CustomData.x;
+	const float ClearCoatRoughness	= GBuffer.CustomData.y;
+	const float Film = 1 * ClearCoat;
+	const float MetalSpec = 0.9;
 
 #if 1
-		// Shortest distance
-		float a = Square( SourceLength );
-		float b = dot( R, Ld );
-		float t = saturate( dot( L0, b*R - Ld ) / (a - b*b) );
-#else
-		// Smallest angle
-		float A = Square( SourceLength );
-		float B = 2 * dot( L0, Ld );
-		float C = dot( L0, L0 );
-		float D = dot( R, L0 );
-		float E = dot( R, Ld );
-		float t = saturate( (B*D - 2*C*E) / (B*E - 2*A*D) );
-#endif
+	float3 H = normalize(V + L);
+	float NoL = saturate( dot(N, L) );
+	float NoV = saturate( dot(N, V) );
+	float NoH = saturate( dot(N, H) );
+	float VoH = saturate( dot(V, H) );
+	
+	// Generalized microfacet specular
+	float D = D_GGX( LobeRoughness[0], NoH ) * LobeEnergy[0];
+	float Vis = Vis_Kelemen( VoH );
+	
+	// F_Schlick
+	float F0 = GBuffer.Specular * 0.08;
+	float Fc = pow( 1 - VoH, 5 );
+	float F = Fc + (1 - Fc) * F0;
+	F *= ClearCoat;
 
-		VectorToLight = L0 + t * Ld;
-	}
+	float Fr1 = D * Vis * F;
 
-	BRANCH
-	if( SourceRadius > 0 )
-	{
-		// Energy conservation
-		// asin(x) is angle to sphere, atan(x) is angle to disk, saturate(x) is free and in the middle
-		float SphereAngle = saturate( SourceRadius * RLengthL );
-		Energy *= Square( a / saturate( a + 0.5 * SphereAngle ) );
-		
-		// Closest point on sphere to ray
-		float3 ClosestPointOnRay = dot( VectorToLight, R ) * R;
-		float3 CenterToRay = ClosestPointOnRay - VectorToLight;
-		float3 ClosestPointOnSphere = VectorToLight + CenterToRay * saturate( SourceRadius * rsqrt( dot( CenterToRay, CenterToRay ) ) );
-		VectorToLight = ClosestPointOnSphere;
-	}
+	float LayerAttenuation = (1 - F);
 
-	// normalized direction to light
-	float3 L = normalize( VectorToLight );
+	// Generalized microfacet specular
+	float D2 = D_GGX( LobeRoughness[1], NoH ) * LobeEnergy[1];
+	float Vis2 = Vis_Schlick( LobeRoughness[1], NoV, NoL );
+	//float3 F2 = F_Schlick( GBuffer.SpecularColor, VoH );
+	float3 F2 = saturate( 50.0 * GBuffer.SpecularColor.g ) * Fc + (1 - Fc) * GBuffer.SpecularColor;
 
+	float3 Fr2 = Diffuse_Lambert( GBuffer.DiffuseColor ) * LobeEnergy[2] + (D2 * Vis2) * F2;
+	
+	return Fr1 + Fr2 * LayerAttenuation;
+#else
 	float3 H = normalize(V + L);
 	float NoL = saturate( dot(N, L) );
 	float NoV = saturate( dot(N, V) );
 	float NoH = saturate( dot(N, H) );
 	float VoH = saturate( dot(V, H) );
 	
+	// Hard coded IOR of 1.5
+
 	// Generalized microfacet specular
-	float D = Distribution( Roughness, NoH );
-	float Vis = GeometricVisibility( Roughness, NoV, NoL, VoH, L, V );
-	float3 F = Fresnel( InGBufferData.SpecularColor, VoH );
+	float D = D_GGX( ClearCoatRoughness, NoH ) * LobeEnergy[0];
+	float Vis = Vis_Kelemen( VoH );
+	
+	// F_Schlick
+	float F0 = GBuffer.Specular * 0.08;
+	float Fc = pow( 1 - VoH, 5 );
+	float F = Fc + (1 - Fc) * F0;
+
+	float Fr1 = D * Vis * F;
+
+	// Refract rays
+	//float3 L2 = refract( -L, -H, 1 / 1.5 );
+	//float3 V2 = refract( -V, -H, 1 / 1.5 );
+
+	// LoH == VoH
+	//float RefractBlend = sqrt( 4 * VoH*VoH + 5 ) / 3 + 2.0 / 3 * VoH;
+	//float3 L2 = RefractBlend * H - L / 1.5;
+	//float3 V2 = RefractBlend * H - V / 1.5;
+	//float NoL2 = saturate( dot(N, L2) );
+	//float NoV2 = saturate( dot(N, V2) );
+	
+	// Approximation
+	float RefractBlend = (0.22 * VoH + 0.7) * VoH + 0.745;	// 2 mad
+	// Dot products distribute. No need for L2 and V2.
+	float RefractNoH = RefractBlend * NoH;					// 1 mul
+	float NoL2 = saturate( RefractNoH - (1 / 1.5) * NoL );	// 1 mad
+	float NoV2 = saturate( RefractNoH - (1 / 1.5) * NoV );	// 1 mad
+	// Should refract H too but unimportant
+
+	NoL2 = max( 0.001, NoL2 );
+	NoV2 = max( 0.001, NoV2 );
+
+	float  AbsorptionDist = rcp(NoV2) + rcp(NoL2);
+	float3 Absorption = pow( AbsorptionColor, 0.5 * AbsorptionDist );
+
+	// Approximation
+	//float  AbsorptionDist = ( NoV2 + NoL2 ) / ( NoV2 * NoL2 );
+	//float3 Absorption = AbsorptionColor * ( AbsorptionColor * (AbsorptionDist * 0.5 - 1) + (2 - 0.5 * AbsorptionDist) );
+	//float3 Absorption = AbsorptionColor + AbsorptionColor * (AbsorptionColor - 1) * (AbsorptionDist * 0.5 - 1);	// use for shared version
+	
+	//float F21 = Fresnel( 1 / 1.5, saturate( dot(V2, H) ) );
+	//float TotalInternalReflection = 1 - F21 * G_Schlick( Roughness, NoV2, NoL2 );
+	//float3 LayerAttenuation = ( (1 - F12) * TotalInternalReflection ) * Absorption;
 
-	return (Energy * D * Vis) * F;
+	// Approximation
+	float3 LayerAttenuation = (1 - F) * Absorption;
+
+	// Approximation for IOR == 1.5
+	//SpecularColor = ChangeBaseMedium( SpecularColor, 1.5 );
+	//SpecularColor = saturate( ( 0.55 * SpecularColor + (0.45 * 1.08) ) * SpecularColor - (0.45 * 0.08) );
+	// Treat SpecularColor as relative to IOR. Artist compensates.
+	
+	// Generalized microfacet specular
+	float D2 = D_GGX( Roughness, NoH ) * LobeEnergy[2];
+	float Vis2 = Vis_Schlick( Roughness, NoV2, NoL2 );
+	float3 F2 = F_Schlick( GBuffer.SpecularColor, VoH );
+
+	float3 Fr2 = Diffuse_Lambert( GBuffer.DiffuseColor ) * LobeEnergy[2] + (D2 * Vis2) * F2;
+	
+	return Fr1 + Fr2 * LayerAttenuation;
+#endif
 }
 
-float3 SimplePointLightSpecular( FScreenSpaceData ScreenSpaceData, float3 UnitL, float3 V, half3 N )
+float3 SubsurfaceShadingSubsurface( FGBufferData GBuffer, float3 L, float3 V, half3 N )
 {
-	FGBufferData InGBufferData = ScreenSpaceData.GBuffer;
-	float Roughness = InGBufferData.Roughness;
+	float3 SubsurfaceColor = DecodeSubsurfaceColor( GBuffer.CustomData );
 
-	// TODO move outside tile loop
-	Roughness = max( 0.08, Roughness );
+	float3 H = normalize(V + L);
 
-	float3 H = normalize(V + UnitL);
-	float NoH = saturate( dot(N, H) );
-	
-	// Generalized microfacet specular
-	float  D = D_GGX( Roughness, NoH );
-	float  Vis = Vis_Implicit();
-	float3 F = F_None( InGBufferData.SpecularColor );
+	// to get an effect when you see through the material
+	// hard coded pow constant
+	float InScatter = pow(saturate(dot(L, -V)), 12) * lerp(3, .1f, GBuffer.Opacity);
+	// wrap around lighting, /(PI*2) to be energy consistent (hack do get some view dependnt and light dependent effect)
+	float OpacityFactor = GBuffer.Opacity;
+	// Opacity of 0 gives no normal dependent lighting, Opacity of 1 gives strong normal contribution
+	float NormalContribution = saturate(dot(N, H) * OpacityFactor + 1 - OpacityFactor);
+	float BackScatter = GBuffer.GBufferAO * NormalContribution / (PI * 2);
+
+	// lerp to never exceed 1 (energy conserving)
+	return SubsurfaceColor * lerp(BackScatter, 1, InScatter);
+}
 
-	return (D * Vis) * F;
+Texture2D		PreIntegratedBRDF;
+SamplerState	PreIntegratedBRDFSampler;
+
+float3 SubsurfaceShadingPreintegratedSkin( FGBufferData GBuffer, float3 L, float3 V, half3 N )
+{
+	float3 SubsurfaceColor = DecodeSubsurfaceColor( GBuffer.CustomData );
+
+	float OpacityFactor = GBuffer.Opacity;
+	float3 PreintegratedBRDF = Texture2DSampleLevel(PreIntegratedBRDF, PreIntegratedBRDFSampler, float2(saturate(dot(N, L) * .5 + .5), 1 - OpacityFactor), 0).rgb;
+	return PreintegratedBRDF * SubsurfaceColor;
+}
+
+float3 SurfaceShading( FGBufferData GBuffer, float3 LobeRoughness, float3 LobeEnergy, float3 L, float3 V, half3 N )
+{
+	switch( GBuffer.ShadingModelID )
+	{
+		case SHADINGMODELID_UNLIT:
+		case SHADINGMODELID_DEFAULT_LIT:
+		case SHADINGMODELID_SUBSURFACE:
+		case SHADINGMODELID_PREINTEGRATED_SKIN:
+			return StandardShading( GBuffer, LobeRoughness, LobeEnergy, L, V, N );
+		case SHADINGMODELID_CLEAR_COAT:
+			return ClearCoatShading( GBuffer, LobeRoughness, LobeEnergy, L, V, N );
+		default:
+			return 0;
+	}
 }
 
+float3 SubsurfaceShading( FGBufferData GBuffer, float3 L, float3 V, half3 N )
+{
+	switch( GBuffer.ShadingModelID )
+	{
+		case SHADINGMODELID_SUBSURFACE:
+			return SubsurfaceShadingSubsurface( GBuffer, L, V, N );
+		case SHADINGMODELID_PREINTEGRATED_SKIN:
+			return SubsurfaceShadingPreintegratedSkin( GBuffer, L, V, N );
+		default:
+			return 0;
+	}
+}
+
+
+#undef LIGHT_SOURCE_SHAPE
+#define LIGHT_SOURCE_SHAPE 1
+
 bool RayHitSphere( float3 R, float3 SphereCenter, float SphereRadius )
 {
 	float3 ClosestPointOnRay = max( 0, dot( SphereCenter, R ) ) * R;
@@ -212,7 +323,7 @@ float3 PointLightSpecularMIS( FScreenSpaceData ScreenSpaceData, FDeferredLightDa
 	const float3 ConeAxis = normalize( LightCenter );
 	const float ConeCos = sqrt( 1 - Square( SourceRadius ) / DistanceSqr );
 
-	const float SampleColor = (16.0/PI) / Square(SourceRadius);
+	const float SampleColor = (1.0/PI) / Square(SourceRadius);
 
 	float3 SpecularLighting = 0;
 
@@ -225,13 +336,13 @@ float3 PointLightSpecularMIS( FScreenSpaceData ScreenSpaceData, FDeferredLightDa
 		float3 L, H;
 		if( bSampleBRDF )
 		{
-			float2 E = Hammersley( i, NumSamplesCone, Random );
+			float2 E = Hammersley( i, NumSamplesGGX, Random );
 			H = TangentToWorld( ImportanceSampleGGX( E, Roughness ).xyz, N );
 			L = 2 * dot( V, H ) * H - V;
 		}
 		else
 		{
-			float2 E = Hammersley( i, NumSamplesGGX, Random );
+			float2 E = Hammersley( i, NumSamplesCone, Random );
 			L = TangentToWorld( UniformSampleCone( E, ConeCos ).xyz, ConeAxis );
 			H = normalize(V + L);
 		}
@@ -249,7 +360,7 @@ float3 PointLightSpecularMIS( FScreenSpaceData ScreenSpaceData, FDeferredLightDa
 			
 			// Generalized microfacet specular
 			float D = D_GGX( Roughness, NoH );
-			float Vis = GeometricVisibility( Roughness, NoV, NoL, VoH, L, V );
+			float Vis = GeometricVisibility( Roughness, NoV, NoL, VoH );
 			float3 F = Fresnel( GBuffer.SpecularColor, VoH );
 
 			float ConePDF = 1.0 / ( 2 * PI * (1 - ConeCos) );
@@ -271,38 +382,71 @@ float3 PointLightSpecularMIS( FScreenSpaceData ScreenSpaceData, FDeferredLightDa
 	return SpecularLighting / (NumSamplesGGX + NumSamplesCone);
 }
 
-float3 PointLightSubsurface( FScreenSpaceData ScreenSpaceData, uint ShadingModelID, float3 L, float3 V, half3 N, float SubsurfaceExtinction )
+// Find representative incoming light direction and energy modification
+float3 AreaLightSpecular( FDeferredLightData LightData, inout float3 LobeRoughness, inout float3 ToLight, inout float3 L, float3 V, half3 N )
 {
-	FGBufferData InGBufferData = ScreenSpaceData.GBuffer;
-	float3 SubsurfaceColor = DecodeSubsurfaceColor( InGBufferData.CustomData );
+	float3 LobeEnergy = 1;
+
+	LobeRoughness = max( LobeRoughness, LightData.MinRoughness );
+	float3 m = LobeRoughness * LobeRoughness;
+	
+	const float SourceRadius = LightData.SpotAnglesAndSourceRadius.z;
+	const float SourceLength = LightData.SpotAnglesAndSourceRadius.w;
+
+	// TODO early out for point lights
+	
+	float3 R = reflect( -V, N );
+	float InvDistToLight = rsqrt( dot( ToLight, ToLight ) );
 
-	if (ShadingModelID == SHADINGMODELID_SUBSURFACE)
+	BRANCH
+	if( SourceLength > 0 )
 	{
-		L = normalize( L );
-		float3 H = normalize(V + L);
-
-		// to get an effect when you see through the material
-		// hard coded pow constant
-		float InScatter = pow(saturate(dot(L, -V)), 12) * lerp(3, .1f, InGBufferData.Opacity);
-		// wrap around lighting, /(PI*2) to be energy consistent (hack do get some view dependnt and light dependent effect)
-		float OpacityFactor = InGBufferData.Opacity;
-		// Opacity of 0 gives no normal dependent lighting, Opacity of 1 gives strong normal contribution
-		float NormalContribution = saturate(dot(N, H) * OpacityFactor + 1 - OpacityFactor);
-		float BackScatter = InGBufferData.GBufferAO * NormalContribution / (PI * 2);
-
-		// lerp to never exceed 1 (energy conserving)
-		return SubsurfaceColor * (lerp(BackScatter, 1, InScatter) * SubsurfaceExtinction);
+		// Energy conservation
+		// asin(x) is angle to sphere, atan(x) is angle to disk, saturate(x) is free and in the middle
+		float LineAngle = saturate( SourceLength * InvDistToLight );
+		LobeEnergy *= m / saturate( m + 0.5 * LineAngle );
+
+		// Closest point on line segment to ray
+		float3 L01 = LightData.LightDirection * SourceLength;
+		float3 L0 = ToLight - 0.5 * L01;
+		float3 L1 = ToLight + 0.5 * L01;
+
+#if 1
+		// Shortest distance
+		float a = Square( SourceLength );
+		float b = dot( R, L01 );
+		float t = saturate( dot( L0, b*R - L01 ) / (a - b*b) );
+#else
+		// Smallest angle
+		float A = Square( SourceLength );
+		float B = 2 * dot( L0, L01 );
+		float C = dot( L0, L0 );
+		float D = dot( R, L0 );
+		float E = dot( R, L01 );
+		float t = saturate( (B*D - 2*C*E) / (B*E - 2*A*D) );
+#endif
+
+		ToLight = L0 + t * L01;
 	}
-	else if (ShadingModelID == SHADINGMODELID_PREINTEGRATED_SKIN)
-	{
-		L = normalize( L );
 
-		float OpacityFactor = InGBufferData.Opacity;
-		float3 PreintegratedBRDF = Texture2DSampleLevel(PreIntegratedBRDF, PreIntegratedBRDFSampler, float2(saturate(dot(N, L) * .5 + .5), 1 - OpacityFactor), 0).rgb;
-		return PreintegratedBRDF * SubsurfaceColor * SubsurfaceExtinction;
+	BRANCH
+	if( SourceRadius > 0 )
+	{
+		// Energy conservation
+		// asin(x) is angle to sphere, atan(x) is angle to disk, saturate(x) is free and in the middle
+		float SphereAngle = saturate( SourceRadius * InvDistToLight );
+		LobeEnergy *= Square( m / saturate( m + 0.5 * SphereAngle ) );
+		
+		// Closest point on sphere to ray
+		float3 ClosestPointOnRay = dot( ToLight, R ) * R;
+		float3 CenterToRay = ClosestPointOnRay - ToLight;
+		float3 ClosestPointOnSphere = ToLight + CenterToRay * saturate( SourceRadius * rsqrt( dot( CenterToRay, CenterToRay ) ) );
+		ToLight = ClosestPointOnSphere;
 	}
 
-	return 0;
+	L = normalize( ToLight );
+
+	return LobeEnergy;
 }
 
 /** Returns 0 for positions closer than the fade near distance from the camera, and 1 for positions further than the fade far distance. */
@@ -353,30 +497,33 @@ void GetShadowTerms(FScreenSpaceData ScreenSpaceData, FDeferredLightData LightDa
 /** Calculates lighting for a given position, normal, etc with a fully featured lighting model designed for quality. */
 float4 GetDynamicLighting(float3 WorldPosition, float3 CameraVector, float2 InUV, FScreenSpaceData ScreenSpaceData, uint ShadingModelID, FDeferredLightData LightData, float4 LightAttenuation, uint2 Random)
 {
+	float3 V = -CameraVector;
 	float3 N = ScreenSpaceData.GBuffer.WorldNormal;
-	float3 L = LightData.LightDirection;
-	float NoL = BiasedNDotL( dot(N, L) );
+	float3 ToLight = LightData.LightDirection;
+	float3 L = ToLight;	// no need to normalize
+	float NoL = saturate( dot(N, L) );
 	float DistanceAttenuation = 1;
 	float LightRadiusMask = 1;
 	float SpotFalloff = 1;
 
 	if (LightData.bRadialLight)
 	{
-		L = LightData.LightPositionAndInvRadius.xyz - WorldPosition;
+		ToLight = LightData.LightPositionAndInvRadius.xyz - WorldPosition;
+		
+		float DistanceSqr = dot( ToLight, ToLight );
+		L = ToLight * rsqrt( DistanceSqr );
 
 		if (LightData.bInverseSquared)
 		{
-			float DistanceSqr = dot( L, L );
-
 			const float SourceLength = LightData.SpotAnglesAndSourceRadius.w;
-			
+
 			BRANCH
 			if( SourceLength > 0 )
 			{
 				// Line segment irradiance
-				float3 Ld = LightData.LightDirection * SourceLength;
-				float3 L0 = L - 0.5 * Ld;
-				float3 L1 = L + 0.5 * Ld;
+				float3 L01 = LightData.LightDirection * SourceLength;
+				float3 L0 = ToLight - 0.5 * L01;
+				float3 L1 = ToLight + 0.5 * L01;
 				float LengthL0 = length( L0 );
 				float LengthL1 = length( L1 );
 
@@ -387,22 +534,18 @@ float4 GetDynamicLighting(float3 WorldPosition, float3 CameraVector, float2 InUV
 			{
 				// Sphere irradiance (technically just 1/d^2 but this avoids inf)
 				DistanceAttenuation = 1 / ( DistanceSqr + 1 );
-				NoL = BiasedNDotL( dot( N, normalize(L) ) );
+				NoL = saturate( dot( N, L ) );
 			}
 
-			// TODO scale LightColor
-			// Correction for lumen units
-			DistanceAttenuation *= 16;
-
 			// TODO optimize
 			LightRadiusMask = Square( saturate( 1 - Square( DistanceSqr * Square(LightData.LightPositionAndInvRadius.w) ) ) );
 		}
 		else
 		{
 			DistanceAttenuation = 1;
-			NoL = BiasedNDotL( dot( N, normalize(L) ) );
+			NoL = saturate( dot( N, L ) );
 			
-			LightRadiusMask = RadialAttenuation(L * LightData.LightPositionAndInvRadius.w, LightData.LightColorAndFalloffExponent.w);	
+			LightRadiusMask = RadialAttenuation(ToLight * LightData.LightPositionAndInvRadius.w, LightData.LightColorAndFalloffExponent.w);	
 		}
 
 		if (LightData.bSpotLight)
@@ -420,21 +563,20 @@ float4 GetDynamicLighting(float3 WorldPosition, float3 CameraVector, float2 InUV
 	BRANCH
 	if (LightRadiusMask > 0 && SpotFalloff > 0)
 	{
-		float OpaqueShadowTerm = 1;
-		float SSSShadowTerm = 1;
+		float SurfaceShadow = 1;
+		float SubsurfaceShadow = 1;
 
+		BRANCH
 		if (LightData.bShadowed)
 		{
-			GetShadowTerms(ScreenSpaceData, LightData, WorldPosition, LightAttenuation, OpaqueShadowTerm, SSSShadowTerm);
+			GetShadowTerms(ScreenSpaceData, LightData, WorldPosition, LightAttenuation, SurfaceShadow, SubsurfaceShadow);
+
+			// greatly reduces shadow mapping artifacts
+			SurfaceShadow *= saturate(dot(N, L) * 6 - 0.2);
 		}
 
-		float NonShadowedAttenuation = DistanceAttenuation * LightRadiusMask * SpotFalloff;
-		float ShadowedAttenuation = NonShadowedAttenuation * OpaqueShadowTerm;
-
-#if LIGHTCULLING_FINE == 1
-		// cull behind light (saving shading computations, might not be faster on some hardware)
-		ShadowedAttenuation *= saturate(dot(L, N) * 100000);
-#endif
+		float SurfaceAttenuation	= (DistanceAttenuation * LightRadiusMask * SpotFalloff) * SurfaceShadow;
+		float SubsurfaceAttenuation	= (DistanceAttenuation * LightRadiusMask * SpotFalloff) * SubsurfaceShadow;
 
 		EstimatedCost += 0.3f;	// add the cost of getting the shadow terms
 
@@ -445,13 +587,20 @@ float4 GetDynamicLighting(float3 WorldPosition, float3 CameraVector, float2 InUV
 		{
 			const float3 LightColor = LightData.LightColorAndFalloffExponent.rgb;
 
-			float3 DiffuseLighting = PointLightDiffuse( ScreenSpaceData, L, -CameraVector, N );
-			float3 SpecularLighting = PointLightSpecular( ScreenSpaceData, LightData, L, -CameraVector, N );
-			float3 SubsurfaceLighting = PointLightSubsurface( ScreenSpaceData, ShadingModelID, L, -CameraVector, N, SSSShadowTerm );
+			float3 LobeRoughness = float3( 0.1, ScreenSpaceData.GBuffer.Roughness, 1 );
+			float3 LobeEnergy = AreaLightSpecular( LightData, LobeRoughness, ToLight, L, V, N );
+
+			float3 SurfaceLighting = SurfaceShading( ScreenSpaceData.GBuffer, LobeRoughness, LobeEnergy, L, V, N );
+			float3 SubsurfaceLighting = SubsurfaceShading( ScreenSpaceData.GBuffer, L, V, N );
 			
+#if 1
+			// TODO fix SSSSS skin
+			OutLighting.rgb += LightColor * ( (NoL * SurfaceAttenuation) * SurfaceLighting + SubsurfaceAttenuation * SubsurfaceLighting );
+#else
+
 #define RAY_TRACE 0
 #if RAY_TRACE
-			SpecularLighting = PointLightSpecularMIS( ScreenSpaceData, LightData, L, -CameraVector, N, Random );
+			SpecularLighting = PointLightSpecularMIS( ScreenSpaceData, LightData, ToLight, V, N, Random );
 #endif
 
 			float4 LightColor4 = float4(LightColor, Luminance(LightColor));
@@ -460,13 +609,15 @@ float4 GetDynamicLighting(float3 WorldPosition, float3 CameraVector, float2 InUV
 			float4 SubsurfaceLighting4 = float4(SubsurfaceLighting, 0);
 
 #if RAY_TRACE
-			OutLighting += LightColor4 * ( (NoL * ShadowedAttenuation) * DiffuseLighting4 + SubsurfaceLighting4 * NonShadowedAttenuation );
+			OutLighting += LightColor4 * ( (NoL * ShadowedAttenuation) * DiffuseLighting4 + SubsurfaceLighting4 * SubsurfaceAttenuation );
 			OutLighting += LightColor4 * SpecularLighting4 * (LightRadiusMask * SpotFalloff * OpaqueShadowTerm);
 #else
-			OutLighting += LightColor4 * ( (NoL * ShadowedAttenuation) * (DiffuseLighting4 + SpecularLighting4) + SubsurfaceLighting4 * NonShadowedAttenuation );
+			OutLighting += LightColor4 * ( (NoL * SurfaceAttenuation) * (DiffuseLighting4 + SpecularLighting4) + SubsurfaceLighting4 * SubsurfaceAttenuation );
+#endif
+
 #endif
 
-			EstimatedCost += 0.4f;	// add the cost of the lighting computations (should sum up to sum to 1)
+			EstimatedCost += 0.4f;	// add the cost of the lighting computations (should sum up to 1)
 		}
 	}
 
@@ -481,25 +632,23 @@ float4 GetDynamicLighting(float3 WorldPosition, float3 CameraVector, float2 InUV
 /** 
  * Calculates lighting for a given position, normal, etc with a simple lighting model designed for speed. 
  * All lights rendered through this method are unshadowed point lights with no shadowing or light function or IES.
- * A cheap Blinn specular is used instead of the more correct area specular, no fresnel.
+ * A cheap specular is used instead of the more correct area specular, no fresnel.
  */
 float3 GetSimpleDynamicLighting(float3 WorldPosition, float3 CameraVector, FScreenSpaceData ScreenSpaceData, FSimpleDeferredLightData LightData)
 {
+	float3 V = -CameraVector;
 	float3 N = ScreenSpaceData.GBuffer.WorldNormal;
-	float3 L = LightData.LightPositionAndInvRadius.xyz - WorldPosition;
-	float3 UnitL = normalize(L);
-	float NoL = BiasedNDotL( dot( N, UnitL ) );
+	float3 ToLight = LightData.LightPositionAndInvRadius.xyz - WorldPosition;
 	float DistanceAttenuation = 1;
+	
+	float DistanceSqr = dot( ToLight, ToLight );
+	float3 L = ToLight * rsqrt( DistanceSqr );
+	float NoL = saturate( dot( N, L ) );
 
 	if (LightData.bInverseSquared)
 	{
-		float DistanceSqr = dot( L, L );
-
 		// Sphere falloff (technically just 1/d2 but this avoids inf)
 		DistanceAttenuation = 1 / ( DistanceSqr + 1 );
-
-		// Correction for lumen units
-		DistanceAttenuation *= 16;
 	
 		float LightRadiusMask = Square( saturate( 1 - Square( DistanceSqr * Square(LightData.LightPositionAndInvRadius.w) ) ) );
 		DistanceAttenuation *= LightRadiusMask;
@@ -514,13 +663,12 @@ float3 GetSimpleDynamicLighting(float3 WorldPosition, float3 CameraVector, FScre
 	BRANCH
 	if (DistanceAttenuation > 0)
 	{
-		float3 DiffuseLighting = SimplePointLightDiffuse(ScreenSpaceData);
-		float3 SpecularLighting = SimplePointLightSpecular(ScreenSpaceData, UnitL, -CameraVector, N);
-		float3 CombinedAttenuation = DistanceAttenuation * LightData.LightColorAndFalloffExponent.rgb;
+		const float3 LightColor = LightData.LightColorAndFalloffExponent.rgb;
+
 		// Apply SSAO to the direct lighting since we're not going to have any other shadowing
-		CombinedAttenuation *= ScreenSpaceData.AmbientOcclusion;
+		float3 Attenuation = DistanceAttenuation * ScreenSpaceData.AmbientOcclusion;
 
-		OutLighting += NoL * (DiffuseLighting + SpecularLighting) * CombinedAttenuation;
+		OutLighting += LightColor * (NoL * Attenuation) * SimpleShading( ScreenSpaceData.GBuffer, ScreenSpaceData.GBuffer.Roughness, L, V, N );
 	}
 
 	return OutLighting;
diff --git a/Engine/Shaders/DeferredShadingCommon.usf b/Engine/Shaders/DeferredShadingCommon.usf
--- a/Engine/Shaders/DeferredShadingCommon.usf
+++ b/Engine/Shaders/DeferredShadingCommon.usf
@@ -257,7 +257,8 @@ uint DecodeShadingModelId(float InPackedChannel)
 #define SHADINGMODELID_DEFAULT_LIT			1
 #define SHADINGMODELID_SUBSURFACE			2
 #define SHADINGMODELID_PREINTEGRATED_SKIN	3
-#define SHADINGMODELID_NUM					4
+#define SHADINGMODELID_CLEAR_COAT			4
+#define SHADINGMODELID_NUM					5
 
 // all values that are output by the forward rendering pass
 struct FGBufferData
@@ -287,7 +288,7 @@ struct FGBufferData
 	float Opacity;
 	// 0..1 ambient occlusion  e.g.SSAO, wet surface mask, skylight mask, ...
 	float GBufferAO;
-	// 0..255 SHADINGMODELID_UNLIT / SHADINGMODELID_DEFAULT_LIT / SHADINGMODELID_SUBSURFACE / SHADINGMODELID_PREINTEGRATED_SKIN
+	// 0..255 
 	uint ShadingModelID;
 	// 0..1 decal receiver mask
 	uint DecalMask;
@@ -330,7 +331,7 @@ struct FScreenSpaceData
 
 /** Populates OutGBufferA, B and C */
 void EncodeGBuffer(
-	FGBufferData Data,
+	FGBufferData GBuffer,
 	out float4 OutGBufferA,
 	out float4 OutGBufferB,
 	out float4 OutGBufferC,
@@ -339,13 +340,13 @@ void EncodeGBuffer(
 	float QuantizationBias = 0		// -0.5 to 0.5 random float. Used to bias quantization.
 	)
 {
-	OutGBufferA.rgb = EncodeNormal( Data.WorldNormal );
+	OutGBufferA.rgb = EncodeNormal( GBuffer.WorldNormal );
 
 	// compress in 2 bits
-	uint PackedAlpha = Data.DecalMask * 2 + (Data.HasDistanceFieldRepresentation);
+	uint PackedAlpha = GBuffer.DecalMask * 2 + (GBuffer.HasDistanceFieldRepresentation);
 	OutGBufferA.a = PackedAlpha * (1.0 / 3.0);
 
-	if (Data.ShadingModelID == SHADINGMODELID_UNLIT)
+	if (GBuffer.ShadingModelID == SHADINGMODELID_UNLIT)
 	{
 		OutGBufferB = 0;
 		OutGBufferC = 0;
@@ -353,29 +354,29 @@ void EncodeGBuffer(
 	}
 	else
 	{
-		bool bSubsurface = (Data.ShadingModelID == SHADINGMODELID_SUBSURFACE || Data.ShadingModelID == SHADINGMODELID_PREINTEGRATED_SKIN);
+		bool bSubsurface = (GBuffer.ShadingModelID == SHADINGMODELID_SUBSURFACE || GBuffer.ShadingModelID == SHADINGMODELID_PREINTEGRATED_SKIN);
 
 #if DIFFUSE_SPEC_INPUTS
-		OutGBufferB.rgb = EncodeSpecularColor(Data.SpecularColor);
-		OutGBufferC.rgb = EncodeDiffuseColor(Data.DiffuseColor);
+		OutGBufferB.rgb = EncodeSpecularColor(GBuffer.SpecularColor);
+		OutGBufferC.rgb = EncodeDiffuseColor(GBuffer.DiffuseColor);
 #else
 		// NOTE OutGBufferB.b is currently unused!
-		OutGBufferB.r = bSubsurface ? Data.Opacity : Data.Metallic;
-		OutGBufferB.g = Data.Specular;
+		OutGBufferB.r = bSubsurface ? GBuffer.Opacity : GBuffer.Metallic;
+		OutGBufferB.g = GBuffer.Specular;
 		OutGBufferB.b = 0;
-		OutGBufferC.rgb = EncodeBaseColor(Data.BaseColor);
+		OutGBufferC.rgb = EncodeBaseColor(GBuffer.BaseColor);
 #endif
-		OutGBufferB.a = Data.ShadingModelID * (1.0 / 255.0);
+		OutGBufferB.a = GBuffer.ShadingModelID * (1.0 / 255.0);
 
 		// No space for AO. Multiply IndirectIrradiance by AO instead of storing.
-		OutGBufferC.a = EncodeIndirectIrradiance( Data.IndirectIrradiance * Data.GBufferAO ) + QuantizationBias * (1.0 / 255.0);
+		OutGBufferC.a = EncodeIndirectIrradiance( GBuffer.IndirectIrradiance * GBuffer.GBufferAO ) + QuantizationBias * (1.0 / 255.0);
 
 		// Roughness in OutGBufferD.r instead of OutGBufferB.a so that deferred decals can blend in roughness while using MRT
-		OutGBufferD.r = Data.Roughness;
-		OutGBufferD.gba = Data.CustomData;
+		OutGBufferD.r = GBuffer.Roughness;
+		OutGBufferD.gba = GBuffer.CustomData;
 	}
 
-	OutGBufferE = Data.PrecomputedShadowFactors;
+	OutGBufferE = GBuffer.PrecomputedShadowFactors;
 }
 
 /** Populates FGBufferData */
@@ -389,58 +390,58 @@ FGBufferData DecodeGBufferData(
 	float SceneDepth,
 	bool bGetNormalizedNormal)
 {
-	FGBufferData ret;
+	FGBufferData GBuffer;
 
-	ret.WorldNormal = DecodeNormal( InGBufferA.xyz );
+	GBuffer.WorldNormal = DecodeNormal( InGBufferA.xyz );
 
 	if(bGetNormalizedNormal)
 	{
-		ret.WorldNormal = normalize(ret.WorldNormal);
+		GBuffer.WorldNormal = normalize(GBuffer.WorldNormal);
 	}
 
 	uint PackedAlpha = (uint)( InGBufferA.a * 3.999 );
-	ret.DecalMask = PackedAlpha & 2;
-	ret.HasDistanceFieldRepresentation = PackedAlpha & 1;
-	ret.ShadingModelID = DecodeShadingModelId(InGBufferB.a);
+	GBuffer.DecalMask = PackedAlpha & 2;
+	GBuffer.HasDistanceFieldRepresentation = PackedAlpha & 1;
+	GBuffer.ShadingModelID = DecodeShadingModelId(InGBufferB.a);
 	
-	bool bSubsurface = (ret.ShadingModelID == SHADINGMODELID_SUBSURFACE || ret.ShadingModelID == SHADINGMODELID_PREINTEGRATED_SKIN);
+	bool bSubsurface = (GBuffer.ShadingModelID == SHADINGMODELID_SUBSURFACE || GBuffer.ShadingModelID == SHADINGMODELID_PREINTEGRATED_SKIN);
 
 #if DIFFUSE_SPEC_INPUTS
-	ret.BaseColor = 0;
-	ret.Metallic = 0;
-	ret.Specular = 0;
+	GBuffer.BaseColor = 0;
+	GBuffer.Metallic = 0;
+	GBuffer.Specular = 0;
 
-	ret.DiffuseColor = DecodeDiffuseColor(InGBufferC.rgb);
-	ret.SpecularColor = DecodeSpecularColor(InGBufferB.rgb);
+	GBuffer.DiffuseColor = DecodeDiffuseColor(InGBufferC.rgb);
+	GBuffer.SpecularColor = DecodeSpecularColor(InGBufferB.rgb);
 #else
-	ret.BaseColor = DecodeBaseColor(InGBufferC.rgb);
-	ret.Metallic = bSubsurface ? 0 : InGBufferB.r;
-	ret.Specular = InGBufferB.g;
+	GBuffer.BaseColor = DecodeBaseColor(InGBufferC.rgb);
+	GBuffer.Metallic = bSubsurface ? 0 : InGBufferB.r;
+	GBuffer.Specular = InGBufferB.g;
 
-	ret.DiffuseColor = ret.BaseColor - ret.BaseColor * ret.Metallic;
-	ret.SpecularColor = lerp( 0.08 * ret.Specular.xxx, ret.BaseColor, ret.Metallic );
+	GBuffer.DiffuseColor = GBuffer.BaseColor - GBuffer.BaseColor * GBuffer.Metallic;
+	GBuffer.SpecularColor = lerp( 0.08 * GBuffer.Specular.xxx, GBuffer.BaseColor, GBuffer.Metallic );
 
 	// todo: COMPILE_SHADERS_FOR_DEVELOPMENT is unfinished feature, using XBOXONE_PROFILE as workaround
 	#if COMPILE_SHADERS_FOR_DEVELOPMENT == 1 && !XBOXONE_PROFILE
 	{
 		// this feature is only needed for development/editor - we can compile it out for a shipping build (see r.CompileShadersForDevelopment cvar help)
-		ret.DiffuseColor = ret.DiffuseColor * View.DiffuseOverrideParameter.www + View.DiffuseOverrideParameter.xyz;
-		ret.SpecularColor = ret.SpecularColor * View.SpecularOverrideParameter.w + View.SpecularOverrideParameter.xyz;
+		GBuffer.DiffuseColor = GBuffer.DiffuseColor * View.DiffuseOverrideParameter.www + View.DiffuseOverrideParameter.xyz;
+		GBuffer.SpecularColor = GBuffer.SpecularColor * View.SpecularOverrideParameter.w + View.SpecularOverrideParameter.xyz;
 	}
 	#endif //COMPILE_SHADERS_FOR_DEVELOPMENT == 1
 
 #endif
 
-	ret.Opacity = bSubsurface ? InGBufferB.r : 0;
+	GBuffer.Opacity = bSubsurface ? InGBufferB.r : 0;
 	
-	ret.Roughness = InGBufferD.r;
-	ret.CustomData = InGBufferD.gba;
-	ret.GBufferAO = 1;
-	ret.IndirectIrradiance = DecodeIndirectIrradiance(InGBufferC.a);
-	ret.PrecomputedShadowFactors = InGBufferE;
-	ret.CustomDepth = ConvertFromDeviceZ(CustomNativeDepth);
-	ret.Depth = SceneDepth;
-	return ret;
+	GBuffer.Roughness = InGBufferD.r;
+	GBuffer.CustomData = InGBufferD.gba;
+	GBuffer.GBufferAO = 1;
+	GBuffer.IndirectIrradiance = DecodeIndirectIrradiance(InGBufferC.a);
+	GBuffer.PrecomputedShadowFactors = InGBufferE;
+	GBuffer.CustomDepth = ConvertFromDeviceZ(CustomNativeDepth);
+	GBuffer.Depth = SceneDepth;
+	return GBuffer;
 }
 
 
diff --git a/Engine/Shaders/DynamicLightingCommon.usf b/Engine/Shaders/DynamicLightingCommon.usf
--- a/Engine/Shaders/DynamicLightingCommon.usf
+++ b/Engine/Shaders/DynamicLightingCommon.usf
@@ -39,13 +39,13 @@ float RadialAttenuation(float3 WorldLightVector, half FalloffExponent)
 
 /** 
  * Calculates attenuation for a spot light.
- * WorldLightVector is the vector from the position being shaded to the light, divided by the radius of the light. 
+ * L normalize vector to light. 
  * SpotDirection is the direction of the spot light.
  * SpotAngles.x is CosOuterCone, SpotAngles.y is InvCosConeDifference. 
  */
-float SpotAttenuation(float3 WorldLightVector, float3 SpotDirection, float2 SpotAngles)
+float SpotAttenuation(float3 L, float3 SpotDirection, float2 SpotAngles)
 {
-	float ConeAngleFalloff = Square(saturate((dot(normalize(WorldLightVector), -SpotDirection) - SpotAngles.x) * SpotAngles.y));
+	float ConeAngleFalloff = Square(saturate((dot(L, -SpotDirection) - SpotAngles.x) * SpotAngles.y));
 	return ConeAngleFalloff;
 }
 
@@ -77,7 +77,7 @@ float CalcLightAttenuation(float3 WorldPosition, out float3 WorldLightVector)
 
 	float SpotFalloff = 1;
 	#if RADIAL_ATTENUATION
-		SpotFalloff = SpotAttenuation(WorldLightVector, -DeferredLightUniforms.NormalizedLightDirection, DeferredLightUniforms.SpotAngles);
+		SpotFalloff = SpotAttenuation( normalize(WorldLightVector), -DeferredLightUniforms.NormalizedLightDirection, DeferredLightUniforms.SpotAngles);
 	#endif
 
 	return SpotFalloff * DistanceAttenuation;
diff --git a/Engine/Shaders/PostProcessAmbient.usf b/Engine/Shaders/PostProcessAmbient.usf
--- a/Engine/Shaders/PostProcessAmbient.usf
+++ b/Engine/Shaders/PostProcessAmbient.usf
@@ -43,7 +43,7 @@ float3 SpecularIBL( uint2 Random, float3 SpecularColor, float Roughness, float3
 {
 	float3 SpecularLighting = 0;
 
-	const uint NumSamples = 64;
+	const uint NumSamples = 32;
 	for( uint i = 0; i < NumSamples; i++ )
 	{
 		float2 E = Hammersley( i, NumSamples, Random );
@@ -73,6 +73,374 @@ float3 SpecularIBL( uint2 Random, float3 SpecularColor, float Roughness, float3
 	return SpecularLighting / NumSamples;
 }
 
+float3 StandardShading( FGBufferData GBuffer, float Roughness, float3 L, float3 V, half3 N )
+{
+	float3 H = normalize(V + L);
+	float NoL = saturate( dot(N, L) );
+	float NoV = saturate( dot(N, V) );
+	float NoH = saturate( dot(N, H) );
+	float VoH = saturate( dot(V, H) );
+	
+	// Generalized microfacet specular
+	float D = D_GGX( Roughness, NoH );
+	float Vis = Vis_Schlick( Roughness, NoV, NoL );
+	float3 F = F_Schlick( GBuffer.SpecularColor, VoH );
+	
+	//float k = Square( Roughness ) * 0.25;
+	//float Vis = 0.25 / Square( VoH * (1 - k) + k );
+	//float Vis = 0.25 / ( VoH * VoH * (1 - k*k) + k*k );
+	//float Vis = Vis_Smith( Roughness, NoV, NoL );
+
+	return Diffuse_Lambert( GBuffer.DiffuseColor ) + (D * Vis) * F;
+}
+
+#if 0
+void ClearCoatShadingShared( FGBufferData GBuffer, float3 V, half3 N )
+{
+	float NoV = saturate( dot(N, V) );
+	Shared.NoV = NoV;
+
+	// D_GGX, Vis_Schlick
+	float m = Roughness * Roughness;
+	Shared.m2 = m * m;
+	Shared.SpecularMul = (0.25 / PI) * ( Shared.m2 / (NoV * NoV) );
+	
+	// F_Schlick
+	float F0 = GBuffer.Specular * 0.08;
+	float Fc = exp2( (-5.55473 * NoV - 6.98316) * NoV );	// 1 mad, 1 mul, 1 exp
+	float F = Fc + (1 - Fc) * F0;							// 1 add, 3 mad
+	Shared.SpecularMul *= F;
+
+	Shared.AbsorptionFactor = 0.5 * log2( GBuffer.BaseColor );
+
+
+	float AbsorptionFactor = 0.72134752 * rcp( NoL2 ) + ( 0.39674113 + 0.72134752 * rcp( NoV2 ) );
+	float3 Absorption = exp2( AbsorptionFactor * GBuffer.BaseColor - AbsorptionFactor );
+
+	
+
+	// Diffuse_Lambert
+	Shared.DiffuseMul = DiffuseColor * (1.0 / PI);
+
+	// D_GGX, Vis_Schlick
+	float m = Roughness * Roughness;
+	Shared.m2 = m * m;
+	Shared.k = Square( Roughness ) * 0.5;
+	//float Vis_SchlickV = NoV * (1 - Shared.k) + Shared.k;
+	Shared.SpecularMul = (0.25 / PI) * Shared.m2;
+	
+	// F_Schlick
+	Shared.SpecularMul *= saturate( 50.0 * SpecularColor.g );
+}
+
+float3 ClearCoatShadingPerLight( Shared, float3 L, float3 V, half3 N )
+{
+	float ClearCoatRoughness = 0.1;
+
+	float Thickness = 1.5;
+	float3 WaterAbsorption = { 0.98, 0.4, 0.28 };
+
+	GBuffer.BaseColor = exp( -WaterAbsorption * Thickness * 2 );
+
+	const float3 n = { 1.94, 1, 0.6 };
+	const float3 k = { 8.21, 6.69, 4.86 };
+	GBuffer.SpecularColor = ( Square(n - 1) + k*k ) / ( Square(n + 1) + k*k );
+
+	// TODO control SpecularColor
+
+
+	float3 H = normalize(V + L);
+	float NoL = saturate( dot(N, L) );
+	float NoV = saturate( dot(N, V) );
+	float NoH = saturate( dot(N, H) );
+	float VoH = saturate( dot(V, H) );
+	
+	// simplify for low roughness
+	// VoH == NoV
+	// EnvBRDF == F_Schlick()
+
+	// Generalized microfacet specular
+	float D = D_GGX( ClearCoatRoughness, NoH );
+	float Vis = Vis_Kelemen( VoH );
+	//float Vis = rcp( 4 * NoV * NoV );
+	
+	// F_Schlick
+	float F0 = GBuffer.Specular * 0.08;
+	float Fc = exp2( (-5.55473 * VoH - 6.98316) * VoH );	// 1 mad, 1 mul, 1 exp
+	float F = Fc + (1 - Fc) * F0;							// 1 add, 3 mad
+
+	float Fr1 = D * Vis * F;
+
+	
+	// Approximation
+	float RefractBlend = (0.22 * VoH + 0.7) * VoH + 0.745;	// 2 mad
+
+	// dot products distribute
+	float RefractNoH = RefractBlend * NoH;					// 1 mul
+	float NoL2 = saturate( RefractNoH - (1 / 1.5) * NoL );	// 1 mad
+	float NoV2 = saturate( RefractNoH - (1 / 1.5) * NoV );	// 1 mad
+
+	float3 AbsorptionFactor = 0.5 * log2( GBuffer.BaseColor );
+	float3 Absorption = exp2( AbsorptionFactor * ( 1 / NoV2 + 1 / NoL2 ) );
+	//float3 Absorption = pow( GBuffer.BaseColor, 0.5 * ( rcp(NoV2) + rcp(NoL2) ) );
+	
+	float3 LayerAttenuation = (1 - F) * Absorption;
+	
+	// Generalized microfacet specular
+	float D2 = D_GGX( Roughness, NoH );
+	float Vis2 = Vis_Schlick( Roughness, NoV2, NoL2 );
+	float3 F2 = F_Schlick( SpecularColor, VoH );
+	//float3 F2 = GBuffer.SpecularColor;
+
+	float3 Fr2 = Diffuse_Lambert( GBuffer.DiffuseColor ) + (D2 * Vis2) * F2;
+	
+	return Fr1 + Fr2 * LayerAttenuation;
+}
+#endif
+
+float3 ClearCoatShading( FGBufferData GBuffer, float Roughness, float3 L, float3 V, half3 N )
+{
+	float ClearCoat = 1;
+	float ClearCoatRoughness = 0.1;
+	float Film = 1;
+	float MetalSpec = 0.9;
+
+	Film *= ClearCoat;
+
+#if 0
+	float3 H = normalize(V + L);
+	float NoL = saturate( dot(N, L) );
+	float NoV = saturate( dot(N, V) );
+	float NoH = saturate( dot(N, H) );
+	float VoH = saturate( dot(V, H) );
+	
+	// Generalized microfacet specular
+	float D = D_GGX( ClearCoatRoughness, NoH );
+	float Vis = Vis_Kelemen( VoH );
+	
+	// F_Schlick
+	float F0 = GBuffer.Specular * 0.08;
+	float Fc = pow( 1 - VoH, 5 );
+	float F = Fc + (1 - Fc) * F0;
+
+	float Fr1 = D * Vis * F;
+
+	float LayerAttenuation = (1 - F);
+
+	// Generalized microfacet specular
+	float D2 = D_GGX( Roughness, NoH );
+	float Vis2 = Vis_Schlick( Roughness, NoV, NoL );
+	//float3 F2 = F_Schlick( GBuffer.SpecularColor, VoH );
+	float3 F2 = (1 - Fc) * GBuffer.SpecularColor;
+
+	float3 Fr2 = Diffuse_Lambert( GBuffer.DiffuseColor ) + (D2 * Vis2) * F2;
+	
+	return Fr1 + Fr2 * LayerAttenuation;
+#else
+	float3 AbsorptionColor = (1 - Film) + GBuffer.BaseColor * ( Film * (1 / MetalSpec) );
+	GBuffer.BaseColor = lerp( GBuffer.BaseColor, MetalSpec, Film );
+	
+	GBuffer.DiffuseColor = GBuffer.BaseColor - GBuffer.BaseColor * GBuffer.Metallic;
+	GBuffer.SpecularColor = lerp( 0.08 * GBuffer.Specular.xxx, GBuffer.BaseColor, GBuffer.Metallic );
+
+#if 0	// Ref
+	float3 H = normalize(V + L);
+	float NoL = saturate( dot(N, L) );
+	float NoV = saturate( dot(N, V) );
+	float NoH = saturate( dot(N, H) );
+	float VoH = saturate( dot(V, H) );
+	
+	// Hard coded IOR of 1.5
+
+	// Generalized microfacet specular
+	float D = D_GGX( ClearCoatRoughness, NoH );
+	float Vis = Vis_Kelemen( VoH );
+	
+	// F_Schlick
+	float F0 = GBuffer.Specular * 0.08;
+	float Fc = pow( 1 - VoH, 5 );
+	float F = Fc + (1 - Fc) * F0;
+
+	float Fr1 = D * Vis * F;
+
+	// Refract rays
+	float3 L2 = refract( -L, -H, 1 / 1.5 );
+	float3 V2 = refract( -V, -H, 1 / 1.5 );
+	float3 H2 = normalize(V2 + L2);
+
+	float NoL2 = saturate( dot(N, L2) );
+	float NoV2 = saturate( dot(N, V2) );
+	float NoH2 = saturate( dot(N, H2) );
+	float VoH2 = saturate( dot(V2, H2) );
+
+	float  AbsorptionDist = rcp(NoV2) + rcp(NoL2);
+	float3 Absorption = pow( AbsorptionColor, 0.5 * AbsorptionDist );
+	
+	float F21 = Fresnel( 1 / 1.5, saturate( dot(V2, H) ) );
+	float TotalInternalReflection = 1 - F21 * G_Schlick( GBuffer.Roughness, NoV2, NoL2 );
+	float3 LayerAttenuation = ( (1 - F) * TotalInternalReflection ) * Absorption;
+
+	// Treat SpecularColor as relative to IOR. Artist compensates.
+	//GBuffer.SpecularColor = ChangeBaseMedium( GBuffer.SpecularColor, 1.5 );
+	
+	// Generalized microfacet specular
+	float D2 = D_GGX( Roughness, NoH2 );
+	float Vis2 = Vis_Schlick( GBuffer.Roughness, NoV2, NoL2 );
+	float3 F2 = F_Schlick( GBuffer.SpecularColor, VoH2 );
+
+	float3 Fr2 = Diffuse_Lambert( GBuffer.DiffuseColor ) + (D2 * Vis2) * F2;
+	
+	return Fr1 + Fr2 * LayerAttenuation;
+#else
+	float3 H = normalize(V + L);
+	float NoL = saturate( dot(N, L) );
+	float NoV = saturate( dot(N, V) );
+	float NoH = saturate( dot(N, H) );
+	float VoH = saturate( dot(V, H) );
+	
+	// Hard coded IOR of 1.5
+
+	// Generalized microfacet specular
+	float D = D_GGX( ClearCoatRoughness, NoH );
+	float Vis = Vis_Kelemen( VoH );
+	
+	// F_Schlick
+	float F0 = GBuffer.Specular * 0.08;
+	float Fc = pow( 1 - VoH, 5 );
+	float F = Fc + (1 - Fc) * F0;
+
+	float Fr1 = D * Vis * F;
+
+	// Refract rays
+	//float3 L2 = refract( -L, -H, 1 / 1.5 );
+	//float3 V2 = refract( -V, -H, 1 / 1.5 );
+
+	// LoH == VoH
+	//float RefractBlend = sqrt( 4 * VoH*VoH + 5 ) / 3 + 2.0 / 3 * VoH;
+	//float3 L2 = RefractBlend * H - L / 1.5;
+	//float3 V2 = RefractBlend * H - V / 1.5;
+	//float NoL2 = saturate( dot(N, L2) );
+	//float NoV2 = saturate( dot(N, V2) );
+	
+	// Approximation
+	float RefractBlend = (0.22 * VoH + 0.7) * VoH + 0.745;	// 2 mad
+	// Dot products distribute. No need for L2 and V2.
+	float RefractNoH = RefractBlend * NoH;					// 1 mul
+	float NoL2 = saturate( RefractNoH - (1 / 1.5) * NoL );	// 1 mad
+	float NoV2 = saturate( RefractNoH - (1 / 1.5) * NoV );	// 1 mad
+	// Should refract H too but unimportant
+
+	float  AbsorptionDist = rcp(NoV2) + rcp(NoL2);
+	float3 Absorption = pow( AbsorptionColor, 0.5 * AbsorptionDist );
+
+	// Approximation
+	//float  AbsorptionDist = ( NoV2 + NoL2 ) / ( NoV2 * NoL2 );
+	//float3 Absorption = AbsorptionColor * ( AbsorptionColor * (AbsorptionDist * 0.5 - 1) + (2 - 0.5 * AbsorptionDist) );
+	//float3 Absorption = AbsorptionColor + AbsorptionColor * (AbsorptionColor - 1) * (AbsorptionDist * 0.5 - 1);	// use for shared version
+	
+	//float F21 = Fresnel( 1 / 1.5, saturate( dot(V2, H) ) );
+	//float TotalInternalReflection = 1 - F21 * G_Schlick( Roughness, NoV2, NoL2 );
+	//float3 LayerAttenuation = ( (1 - F12) * TotalInternalReflection ) * Absorption;
+
+	// Approximation
+	float3 LayerAttenuation = (1 - F) * Absorption;
+
+	// Treat SpecularColor as relative to IOR. Artist compensates.
+	//SpecularColor = ChangeBaseMedium( SpecularColor, 1.5 );
+	
+	// Generalized microfacet specular
+	float D2 = D_GGX( Roughness, NoH );
+	float Vis2 = Vis_Schlick( Roughness, NoV2, NoL2 );
+	float3 F2 = F_Schlick( GBuffer.SpecularColor, VoH );
+
+	float3 Fr2 = Diffuse_Lambert( GBuffer.DiffuseColor ) + (D2 * Vis2) * F2;
+	
+	return Fr1 + Fr2 * LayerAttenuation;
+#endif
+
+#endif
+}
+
+
+float3 ImageBasedLightingMIS( FGBufferData GBuffer, float3 V, float3 N, uint2 Random )
+{
+	BRANCH
+	if( GBuffer.ShadingModelID == 0 )
+	{
+		return 0;
+	}
+
+	float3 Lighting = 0;
+
+	float Roughness1 = GBuffer.Roughness;
+	float Roughness2 = 0.1;
+	
+	uint NumSamples[] =
+	{
+		16,
+		16,
+		16,
+	};
+
+	UNROLL
+	for( uint Set = 0; Set < 3; Set++ )
+	{
+		LOOP
+		for( uint i = 0; i < NumSamples[ Set ]; i++ )
+		{
+			float2 E = Hammersley( i, NumSamples[ Set ], Random );
+			
+			float3 L, H;
+			if( Set == 0 )
+			{
+				L = TangentToWorld( CosineSampleHemisphere( E ).xyz, N );
+				H = normalize(V + L);
+			}
+			else if( Set == 1 )
+			{
+				H = TangentToWorld( ImportanceSampleGGX( E, Roughness1 ).xyz, N );
+				L = 2 * dot( V, H ) * H - V;
+			}
+			else
+			{
+				H = TangentToWorld( ImportanceSampleGGX( E, Roughness2 ).xyz, N );
+				L = 2 * dot( V, H ) * H - V;
+			}
+
+			float NoL = saturate( dot(N, L) );
+			float NoH = saturate( dot(N, H) );
+			float VoH = saturate( dot(V, H) );
+		
+			if( NoL > 0 && VoH > 0 )
+			{
+				float3 SampleColor = AmbientCubemap.SampleLevel( AmbientCubemapSampler, L, 0 ).rgb;
+
+				float PDF[] =
+				{
+					NoL / PI,
+					D_GGX( Roughness1, NoH ) * NoH / (4 * VoH),
+					D_GGX( Roughness2, NoH ) * NoH / (4 * VoH),
+				};
+
+				// MIS balance heuristic
+				float InvWeight = 0;
+				UNROLL for( uint j = 0; j < 3; j++ )
+				{
+					InvWeight += PDF[j] * NumSamples[j];
+				}
+				float Weight = rcp( InvWeight );
+
+				//float3 Shading = ClearCoatShading( GBuffer, GBuffer.Roughness, L, V, N );
+				float3 Shading = StandardShading( GBuffer, GBuffer.Roughness, L, V, N );
+			
+				Lighting += SampleColor * Shading * ( NoL * Weight );
+			}
+		}
+	}
+
+	return Lighting;
+}
+
 float3 FilterEnvMap( uint2 Random, float Roughness, float3 N, float3 V )
 {
 	float3 FilteredColor = 0;
@@ -142,7 +510,7 @@ float2 IntegrateBRDF( uint2 Random, float Roughness, float NoV )
 
 		if( NoL > 0 )
 		{
-			float Vis = Vis_Smith( Roughness, NoV, NoL );
+			float Vis = Vis_Schlick( Roughness, NoV, NoL );
 
 			// Incident light = NoL
 			// pdf = D * NoH / (4 * VoH)
@@ -171,6 +539,65 @@ float3 ApproximateSpecularIBL( uint2 Random, float3 SpecularColor, float Roughne
 
 	return PrefilteredColor * ( SpecularColor * AB.x + AB.y );
 }
+
+float3 ImageBasedLightingSplitSum( FGBufferData GBuffer, float3 V, float3 N, uint2 Random )
+{
+	BRANCH
+	if( GBuffer.ShadingModelID == 0 )
+	{
+		return 0;
+	}
+
+	float ClearCoat = 1;
+	float ClearCoatRoughness = 0.1;
+	float Film = 1;
+	float MetalSpec = 0.9;
+
+	Film *= ClearCoat;
+
+	float NoV = saturate( dot( N, V ) );
+	
+	float Lobe1Roughness = ClearCoatRoughness;
+	float Lobe2Roughness = GBuffer.Roughness;
+	
+#if 1
+	float3 Lobe1Color = FilterEnvMap( Random, Lobe1Roughness, N, V );
+	float3 Lobe2Color = FilterEnvMap( Random, Lobe2Roughness, N, V );
+#else
+	float3 R = 2 * dot( V, N ) * N - V;
+	//float3 Lobe1Color = PrefilterEnvMap( Random, Lobe1Roughness, R );
+	//float3 Lobe2Color = PrefilterEnvMap( Random, Lobe2Roughness, R );
+
+	float Lobe1Mip = ComputeCubemapMipFromRoughness( Lobe1Roughness, AmbientCubemapMipAdjust.w );
+	float Lobe2Mip = ComputeCubemapMipFromRoughness( Lobe2Roughness, AmbientCubemapMipAdjust.w );
+
+	float3 Lobe1Color = TextureCubeSampleLevel( AmbientCubemap, AmbientCubemapSampler, R, Lobe1Mip ).rgb;
+	float3 Lobe2Color = TextureCubeSampleLevel( AmbientCubemap, AmbientCubemapSampler, R, Lobe2Mip ).rgb;
+#endif
+
+	// F_Schlick
+	float F0 = GBuffer.Specular * 0.08;
+	float Fc = pow( 1 - NoV, 5 );
+	float F = Fc + (1 - Fc) * F0;
+	F *= ClearCoat;
+
+	float3 Fr1 = Lobe1Color * F;
+
+	float LayerAttenuation = (1 - F);
+
+	// Generalized microfacet specular
+	//float Vis2 = Vis_Schlick( Roughness, NoV2, NoL2 );
+	//float3 F2 = F_Schlick( GBuffer.SpecularColor, VoH );
+	
+	//float3 Fr2 = Diffuse_Lambert( GBuffer.DiffuseColor ) + (D2 * Vis2) * F2;
+
+	float2 AB = PreIntegratedGF.SampleLevel( PreIntegratedGFSampler, float2( NoV, GBuffer.Roughness ), 0 ).rg;
+
+	float3 Fr2 = Lobe2Color * ( GBuffer.SpecularColor * AB.x + AB.y * (1 - ClearCoat) );
+	//float3 Fr2 = Lobe2Color * EnvBRDF( GBuffer.SpecularColor, GBuffer.Roughness, NoV );
+	
+	return Fr1 + Fr2 * LayerAttenuation;
+}
 #endif
 
 void MainPS(in float4 UVAndScreenPos : TEXCOORD0, out float4 OutColor : SV_Target0)
@@ -188,8 +615,8 @@ void MainPS(in float4 UVAndScreenPos : TEXCOORD0, out float4 OutColor : SV_Targe
 	float3 ScreenVector = normalize(mul(float4(ScreenSpacePos, 1, 0), View.ScreenToWorld).xyz);
 	
 	uint2 Random = ScrambleTEA( PixelPos );
-	Random.x ^= uint( View.RealTime * 21389 );
-	Random.y ^= uint( View.RealTime * 49233 );
+	Random.x ^= View.Random;
+	Random.y ^= View.Random;
 
 	float3 DiffuseAndSpecular = 0;
 
@@ -198,7 +625,7 @@ void MainPS(in float4 UVAndScreenPos : TEXCOORD0, out float4 OutColor : SV_Targe
 	// Diffuse
 	{
 	#if IMPORTANCE_SAMPLE
-		DiffuseAndSpecular += DiffuseIBL( Random, GBuffer.DiffuseColor, GBuffer.WorldNormal );
+		//DiffuseAndSpecular += DiffuseIBL( Random, GBuffer.DiffuseColor, GBuffer.WorldNormal );
 	#else
 		// we want to access the mip with the preconvolved diffuse lighting (coneangle=90 degree)
 		DiffuseAndSpecular += GBuffer.DiffuseColor * DiffuseLookup;
@@ -208,7 +635,7 @@ void MainPS(in float4 UVAndScreenPos : TEXCOORD0, out float4 OutColor : SV_Targe
 	// Specular
 	{
 	#if IMPORTANCE_SAMPLE
-		DiffuseAndSpecular += SpecularIBL( Random, GBuffer.SpecularColor, GBuffer.Roughness, GBuffer.WorldNormal, -ScreenVector );
+		//DiffuseAndSpecular += SpecularIBL( Random, GBuffer.SpecularColor, GBuffer.Roughness, GBuffer.WorldNormal, -ScreenVector );
 		//DiffuseAndSpecular += ApproximateSpecularIBL( Random, GBuffer.SpecularColor, GBuffer.Roughness, GBuffer.WorldNormal, -ScreenVector );
 	#else
 		float3 R = reflect( ScreenVector, GBuffer.WorldNormal );
@@ -220,6 +647,15 @@ void MainPS(in float4 UVAndScreenPos : TEXCOORD0, out float4 OutColor : SV_Targe
 		DiffuseAndSpecular += SampleColor * EnvBRDF( GBuffer.SpecularColor, GBuffer.Roughness, NoV );
 	#endif
 	}
+
+	#if IMPORTANCE_SAMPLE
+		DiffuseAndSpecular += ImageBasedLightingMIS( GBuffer, -ScreenVector, GBuffer.WorldNormal, Random );
+	#endif
+
+	BRANCH if( GBuffer.ShadingModelID == SHADINGMODELID_CLEAR_COAT )
+	{
+		DiffuseAndSpecular = ImageBasedLightingSplitSum( GBuffer, -ScreenVector, GBuffer.WorldNormal, Random );
+	}
 	
 	// apply darkening from ambient occlusion (does not use PostprocessInput1 to set white texture if SSAO is off)
 	float AmbientOcclusion = GBuffer.GBufferAO * ScreenSpaceData.AmbientOcclusion;
diff --git a/Engine/Shaders/ReflectionEnvironmentComputeShaders.usf b/Engine/Shaders/ReflectionEnvironmentComputeShaders.usf
--- a/Engine/Shaders/ReflectionEnvironmentComputeShaders.usf
+++ b/Engine/Shaders/ReflectionEnvironmentComputeShaders.usf
@@ -24,6 +24,8 @@ SamplerState ReflectionEnvironmentColorSampler;
 #define ATOMIC_REDUCTION 0
 #endif
 
+#define VISUALIZE_OVERLAP	0
+
 uint NumCaptures;
 /** View rect min in xy, max in zw. */
 uint4 ViewDimensions;
@@ -151,7 +153,7 @@ void ComputeTileMinMax(uint ThreadIndex, float SceneDepth, out float MinTileZ, o
 
 	MinTileZ = asfloat(IntegerTileMinZ);
 	MaxTileZ = asfloat(IntegerTileMaxZ);
-	
+
 	float HalfZ = .5f * (MinTileZ + MaxTileZ);
 
 	MinTileZ2 = HalfZ;
@@ -171,6 +173,8 @@ void DoTileCulling(uint3 GroupId, uint ThreadIndex, float MinTileZ, float MaxTil
     float4 C2 = float4(0.0f,								-View.ViewToClip._22 * TileScale.y, View.ViewToClip._32 * TileScale.y + TileBias.y,	0.0f);
     float4 C4 = float4(0.0f,								0.0f,								1.0f,											0.0f);
 
+	// TODO transform to world space
+#if ATOMIC_REDUCTION
 	float4 frustumPlanes[8];
 	frustumPlanes[0] = C4 - C1;
 	frustumPlanes[1] = C4 + C1;
@@ -180,6 +184,15 @@ void DoTileCulling(uint3 GroupId, uint ThreadIndex, float MinTileZ, float MaxTil
 	frustumPlanes[5] = float4(0.0f, 0.0f, -1.0f,  MaxTileZ2);
 	frustumPlanes[6] = float4(0.0f, 0.0f,  1.0f, -MinTileZ2);
 	frustumPlanes[7] = float4(0.0f, 0.0f, -1.0f,  MaxTileZ);
+#else
+	float4 frustumPlanes[6];
+	frustumPlanes[0] = C4 - C1;
+	frustumPlanes[1] = C4 + C1;
+	frustumPlanes[2] = C4 - C2;
+	frustumPlanes[3] = C4 + C2;
+	frustumPlanes[4] = float4(0.0f, 0.0f,  1.0f, -MinTileZ);
+	frustumPlanes[5] = float4(0.0f, 0.0f, -1.0f,  MaxTileZ);
+#endif
 
 	// Normalize tile frustum planes
 	UNROLL 
@@ -207,7 +220,7 @@ void DoTileCulling(uint3 GroupId, uint ThreadIndex, float MinTileZ, float MaxTil
 		// Cull the light against the tile's frustum planes
 		// Note: this has some false positives, a light that is intersecting three different axis frustum planes yet not intersecting the volume of the tile will be treated as intersecting
 		bool bInTile = true;  
-				
+
 		// Test against the screen x and y oriented planes first
 		UNROLL
 		for (uint i = 0; i < 4; ++i) 
@@ -219,29 +232,43 @@ void DoTileCulling(uint3 GroupId, uint ThreadIndex, float MinTileZ, float MaxTil
 		BRANCH
 		if (bInTile)
 		{
-			bool bInNearDepthRange = true;  
+#if ATOMIC_REDUCTION
+			bool bInNearDepthRange = true;
 				
 			// Test against the near depth range
 			UNROLL
-			for (uint i = 4; i < 6; ++i) 
+			for (uint i = 4; i < 6; ++i)
 			{
 				float PlaneDistance = dot(frustumPlanes[i], float4(BoundsViewPosition, 1.0f));
 				bInNearDepthRange = bInNearDepthRange && (PlaneDistance >= -CapturePositionAndRadius.w);
 			}
 
-			bool bInFarDepthRange = true;  
+			bool bInFarDepthRange = true;
 				
 			// Test against the far depth range
 			UNROLL
-			for (uint j = 6; j < 8; ++j) 
+			for (uint j = 6; j < 8; ++j)
 			{
 				float PlaneDistance = dot(frustumPlanes[j], float4(BoundsViewPosition, 1.0f));
 				bInFarDepthRange = bInFarDepthRange && (PlaneDistance >= -CapturePositionAndRadius.w);
 			}
-				 
+
+			bool bInDepthRange = bInNearDepthRange || bInFarDepthRange;
+#else
+			bool bInDepthRange = true;  
+				
+			// Test against the depth range
+			UNROLL
+			for (uint i = 4; i < 6; ++i) 
+			{
+				float PlaneDistance = dot(frustumPlanes[i], float4(BoundsViewPosition, 1.0f));
+				bInDepthRange = bInDepthRange && (PlaneDistance >= -CapturePositionAndRadius.w);
+			}
+#endif
+
 			// Add this capture to the list of indices if it intersects
 			BRANCH
-			if (bInNearDepthRange || bInFarDepthRange) 
+			if (bInDepthRange)
 			{
 				uint ListIndex;
 				InterlockedAdd(TileNumReflectionCaptures, 1U, ListIndex);
@@ -289,6 +316,259 @@ void DoTileCulling(uint3 GroupId, uint ThreadIndex, float MinTileZ, float MaxTil
 	GroupMemoryBarrierWithGroupSync();
 }
 
+struct FLobe
+{
+	float4	Color;
+	float3	Direction;
+	float	Roughness;
+	float	RayInfluance;
+};
+
+void GatherRadiance( inout FLobe Lobes[2], const uint NumLobes, float3 WorldPosition, float3 RayDirection )
+{
+	float LobeMip[2];
+
+	UNROLL
+	for( uint Lobe = 0; Lobe < NumLobes; Lobe++ )
+	{
+		LobeMip[ Lobe ] = ComputeReflectionCaptureMipFromRoughness( Lobes[ Lobe ].Roughness );
+	}
+
+	uint NumCapturesAffectingTile = TileNumReflectionCaptures;
+
+	// Accumulate reflections from captures affecting this tile, applying largest captures first so that the smallest ones display on top
+	LOOP
+	for (uint TileCaptureIndex = 0; TileCaptureIndex < NumCapturesAffectingTile; TileCaptureIndex++) 
+	{
+		#if SORT_CAPTURES
+			uint CaptureIndex = SortedTileReflectionCaptureIndices[TileCaptureIndex];
+		#else
+			uint CaptureIndex = TileReflectionCaptureIndices[TileCaptureIndex];
+		#endif
+			
+		float4 CapturePositionAndRadius = ReflectionCapture.PositionAndRadius[CaptureIndex];
+		float4 CaptureProperties = ReflectionCapture.CaptureProperties[CaptureIndex];
+		float3 CaptureVector = WorldPosition - CapturePositionAndRadius.xyz;
+		float CaptureVectorLength = sqrt(dot(CaptureVector, CaptureVector));
+
+		BRANCH
+		if (CaptureVectorLength < CapturePositionAndRadius.w)
+		{
+			float NormalizedDistanceToCapture = saturate(CaptureVectorLength / CapturePositionAndRadius.w);
+			float3 ProjectedCaptureVector = RayDirection;
+
+			// Fade out based on distance to capture
+			float DistanceAlpha = 0;
+
+			#define PROJECT_ONTO_SHAPE 1
+			#if PROJECT_ONTO_SHAPE
+					
+				#define SUPPORT_PLANE 0
+				#if SUPPORT_PLANE
+					BRANCH
+					// Plane
+					if (CaptureProperties.b > 1)
+					{
+						float4 ImagePlane = float4(ReflectionCapture.BoxTransform[CaptureIndex][0][0], ReflectionCapture.BoxTransform[CaptureIndex][1][0], ReflectionCapture.BoxTransform[CaptureIndex][2][0], ReflectionCapture.BoxTransform[CaptureIndex][3][0]);
+
+						float VectorDotPlaneNormal = dot(ImagePlane.xyz, RayDirection);
+						// VectorDotPlaneNormal < 0 means the ray hit the front face
+						BRANCH
+						if (VectorDotPlaneNormal < 0)
+						{
+							float PlaneDistance = dot(ImagePlane.xyz, WorldPosition) - ImagePlane.w;
+							// Time along the ray defined by WorldPosition + IntersectionTime * RayDirection that the intersection took place
+							float IntersectionTime = -PlaneDistance / VectorDotPlaneNormal;
+
+							BRANCH
+							// Skip intersections behind the pixel being shaded
+							if (IntersectionTime > 0)
+							{
+								// Calculate the world space intersection position
+								float3 IntersectPosition = WorldPosition + IntersectionTime * ReflectionVector;
+									
+								float2 ReflectionUVs;
+								float4 CurrentReflectionXAxis = float4(ReflectionCapture.BoxTransform[CaptureIndex][0][1], ReflectionCapture.BoxTransform[CaptureIndex][1][1], ReflectionCapture.BoxTransform[CaptureIndex][2][1], ReflectionCapture.BoxTransform[CaptureIndex][3][1]);
+								float3 CurrentImageReflectionOrigin = CapturePositionAndRadius.xyz;
+								float XLength = length(CurrentReflectionXAxis.xyz);
+								float3 NormalizedXAxis = CurrentReflectionXAxis.xyz / XLength;
+								// Calculate the quad UVs by projecting the vector from the intersection to the quad origin onto each quad axis
+								ReflectionUVs.x = dot(NormalizedXAxis, IntersectPosition - CurrentImageReflectionOrigin.xyz);
+								float3 ReflectionYAxis = cross(ImagePlane.xyz, NormalizedXAxis) * CurrentReflectionXAxis.w;
+								ReflectionUVs.y = dot(ReflectionYAxis, IntersectPosition - CurrentImageReflectionOrigin.xyz);
+								ReflectionUVs = .5f * ReflectionUVs / float2(XLength, XLength * CurrentReflectionXAxis.w) + .5f;
+
+								ProjectedCaptureVector = IntersectPosition - CapturePositionAndRadius.xyz;
+
+								if (ReflectionUVs.x > 0 && ReflectionUVs.x < 1 && ReflectionUVs.y > 0 && ReflectionUVs.y < 1)
+								{
+									CompositedLighting.rgb += 1;
+								}
+							}
+						}
+					}
+					else 
+				#endif
+				// Box
+				BRANCH if (CaptureProperties.b > 0)
+				{
+					// Transform the ray into the local space of the box, where it is an AABB with mins at -1 and maxs at 1
+					float3 LocalRayStart		= mul(float4(WorldPosition, 1), ReflectionCapture.BoxTransform[CaptureIndex]).xyz;
+					float3 LocalRayDirection	= mul(float4(RayDirection,  0), ReflectionCapture.BoxTransform[CaptureIndex]).xyz;
+
+					float3 InvRayDir = rcp(LocalRayDirection);
+	
+					//find the ray intersection with each of the 3 planes defined by the minimum extrema.
+					float3 FirstPlaneIntersections = -InvRayDir - LocalRayStart * InvRayDir;
+					//find the ray intersection with each of the 3 planes defined by the maximum extrema.
+					float3 SecondPlaneIntersections = InvRayDir - LocalRayStart * InvRayDir;
+					//get the furthest of these intersections along the ray
+					float3 FurthestPlaneIntersections = max(FirstPlaneIntersections, SecondPlaneIntersections);
+
+					//clamp the intersections to be between RayOrigin and RayEnd on the ray
+					float Intersection = min(FurthestPlaneIntersections.x, min(FurthestPlaneIntersections.y, FurthestPlaneIntersections.z));
+
+					// Compute the reprojected vector
+					float3 IntersectPosition = WorldPosition + Intersection * RayDirection;
+					ProjectedCaptureVector = IntersectPosition - CapturePositionAndRadius.xyz;
+
+					// Compute the distance from the receiving pixel to the box for masking
+					// Apply local to world scale to take scale into account without transforming back to world space
+					// Shrink the box by the transition distance (BoxScales.w) so that the fade happens inside the box influence area
+					float4 BoxScales = ReflectionCapture.BoxScales[CaptureIndex];
+					float BoxDistance = ComputeDistanceFromBoxToPoint(-(BoxScales.xyz - .5f * BoxScales.w), BoxScales.xyz - .5f * BoxScales.w, LocalRayStart * BoxScales.xyz);
+
+					// Setup a fade based on receiver distance to the box, hides the box influence shape
+					DistanceAlpha = 1.0 - smoothstep(0, .7f * BoxScales.w, BoxDistance);
+				}
+				// Sphere
+				else
+				{
+					float ProjectionSphereRadius = CapturePositionAndRadius.w * 1.2f;
+					float SphereRadiusSquared = ProjectionSphereRadius * ProjectionSphereRadius;
+
+					float3 ReceiverToSphereCenter = WorldPosition - CapturePositionAndRadius.xyz;
+					float ReceiverToSphereCenterSq = dot(ReceiverToSphereCenter, ReceiverToSphereCenter);
+
+					// Find the intersection between the ray along the reflection vector and the capture's sphere
+					float3 QuadraticCoef;
+					QuadraticCoef.x = 1;
+					QuadraticCoef.y = 2 * dot(RayDirection, ReceiverToSphereCenter);
+					QuadraticCoef.z = ReceiverToSphereCenterSq - SphereRadiusSquared;
+
+					float Determinant = QuadraticCoef.y * QuadraticCoef.y - 4 * QuadraticCoef.z;
+							 
+					// Only continue if the ray intersects the sphere
+					if (Determinant >= 0)
+					{
+						float FarIntersection = (sqrt(Determinant) - QuadraticCoef.y) * 0.5;
+
+						float3 IntersectPosition = WorldPosition + FarIntersection * RayDirection;
+						ProjectedCaptureVector = IntersectPosition - CapturePositionAndRadius.xyz;
+						// Note: some compilers don't handle smoothstep min > max (this was 1, .6)
+						DistanceAlpha = 1.0 - smoothstep(.6, 1, NormalizedDistanceToCapture);
+					}
+				}
+
+			#else 
+				DistanceAlpha = 1.0;
+			#endif
+
+			float CaptureArrayIndex = CaptureProperties.g;
+			float Opacity = 0;
+
+			UNROLL
+			for( uint Lobe = 0; Lobe < NumLobes; Lobe++ )
+			{
+				float3 SampleDir = lerp( Lobes[ Lobe ].Direction, ProjectedCaptureVector, Lobes[ Lobe ].RayInfluance );
+				float4 Sample = ReflectionEnvironmentColorTexture.SampleLevel( ReflectionEnvironmentColorSampler, float4( SampleDir, CaptureArrayIndex ), LobeMip[ Lobe ] );
+
+				Sample.rgb *= CaptureProperties.r;
+				Sample *= DistanceAlpha;
+
+				// Under operator (back to front)
+				Lobes[ Lobe ].Color.rgb += Sample.rgb * Lobes[ Lobe ].Color.a;
+				Lobes[ Lobe ].Color.a *= 1 - Sample.a;
+
+				Opacity += Lobes[ Lobe ].Color.a;
+			}
+
+			BRANCH
+			if( Opacity < 0.001 )
+			{
+				break;
+			}
+		}
+	}
+
+#define APPLY_SKY_LIGHT 1
+#if APPLY_SKY_LIGHT
+	float Opacity = 0;
+
+	UNROLL
+	for( uint Lobe = 0; Lobe < NumLobes; Lobe++ )
+	{
+		Opacity += Lobes[ Lobe ].Color.a;
+	}
+
+	BRANCH
+	if( SkyLightParameters.y > 0 && Opacity >= 0.001 )
+	{
+		UNROLL
+		for( uint Lobe = 0; Lobe < NumLobes; Lobe++ )
+		{
+			// TODO use LobeMip
+			Lobes[ Lobe ].Color.rgb += Lobes[ Lobe ].Color.a * GetSkyLightReflection( Lobes[ Lobe ].Direction, Lobes[ Lobe ].Roughness, USE_LIGHTMAPS );
+		}
+	}
+#endif
+}
+
+
+float CountOverlap( float3 WorldPosition )
+{
+	float Overlap = 0;
+	float Opacity = 1;
+
+	uint NumCapturesAffectingTile = TileNumReflectionCaptures;
+
+	// Accumulate reflections from captures affecting this tile, applying largest captures first so that the smallest ones display on top
+	LOOP
+	for (uint TileCaptureIndex = 0; TileCaptureIndex < NumCapturesAffectingTile; TileCaptureIndex++) 
+	{
+		#if SORT_CAPTURES
+			uint CaptureIndex = SortedTileReflectionCaptureIndices[TileCaptureIndex];
+		#else
+			uint CaptureIndex = TileReflectionCaptureIndices[TileCaptureIndex];
+		#endif
+			
+		float4 CapturePositionAndRadius = ReflectionCapture.PositionAndRadius[CaptureIndex];
+		float3 CaptureVector = WorldPosition - CapturePositionAndRadius.xyz;
+		float CaptureVectorLength = sqrt(dot(CaptureVector, CaptureVector));
+
+		BRANCH
+		if (CaptureVectorLength < CapturePositionAndRadius.w)
+		{
+			float NormalizedDistanceToCapture = saturate(CaptureVectorLength / CapturePositionAndRadius.w);
+
+			// Fade out based on distance to capture
+			float DistanceAlpha = 1.0 - smoothstep(.6, 1, NormalizedDistanceToCapture);
+
+			Overlap += 1;
+			Opacity *= 1 - DistanceAlpha;
+
+			BRANCH
+			if( Opacity < 0.001 )
+			{
+				break;
+			}
+		}
+	}
+
+	return Overlap;
+}
+
+
 Texture2D ScreenSpaceReflections;
 Texture2D InSceneColor;
 
@@ -314,9 +594,6 @@ void ReflectionEnvironmentTiledDeferredMain(
 	float MaxTileZ2;
 	ComputeTileMinMax(ThreadIndex, SceneDepth, MinTileZ, MaxTileZ, MinTileZ2, MaxTileZ2);
 
-#define DO_CULLING_AND_SHADING 1
-#if DO_CULLING_AND_SHADING
-
 	DoTileCulling(GroupId, ThreadIndex, MinTileZ, MaxTileZ, MinTileZ2, MaxTileZ2);
 
 	// Lookup GBuffer properties once per pixel
@@ -327,228 +604,124 @@ void ReflectionEnvironmentTiledDeferredMain(
 	float3 WorldPosition = HomogeneousWorldPosition.xyz / HomogeneousWorldPosition.w;
 	float3 CameraToPixel = normalize(WorldPosition - View.ViewOrigin.xyz);
 	float3 ReflectionVector = reflect(CameraToPixel, GBuffer.WorldNormal);
-	float Overdraw = 0;
-
-	uint NumCapturesAffectingTile = TileNumReflectionCaptures;
-
-	float4 DiffuseLighting = float4(0, 0, 0, 1);
-	float4 SpecularLighting = float4(0, 0, 0, 1);
-
-	#define VISUALIZE_NUM_CAPTURES 0
-	#if VISUALIZE_NUM_CAPTURES 
-
-		SpecularLighting = NumCapturesAffectingTile / 5.0f;
-		SpecularLighting.a = 0;
-
-	#else
-
-		BRANCH
-		// Only light pixels marked as using deferred shading
-		if (GBuffer.ShadingModelID > 0)
-		{
-			// Accumulate reflections from captures affecting this tile, applying largest captures first so that the smallest ones display on top
-			LOOP
-			for (uint TileCaptureIndex = 0; TileCaptureIndex < NumCapturesAffectingTile; TileCaptureIndex++) 
-			{
-				#if SORT_CAPTURES
-					uint CaptureIndex = SortedTileReflectionCaptureIndices[TileCaptureIndex];
-				#else
-					uint CaptureIndex = TileReflectionCaptureIndices[TileCaptureIndex];
-				#endif
-			
-				float4 CapturePositionAndRadius = ReflectionCapture.PositionAndRadius[CaptureIndex];
-				float4 CaptureProperties = ReflectionCapture.CaptureProperties[CaptureIndex];
-				float3 CaptureVector = WorldPosition - CapturePositionAndRadius.xyz;
-				float CaptureVectorLength = sqrt(dot(CaptureVector, CaptureVector));
-
-				BRANCH
-				if (CaptureVectorLength < CapturePositionAndRadius.w)
-				{
-					float NormalizedDistanceToCapture = saturate(CaptureVectorLength / CapturePositionAndRadius.w);
-					float3 ProjectedCaptureVector = ReflectionVector;
-
-					// Fade out based on distance to capture
-					float DistanceAlpha = 0;
-
-					#define PROJECT_ONTO_SHAPE 1
-					#if PROJECT_ONTO_SHAPE
-					
-						#define SUPPORT_PLANE 0
-						#if SUPPORT_PLANE
-							BRANCH
-							// Plane
-							if (CaptureProperties.b > 1)
-							{
-								float4 ImagePlane = float4(ReflectionCapture.BoxTransform[CaptureIndex][0][0], ReflectionCapture.BoxTransform[CaptureIndex][1][0], ReflectionCapture.BoxTransform[CaptureIndex][2][0], ReflectionCapture.BoxTransform[CaptureIndex][3][0]);
-
-								float VectorDotPlaneNormal = dot(ImagePlane.xyz, ReflectionVector);
-								// VectorDotPlaneNormal < 0 means the ray hit the front face
-								BRANCH
-								if (VectorDotPlaneNormal < 0)
-								{
-									float PlaneDistance = dot(ImagePlane.xyz, WorldPosition) - ImagePlane.w;
-									// Time along the ray defined by WorldPosition + IntersectionTime * RayDirection that the intersection took place
-									float IntersectionTime = -PlaneDistance / VectorDotPlaneNormal;
-
-									BRANCH
-									// Skip intersections behind the pixel being shaded
-									if (IntersectionTime > 0)
-									{
-										// Calculate the world space intersection position
-										float3 IntersectPosition = WorldPosition + IntersectionTime * ReflectionVector;
-									
-										float2 ReflectionUVs;
-										float4 CurrentReflectionXAxis = float4(ReflectionCapture.BoxTransform[CaptureIndex][0][1], ReflectionCapture.BoxTransform[CaptureIndex][1][1], ReflectionCapture.BoxTransform[CaptureIndex][2][1], ReflectionCapture.BoxTransform[CaptureIndex][3][1]);
-										float3 CurrentImageReflectionOrigin = CapturePositionAndRadius.xyz;
-										float XLength = length(CurrentReflectionXAxis.xyz);
-										float3 NormalizedXAxis = CurrentReflectionXAxis.xyz / XLength;
-										// Calculate the quad UVs by projecting the vector from the intersection to the quad origin onto each quad axis
-										ReflectionUVs.x = dot(NormalizedXAxis, IntersectPosition - CurrentImageReflectionOrigin.xyz);
-										float3 ReflectionYAxis = cross(ImagePlane.xyz, NormalizedXAxis) * CurrentReflectionXAxis.w;
-										ReflectionUVs.y = dot(ReflectionYAxis, IntersectPosition - CurrentImageReflectionOrigin.xyz);
-										ReflectionUVs = .5f * ReflectionUVs / float2(XLength, XLength * CurrentReflectionXAxis.w) + .5f;
-
-										ProjectedCaptureVector = IntersectPosition - CapturePositionAndRadius.xyz;
-
-										if (ReflectionUVs.x > 0 && ReflectionUVs.x < 1 && ReflectionUVs.y > 0 && ReflectionUVs.y < 1)
-										{
-											CompositedLighting.rgb += 1;
-										}
-									}
-								}
-							}
-							else 
-						#endif
-						// Box
-						BRANCH if (CaptureProperties.b > 0)
-						{
-							float3 RayDirection = ReflectionVector;
-
-							// Transform the ray into the local space of the box, where it is an AABB with mins at -1 and maxs at 1
-							float3 LocalRayStart		= mul(float4(WorldPosition, 1), ReflectionCapture.BoxTransform[CaptureIndex]).xyz;
-							float3 LocalRayDirection	= mul(float4(RayDirection,  0), ReflectionCapture.BoxTransform[CaptureIndex]).xyz;
-
-							float3 InvRayDir = rcp(LocalRayDirection);
 	
-							//find the ray intersection with each of the 3 planes defined by the minimum extrema.
-							float3 FirstPlaneIntersections = -InvRayDir - LocalRayStart * InvRayDir;
-							//find the ray intersection with each of the 3 planes defined by the maximum extrema.
-							float3 SecondPlaneIntersections = InvRayDir - LocalRayStart * InvRayDir;
-							//get the furthest of these intersections along the ray
-							float3 FurthestPlaneIntersections = max(FirstPlaneIntersections, SecondPlaneIntersections);
-
-							//clamp the intersections to be between RayOrigin and RayEnd on the ray
-							float Intersection = min(FurthestPlaneIntersections.x, min(FurthestPlaneIntersections.y, FurthestPlaneIntersections.z));
-
-							// Compute the reprojected vector
-							float3 IntersectPosition = WorldPosition + Intersection * RayDirection;
-							ProjectedCaptureVector = IntersectPosition - CapturePositionAndRadius.xyz;
-
-							// Compute the distance from the receiving pixel to the box for masking
-							// Apply local to world scale to take scale into account without transforming back to world space
-							// Shrink the box by the transition distance (BoxScales.w) so that the fade happens inside the box influence area
-							float4 BoxScales = ReflectionCapture.BoxScales[CaptureIndex];
-							float BoxDistance = ComputeDistanceFromBoxToPoint(-(BoxScales.xyz - .5f * BoxScales.w), BoxScales.xyz - .5f * BoxScales.w, LocalRayStart * BoxScales.xyz);
-
-							// Setup a fade based on receiver distance to the box, hides the box influence shape
-							DistanceAlpha = 1.0 - smoothstep(0, .7f * BoxScales.w, BoxDistance);
-						}
-						// Sphere
-						else
-						{
-							float3 RayDirection = ReflectionVector;
-							float ProjectionSphereRadius = CapturePositionAndRadius.w * 1.2f;
-							float SphereRadiusSquared = ProjectionSphereRadius * ProjectionSphereRadius;
-
-							float3 ReceiverToSphereCenter = WorldPosition - CapturePositionAndRadius.xyz;
-							float ReceiverToSphereCenterSq = dot(ReceiverToSphereCenter, ReceiverToSphereCenter);
+	// Save GPRs by using R instead of V
+	float NoV = saturate( dot( GBuffer.WorldNormal, ReflectionVector ) );
 
-							// Find the intersection between the ray along the reflection vector and the capture's sphere
-							float3 QuadraticCoef;
-							QuadraticCoef.x = 1;
-							QuadraticCoef.y = 2 * dot(RayDirection, ReceiverToSphereCenter);
-							QuadraticCoef.z = ReceiverToSphereCenterSq - SphereRadiusSquared;
-
-							float Determinant = QuadraticCoef.y * QuadraticCoef.y - 4 * QuadraticCoef.z;
-							 
-							// Only continue if the ray intersects the sphere
-							if (Determinant >= 0)
-							{
-								float FarIntersection = (sqrt(Determinant) - QuadraticCoef.y) * 0.5;
-
-								float3 IntersectPosition = WorldPosition + FarIntersection * RayDirection;
-								ProjectedCaptureVector = IntersectPosition - CapturePositionAndRadius.xyz;
-								// Note: some compilers don't handle smoothstep min > max (this was 1, .6)
-								DistanceAlpha = 1.0 - smoothstep(.6, 1, NormalizedDistanceToCapture);
-							}
-						}
+	FLobe Lobes[2];
 
-					#else 
-						DistanceAlpha = 1.0;
-					#endif
+	UNROLL
+	for( uint Lobe = 0; Lobe < 2; Lobe++ )
+	{
+		Lobes[ Lobe ].Color			= float4(0, 0, 0, 1);
+		Lobes[ Lobe ].Direction		= ReflectionVector;
+		Lobes[ Lobe ].RayInfluance	= 1;
+		Lobes[ Lobe ].Roughness		= GBuffer.Roughness;
+	}
 
-					float CaptureArrayIndex = CaptureProperties.g;
-					
-#if !USE_LIGHTMAPS
-					float DiffuseMip = ComputeReflectionCaptureMipFromRoughness(1);
-					float4 DiffuseIBL = TextureCubeArraySampleLevel(ReflectionEnvironmentColorTexture, ReflectionEnvironmentColorSampler, GBuffer.WorldNormal, CaptureArrayIndex, DiffuseMip);
+#if VISUALIZE_OVERLAP
+	float Overlap = CountOverlap( WorldPosition );
+#else
 
-					DiffuseIBL.rgb *= CaptureProperties.r;
-					DiffuseIBL *= DistanceAlpha;
+#if 1//APPLY_SSR
+	float4 SSR = ScreenSpaceReflections.Load( int3(PixelPos, 0) );
+	Lobes[0].Color.rgb = SSR.rgb;
+	Lobes[0].Color.a = 1 - SSR.a;
+#endif
 
-					// Under operator (back to front)
-					DiffuseLighting.rgb += DiffuseIBL.rgb * DiffuseLighting.a;
-					DiffuseLighting.a *= 1 - DiffuseIBL.a;
+#if USE_LIGHTMAPS
+	// We have high frequency directional data but low frequency spatial data in the envmap.
+	// We have high frequency spatial data but low frequency directional data in the lightmap.
+	// So, we combine the two for the best of both. This is done by removing the low spatial frequencies from the envmap and replacing them with the lightmap data.
+	// This is only done with luma so as to not get odd color shifting.
+	// Note: make sure this matches the lightmap mixing done for translucency (BasePassPixelShader.usf)
+	Lobes[0].Color.a *= GBuffer.IndirectIrradiance;
+	Lobes[1].Color.a *= GBuffer.IndirectIrradiance;
+#else
+	// Diffuse lobe
+	Lobes[1].Direction		= GBuffer.WorldNormal;
+	Lobes[1].RayInfluance	= 0;
+	Lobes[1].Roughness		= 1;
 #endif
+	
+	float AO = ScreenSpaceData.AmbientOcclusion;
+	float SpecularOcclusion = saturate( Square( NoV + AO ) - 1 + AO );
 
-					float SpecularMip = ComputeReflectionCaptureMipFromRoughness(GBuffer.Roughness);
-					float4 SpecularIBL = TextureCubeArraySampleLevel(ReflectionEnvironmentColorTexture, ReflectionEnvironmentColorSampler, ProjectedCaptureVector, CaptureArrayIndex, SpecularMip);
+	Lobes[0].Color.a *= SpecularOcclusion;
+	Lobes[1].Color.a *= AO;
 
-					SpecularIBL.rgb *= CaptureProperties.r;
-					SpecularIBL *= DistanceAlpha;
+	BRANCH	//[forcecase]
+	switch( GBuffer.ShadingModelID )
+	{
+		case SHADINGMODELID_UNLIT:
+			break;
+		case SHADINGMODELID_DEFAULT_LIT:
+		case SHADINGMODELID_SUBSURFACE:
+		case SHADINGMODELID_PREINTEGRATED_SKIN:
+		{
+			BRANCH
+			if( GBuffer.ShadingModelID == SHADINGMODELID_SUBSURFACE )
+			{
+				GBuffer.DiffuseColor += DecodeSubsurfaceColor(GBuffer.CustomData);
+			}
 
-					// Under operator (back to front)
-					SpecularLighting.rgb += SpecularIBL.rgb * SpecularLighting.a;
-					SpecularLighting.a *= 1 - SpecularIBL.a;
+		#if USE_LIGHTMAPS
+			GBuffer.DiffuseColor = 0;
+		#endif
 
-					Overdraw += .05f;
+			GBuffer.SpecularColor = EnvBRDF( GBuffer.SpecularColor, GBuffer.Roughness, NoV );
 
-#if USE_LIGHTMAPS
-					bool Covered = SpecularLighting.a < 0.001;
-#else
-					bool Covered = DiffuseLighting.a + SpecularLighting.a < 0.001;
-#endif
-					BRANCH if( Covered )
-					{
-						break;
-					}
-				}
-			}
-		}
-	#endif
-#else
-	SpecularLighting = abs(MaxTileZ - MinTileZ) * .0001f;
-#endif
+			GatherRadiance( Lobes, USE_LIGHTMAPS ? 1 : 2, WorldPosition, ReflectionVector );
 
-#define APPLY_SKY_LIGHT 1
-#if APPLY_SKY_LIGHT
+			Lobes[0].Color.rgb *= GBuffer.SpecularColor;
+			Lobes[1].Color.rgb *= GBuffer.DiffuseColor;
+			break;
+		}
+		case SHADINGMODELID_CLEAR_COAT:
+#if 1
+		{
+			const float ClearCoat			= GBuffer.CustomData.x;
+			const float ClearCoatRoughness	= GBuffer.CustomData.y;
+
+		#if USE_LIGHTMAPS
+			Lobes[0].Roughness = ClearCoatRoughness;
+		#else
+			Lobes[0].Roughness = lerp( GBuffer.Roughness, ClearCoatRoughness, ClearCoat );
+			Lobes[1].Roughness = lerp( 1, GBuffer.Roughness, GBuffer.Metallic * ClearCoat );
+			Lobes[1].Direction = lerp( GBuffer.WorldNormal, ReflectionVector, ClearCoat );
+			Lobes[1].RayInfluance = ClearCoat;
+		#endif
 
-#if USE_LIGHTMAPS
-	bool Covered = SpecularLighting.a < 0.001;
-#else
-	bool Covered = DiffuseLighting.a + SpecularLighting.a < 0.001;
-#endif
+			// TODO EnvBRDF should have a mask param
+			float2 AB = PreIntegratedGF.SampleLevel( PreIntegratedGFSampler, float2( NoV, GBuffer.Roughness ), 0 ).rg;
+			GBuffer.SpecularColor = GBuffer.SpecularColor * AB.x + AB.y * saturate( 50 * GBuffer.SpecularColor.g ) * (1 - ClearCoat);
+			
+			// F_Schlick
+			float F0 = GBuffer.Specular * 0.08;
+			float Fc = pow( 1 - NoV, 5 );
+			float F = Fc + (1 - Fc) * F0;
+			F *= ClearCoat;
+			
+			float LayerAttenuation = (1 - F);
+
+		#if USE_LIGHTMAPS
+			float  Lobe0Reflectance = F;
+			float3 Lobe1Reflectance = GBuffer.SpecularColor * LayerAttenuation;
+		#else
+			float3 Lobe0Reflectance = lerp( GBuffer.SpecularColor, F, ClearCoat );
+			float3 Lobe1Reflectance = ( GBuffer.DiffuseColor + GBuffer.SpecularColor * ClearCoat ) * LayerAttenuation;
+		#endif
 
-	// Only light pixels marked as using deferred shading
-	BRANCH if (GBuffer.ShadingModelID > 0 && SkyLightParameters.y > 0 && !Covered)
-	{
-#if !USE_LIGHTMAPS
-		DiffuseLighting.rgb += DiffuseLighting.a * GetSkyLightReflection(GBuffer.WorldNormal, 1, USE_LIGHTMAPS);
+			GatherRadiance( Lobes, 2, WorldPosition, ReflectionVector );
+			
+			Lobes[0].Color.rgb *= Lobe0Reflectance;
+			Lobes[1].Color.rgb *= Lobe1Reflectance;
+			break;
+		}
 #endif
-		SpecularLighting.rgb += SpecularLighting.a * GetSkyLightReflection(ReflectionVector, GBuffer.Roughness, USE_LIGHTMAPS);
-		Overdraw += .05f;
+		default:
+			break;
 	}
-
 #endif
 
 	// Only write to the buffer for threads inside the view
@@ -557,49 +730,19 @@ void ReflectionEnvironmentTiledDeferredMain(
 	{
 		float4 OutColor = 0;
 
-		#define VISUALIZE_CAPTURE_OVERDRAW 0
-		#if VISUALIZE_CAPTURE_OVERDRAW
-			OutColor.rgb = Overdraw / 4;
-		#endif
-		
-		// Save GPRs by using R instead of V
-		float NoV = saturate( dot( GBuffer.WorldNormal, ReflectionVector ) );
-		float3 SpecularColor = EnvBRDF( GBuffer.SpecularColor, GBuffer.Roughness, NoV );
-		//float3 SpecularColor = EnvBRDFApprox( GBuffer.SpecularColor, GBuffer.Roughness, NoV );
-		
-		float AO = GBuffer.GBufferAO * ScreenSpaceData.AmbientOcclusion;
-		float SpecularOcclusion = saturate( Square( NoV + AO ) - 1 + AO );
-		SpecularColor *= SpecularOcclusion;
-
-#if USE_LIGHTMAPS
-		// We have high frequency directional data but low frequency spatial data in the envmap.
-		// We have high frequency spatial data but low frequency directional data in the lightmap.
-		// So, we combine the two for the best of both. This is done by removing the low spatial frequencies from the envmap and replacing them with the lightmap data.
-		// This is only done with luma so as to not get odd color shifting.
-		// Note: make sure this matches the lightmap mixing done for translucency (BasePassPixelShader.usf)
-		SpecularLighting.rgb *= GBuffer.IndirectIrradiance;
+#if VISUALIZE_OVERLAP
+		//OutColor.rgb = 0.1 * TileNumReflectionCaptures;
+		OutColor.rgb = 0.1 * Overlap;
 #else
-		BRANCH
-		if( GBuffer.ShadingModelID == SHADINGMODELID_SUBSURFACE )
-		{
-			GBuffer.DiffuseColor += DecodeSubsurfaceColor(GBuffer.CustomData);
-		}
-
-		OutColor.rgb += DiffuseLighting.rgb * GBuffer.DiffuseColor * AO;
-#endif
-
-#if 1//APPLY_SSR
-		float4 SSR = ScreenSpaceReflections.Load( int3(PixelPos, 0) );
-		SpecularLighting.rgb = SpecularLighting.rgb * (1 - SSR.a) + SSR.rgb;
+		OutColor.rgb += Lobes[0].Color.rgb;
+		OutColor.rgb += Lobes[1].Color.rgb;
 #endif
 
-		OutColor.rgb += SpecularLighting.rgb * SpecularColor;
-
 		// Transform NaNs to black, transform negative colors to black.
 		OutColor.rgb = -min(-OutColor.rgb, 0.0);
 
+		// TODO platforms that allow RW just add
 		OutColor.rgb += InSceneColor.Load( int3(PixelPos, 0) ).rgb;
-
 		RWOutSceneColor[PixelPos.xy] = OutColor;
 	}
 }
diff --git a/Engine/Shaders/ReflectionEnvironmentShaders.usf b/Engine/Shaders/ReflectionEnvironmentShaders.usf
--- a/Engine/Shaders/ReflectionEnvironmentShaders.usf
+++ b/Engine/Shaders/ReflectionEnvironmentShaders.usf
@@ -225,11 +225,7 @@ void ComputeBrightnessMain(
 
 #endif
 
-#ifndef NUM_FILTER_SAMPLES
-#define NUM_FILTER_SAMPLES 1
-#endif
-
-float FilterRoughness;
+#define NUM_FILTER_SAMPLES 1024
 
 Texture2D AverageBrightnessTexture;
 SamplerState AverageBrightnessSampler;
@@ -246,7 +242,7 @@ void FilterPS(
 #if FILTER_CUBEMAP
 
 	float3 N = normalize(CubeCoordinates);
-	float Roughness = FilterRoughness;
+	float Roughness = ComputeReflectionCaptureRoughnessFromMip( SourceMipIndex );
 
 	float4 FilteredColor = 0;
 	float Weight = 0;
diff --git a/Engine/Shaders/ReflectionEnvironmentShared.usf b/Engine/Shaders/ReflectionEnvironmentShared.usf
--- a/Engine/Shaders/ReflectionEnvironmentShared.usf
+++ b/Engine/Shaders/ReflectionEnvironmentShared.usf
@@ -15,13 +15,20 @@ half ComputeReflectionCaptureMipFromRoughness(half Roughness)
 	// Heuristic that maps roughness to mip level
 	// This is done in a way such that a certain mip level will always have the same roughness, regardless of how many mips are in the texture
 	// Using more mips in the cubemap just allows sharper reflections to be supported
-	// Note: this must match the logic in FilterReflectionEnvironment that generates the mip filter samples!
 	half LevelFrom1x1 = REFLECTION_CAPTURE_ROUGHEST_MIP - REFLECTION_CAPTURE_ROUGHNESS_MIP_SCALE * log2(Roughness);
 	// Note: must match GReflectionCaptureSize
-	half HardcodedNumCaptureArrayMips = 7;
+	const half HardcodedNumCaptureArrayMips = 7;
 	return HardcodedNumCaptureArrayMips - 1 - LevelFrom1x1;
 }
 
+float ComputeReflectionCaptureRoughnessFromMip( float Mip )
+{
+	// Note: must match GReflectionCaptureSize
+	const half HardcodedNumCaptureArrayMips = 7;
+	float LevelFrom1x1 = HardcodedNumCaptureArrayMips - 1 - Mip;
+	return exp2( ( REFLECTION_CAPTURE_ROUGHEST_MIP - LevelFrom1x1 ) / REFLECTION_CAPTURE_ROUGHNESS_MIP_SCALE );
+}
+
 TextureCube SkyLightCubemap;
 SamplerState SkyLightCubemapSampler;
 
diff --git a/Engine/Source/Runtime/Engine/Classes/Engine/EngineTypes.h b/Engine/Source/Runtime/Engine/Classes/Engine/EngineTypes.h
--- a/Engine/Source/Runtime/Engine/Classes/Engine/EngineTypes.h
+++ b/Engine/Source/Runtime/Engine/Classes/Engine/EngineTypes.h
@@ -111,10 +111,11 @@ enum ETrailWidthMode
 UENUM()
 enum EMaterialShadingModel
 {
-	MSM_DefaultLit			UMETA(DisplayName="Default Lit"),
 	MSM_Unlit				UMETA(DisplayName="Unlit"),
+	MSM_DefaultLit			UMETA(DisplayName="Default Lit"),
 	MSM_Subsurface			UMETA(DisplayName="Subsurface"),
 	MSM_PreintegratedSkin	UMETA(DisplayName="Preintegrated Skin"),
+	MSM_ClearCoat			UMETA(DisplayName="Clear Coat"),
 	MSM_MAX,
 };
 
diff --git a/Engine/Source/Runtime/Engine/Private/Materials/Material.cpp b/Engine/Source/Runtime/Engine/Private/Materials/Material.cpp
--- a/Engine/Source/Runtime/Engine/Private/Materials/Material.cpp
+++ b/Engine/Source/Runtime/Engine/Private/Materials/Material.cpp
@@ -468,6 +468,7 @@ UMaterial::UMaterial(const class FPostConstructInitializeProperties& PCIP)
 	: Super(PCIP)
 {
 	BlendMode = BLEND_Opaque;
+	ShadingModel = MSM_DefaultLit;
 	TranslucencyLightingMode = TLM_VolumetricNonDirectional;
 	TranslucencyDirectionalLightingIntensity = 1.0f;
 	TranslucentShadowDensityScale = 0.5f;
diff --git a/Engine/Source/Runtime/Engine/Private/Materials/MaterialShader.cpp b/Engine/Source/Runtime/Engine/Private/Materials/MaterialShader.cpp
--- a/Engine/Source/Runtime/Engine/Private/Materials/MaterialShader.cpp
+++ b/Engine/Source/Runtime/Engine/Private/Materials/MaterialShader.cpp
@@ -35,10 +35,11 @@ FString GetShadingModelString(EMaterialShadingModel ShadingModel)
 	FString ShadingModelName;
 	switch(ShadingModel)
 	{
-		case MSM_DefaultLit: ShadingModelName = TEXT("MSM_DefaultLit"); break;
-		case MSM_Unlit: ShadingModelName = TEXT("MSM_Unlit"); break;
-		case MSM_Subsurface: ShadingModelName = TEXT("MSM_Subsurface"); break;
-		case MSM_PreintegratedSkin: ShadingModelName = TEXT("MSM_PreintegratedSkin"); break;
+		case MSM_Unlit:				ShadingModelName = TEXT("MSM_Unlit"); break;
+		case MSM_DefaultLit:		ShadingModelName = TEXT("MSM_DefaultLit"); break;
+		case MSM_Subsurface:		ShadingModelName = TEXT("MSM_Subsurface"); break;
+		case MSM_PreintegratedSkin:	ShadingModelName = TEXT("MSM_PreintegratedSkin"); break;
+		case MSM_ClearCoat:			ShadingModelName = TEXT("MSM_ClearCoat"); break;
 		default: ShadingModelName = TEXT("Unknown"); break;
 	}
 	return ShadingModelName;
@@ -74,10 +75,15 @@ void UpdateMaterialShaderCompilingStats(const FMaterial* Material)
 
 	switch(Material->GetShadingModel())
 	{
-		case MSM_Subsurface: 
-		case MSM_PreintegratedSkin: 
-		case MSM_DefaultLit: INC_DWORD_STAT_BY(STAT_ShaderCompiling_NumLitMaterialShaders,1); break;
-		case MSM_Unlit: INC_DWORD_STAT_BY(STAT_ShaderCompiling_NumUnlitMaterialShaders,1); break;
+		case MSM_Unlit:
+			INC_DWORD_STAT_BY(STAT_ShaderCompiling_NumUnlitMaterialShaders,1);
+			break;
+		case MSM_DefaultLit:
+		case MSM_Subsurface:
+		case MSM_PreintegratedSkin:
+		case MSM_ClearCoat:
+			INC_DWORD_STAT_BY(STAT_ShaderCompiling_NumLitMaterialShaders,1);
+			break;
 		default: break;
 	};
 
diff --git a/Engine/Source/Runtime/Engine/Private/Materials/MaterialShared.cpp b/Engine/Source/Runtime/Engine/Private/Materials/MaterialShared.cpp
--- a/Engine/Source/Runtime/Engine/Private/Materials/MaterialShared.cpp
+++ b/Engine/Source/Runtime/Engine/Private/Materials/MaterialShared.cpp
@@ -1183,14 +1183,14 @@ void FMaterial::SetupMaterialEnvironment(
 
 	switch(GetShadingModel())
 	{
-	case MSM_DefaultLit: OutEnvironment.SetDefine(TEXT("MATERIAL_SHADINGMODEL_DEFAULT_LIT"),TEXT("1")); break;
-	case MSM_Subsurface: OutEnvironment.SetDefine(TEXT("MATERIAL_SHADINGMODEL_SUBSURFACE"),TEXT("1")); break;
-	case MSM_PreintegratedSkin: OutEnvironment.SetDefine(TEXT("MATERIAL_SHADINGMODEL_PREINTEGRATED_SKIN"),TEXT("1")); break;
-	case MSM_Unlit: OutEnvironment.SetDefine(TEXT("MATERIAL_SHADINGMODEL_UNLIT"),TEXT("1")); break;
-	default: 
-		UE_LOG(LogMaterial, Warning, TEXT("Unknown material shading model: %u  Setting to MSM_DefaultLit"),(int32)GetShadingModel());
-		OutEnvironment.SetDefine(TEXT("MATERIAL_SHADINGMODEL_DEFAULT_LIT"),TEXT("1"));
-	
+		case MSM_Unlit:				OutEnvironment.SetDefine(TEXT("MATERIAL_SHADINGMODEL_UNLIT"),				TEXT("1")); break;
+		case MSM_DefaultLit:		OutEnvironment.SetDefine(TEXT("MATERIAL_SHADINGMODEL_DEFAULT_LIT"),			TEXT("1")); break;
+		case MSM_Subsurface:		OutEnvironment.SetDefine(TEXT("MATERIAL_SHADINGMODEL_SUBSURFACE"),			TEXT("1")); break;
+		case MSM_PreintegratedSkin: OutEnvironment.SetDefine(TEXT("MATERIAL_SHADINGMODEL_PREINTEGRATED_SKIN"),	TEXT("1")); break;
+		case MSM_ClearCoat:			OutEnvironment.SetDefine(TEXT("MATERIAL_SHADINGMODEL_CLEAR_COAT"),			TEXT("1")); break;
+		default: 
+			UE_LOG(LogMaterial, Warning, TEXT("Unknown material shading model: %u  Setting to MSM_DefaultLit"),(int32)GetShadingModel());
+			OutEnvironment.SetDefine(TEXT("MATERIAL_SHADINGMODEL_DEFAULT_LIT"),TEXT("1"));
 	};
 
 	if (IsTranslucentBlendMode(GetBlendMode()))
diff --git a/Engine/Source/Runtime/Renderer/Private/LightRendering.h b/Engine/Source/Runtime/Renderer/Private/LightRendering.h
--- a/Engine/Source/Runtime/Renderer/Private/LightRendering.h
+++ b/Engine/Source/Runtime/Renderer/Private/LightRendering.h
@@ -22,6 +22,7 @@ BEGIN_UNIFORM_BUFFER_STRUCT(FDeferredLightUniformStruct,)
 	DECLARE_UNIFORM_BUFFER_STRUCT_MEMBER(float,MinRoughness)
 	DECLARE_UNIFORM_BUFFER_STRUCT_MEMBER(FVector2D,DistanceFadeMAD)
 	DECLARE_UNIFORM_BUFFER_STRUCT_MEMBER(FVector4,ShadowMapChannelMask)
+	DECLARE_UNIFORM_BUFFER_STRUCT_MEMBER(uint32,bShadowed)
 END_UNIFORM_BUFFER_STRUCT(FDeferredLightUniformStruct)
 
 extern float GMinScreenRadiusForLights;
@@ -56,7 +57,7 @@ void SetDeferredLightParameters(
 
 	const FVector2D FadeParams = LightSceneInfo->Proxy->GetDirectionalLightDistanceFadeParameters();
 
-	// use MAD for efficieny in the shader
+	// use MAD for efficiency in the shader
 	DeferredLightUniformsValue.DistanceFadeMAD = FVector2D(FadeParams.Y, -FadeParams.X * FadeParams.Y);
 
 	int32 ShadowMapChannel = LightSceneInfo->Proxy->GetShadowMapChannel();
@@ -75,6 +76,14 @@ void SetDeferredLightParameters(
 		ShadowMapChannel == 2 ? 1 : 0,
 		ShadowMapChannel == 3 ? 1 : 0);
 
+	DeferredLightUniformsValue.bShadowed = LightSceneInfo->Proxy->CastsDynamicShadow() || LightSceneInfo->Proxy->CastsStaticShadow();
+
+	if( LightSceneInfo->Proxy->IsInverseSquared() )
+	{
+		// Correction for lumen units
+		DeferredLightUniformsValue.LightColor *= 16.0f;
+	}
+
 	// When rendering reflection captures, the direct lighting of the light is actually the indirect specular from the main view
 	if (View.bIsReflectionCapture)
 	{
diff --git a/Engine/Source/Runtime/Renderer/Private/ReflectionEnvironmentCapture.cpp b/Engine/Source/Runtime/Renderer/Private/ReflectionEnvironmentCapture.cpp
--- a/Engine/Source/Runtime/Renderer/Private/ReflectionEnvironmentCapture.cpp
+++ b/Engine/Source/Runtime/Renderer/Private/ReflectionEnvironmentCapture.cpp
@@ -130,8 +130,6 @@ private:
 IMPLEMENT_SHADER_TYPE(,FDownsamplePS,TEXT("ReflectionEnvironmentShaders"),TEXT("DownsamplePS"),SF_Pixel);
 
 
-const int32 NumFilterSamples = 1024;
-
 /** Pixel shader used for filtering a mip. */
 class FCubeFilterPS : public FDownsamplePS
 {
@@ -146,24 +144,20 @@ public:
 	static void ModifyCompilationEnvironment(EShaderPlatform Platform, FShaderCompilerEnvironment& OutEnvironment)
 	{
 		FDownsamplePS::ModifyCompilationEnvironment(Platform, OutEnvironment);
-		OutEnvironment.SetDefine(TEXT("NUM_FILTER_SAMPLES"), NumFilterSamples);
 	}
 
 	FCubeFilterPS(const ShaderMetaType::CompiledShaderInitializerType& Initializer):
 		FDownsamplePS(Initializer)
 	{
-		FilterRoughness.Bind(Initializer.ParameterMap,TEXT("FilterRoughness"));
 		AverageBrightnessTexture.Bind(Initializer.ParameterMap,TEXT("AverageBrightnessTexture"));
 		AverageBrightnessSampler.Bind(Initializer.ParameterMap,TEXT("AverageBrightnessSampler"));
 	}
 	FCubeFilterPS() {}
 
-	void SetParameters(FRHICommandList& RHICmdList, int32 CubeFaceValue, int32 SourceMipIndexValue, float RoughnessValue, FSceneRenderTargetItem& SourceTextureValue)
+	void SetParameters(FRHICommandList& RHICmdList, int32 CubeFaceValue, int32 SourceMipIndexValue, FSceneRenderTargetItem& SourceTextureValue)
 	{
 		FDownsamplePS::SetParameters(RHICmdList, CubeFaceValue, SourceMipIndexValue, SourceTextureValue);
 
-		SetShaderValue(RHICmdList, GetPixelShader(), FilterRoughness, RoughnessValue);
-
 		SetTextureParameter(
 			RHICmdList,
 			GetPixelShader(), 
@@ -176,14 +170,12 @@ public:
 	virtual bool Serialize(FArchive& Ar)
 	{
 		bool bShaderHasOutdatedParameters = FDownsamplePS::Serialize(Ar);
-		Ar << FilterRoughness;
 		Ar << AverageBrightnessTexture;
 		Ar << AverageBrightnessSampler;
 		return bShaderHasOutdatedParameters;
 	}
 
 private:
-	FShaderParameter FilterRoughness;
 	FShaderResourceParameter AverageBrightnessTexture;
 	FShaderResourceParameter AverageBrightnessSampler;
 };
@@ -339,39 +331,39 @@ void FilterReflectionEnvironment(FRHICommandList& RHICmdList, FSHVectorRGB3* Out
 			{
 				SetRenderTarget(RHICmdList, EffectiveRT.TargetableTexture, MipIndex, NULL);
 
-			const FIntRect ViewRect(0, 0, MipSize, MipSize);
+				const FIntRect ViewRect(0, 0, MipSize, MipSize);
 				RHICmdList.SetViewport(0, 0, 0.0f, MipSize, MipSize, 1.0f);
 				RHICmdList.SetRasterizerState(TStaticRasterizerState<FM_Solid, CM_None>::GetRHI());
 				RHICmdList.SetDepthStencilState(TStaticDepthStencilState<false, CF_Always>::GetRHI());
 				RHICmdList.SetBlendState(TStaticBlendState<>::GetRHI());
 
-			TShaderMapRef<FScreenVSForGS> VertexShader(GetGlobalShaderMap());
-			TShaderMapRef<FDownsampleGS> GeometryShader(GetGlobalShaderMap());
+				TShaderMapRef<FScreenVSForGS> VertexShader(GetGlobalShaderMap());
+				TShaderMapRef<FDownsampleGS> GeometryShader(GetGlobalShaderMap());
 				TShaderMapRef<FDownsamplePS> PixelShader(GetGlobalShaderMap());
 
 				SetGlobalBoundShaderState(RHICmdList, DownsampleBoundShaderState, GFilterVertexDeclaration.VertexDeclarationRHI, *VertexShader, *PixelShader, *GeometryShader);
 
-			// Draw each face with a geometry shader that routes to the correct render target slice
-			for (int32 CubeFace = 0; CubeFace < CubeFace_MAX; CubeFace++)
-			{
-				SCOPED_DRAW_EVENT(DownsampleCubeFace, DEC_SCENE_ITEMS);
-				GeometryShader->SetParameters(RHICmdList, CubeFace);
+				// Draw each face with a geometry shader that routes to the correct render target slice
+				for (int32 CubeFace = 0; CubeFace < CubeFace_MAX; CubeFace++)
+				{
+					SCOPED_DRAW_EVENT(DownsampleCubeFace, DEC_SCENE_ITEMS);
+					GeometryShader->SetParameters(RHICmdList, CubeFace);
 
-				PixelShader->SetParameters(RHICmdList, CubeFace, SourceMipIndex, EffectiveSource);
+					PixelShader->SetParameters(RHICmdList, CubeFace, SourceMipIndex, EffectiveSource);
 
-				DrawRectangle( 
+					DrawRectangle( 
 						RHICmdList,
-					ViewRect.Min.X, ViewRect.Min.Y, 
-					ViewRect.Width(), ViewRect.Height(),
-					ViewRect.Min.X, ViewRect.Min.Y, 
-					ViewRect.Width(), ViewRect.Height(),
-					FIntPoint(ViewRect.Width(), ViewRect.Height()),
-					FIntPoint(MipSize, MipSize),
-					*VertexShader,
-					EDRF_UseTriangleOptimization);
+						ViewRect.Min.X, ViewRect.Min.Y, 
+						ViewRect.Width(), ViewRect.Height(),
+						ViewRect.Min.X, ViewRect.Min.Y, 
+						ViewRect.Width(), ViewRect.Height(),
+						FIntPoint(ViewRect.Width(), ViewRect.Height()),
+						FIntPoint(MipSize, MipSize),
+						*VertexShader,
+						EDRF_UseTriangleOptimization);
 
 					RHICmdList.CopyToResolveTarget(EffectiveRT.TargetableTexture, EffectiveRT.ShaderResourceTexture, true, FResolveParams(FResolveRect(), (ECubeFace)CubeFace, MipIndex));
-			}
+				}
 			}
 			else
 			{
@@ -434,21 +426,18 @@ void FilterReflectionEnvironment(FRHICommandList& RHICmdList, FSHVectorRGB3* Out
 			check(EffectiveRT.TargetableTexture != EffectiveSource.ShaderResourceTexture);
 			const int32 MipSize = 1 << (NumMips - MipIndex - 1);
 
-			// Note: this must match the logic in ComputeReflectionCaptureMipFromRoughness!
-			float Roughness = FMath::Pow(2, (MipIndex - NumMips + 1 + ReflectionCaptureRoughestMip) / ReflectionCaptureRoughnessMipScale);
-
 			if (GSupportsGSRenderTargetLayerSwitchingToMips)
 			{
 				SetRenderTarget(RHICmdList, EffectiveRT.TargetableTexture, MipIndex, NULL);
 
-			const FIntRect ViewRect(0, 0, MipSize, MipSize);
+				const FIntRect ViewRect(0, 0, MipSize, MipSize);
 				RHICmdList.SetViewport(0, 0, 0.0f, MipSize, MipSize, 1.0f);
 				RHICmdList.SetRasterizerState(TStaticRasterizerState<FM_Solid, CM_None>::GetRHI());
 				RHICmdList.SetDepthStencilState(TStaticDepthStencilState<false, CF_Always>::GetRHI());
 				RHICmdList.SetBlendState(TStaticBlendState<>::GetRHI());
 
-			TShaderMapRef<FScreenVSForGS> VertexShader(GetGlobalShaderMap());
-			TShaderMapRef<FDownsampleGS> GeometryShader(GetGlobalShaderMap());
+				TShaderMapRef<FScreenVSForGS> VertexShader(GetGlobalShaderMap());
+				TShaderMapRef<FDownsampleGS> GeometryShader(GetGlobalShaderMap());
 
 				FCubeFilterPS* PixelShader;
 				if( bNormalize )
@@ -466,27 +455,27 @@ void FilterReflectionEnvironment(FRHICommandList& RHICmdList, FSHVectorRGB3* Out
 					SetGlobalBoundShaderState(RHICmdList, BoundShaderState, GFilterVertexDeclaration.VertexDeclarationRHI, *VertexShader, PixelShader, *GeometryShader);
 				}
 
-			for (int32 CubeFace = 0; CubeFace < CubeFace_MAX; CubeFace++)
-			{
-				SCOPED_DRAW_EVENT(FilterCubeFace, DEC_SCENE_ITEMS);
-				GeometryShader->SetParameters(RHICmdList, CubeFace);
+				for (int32 CubeFace = 0; CubeFace < CubeFace_MAX; CubeFace++)
+				{
+					SCOPED_DRAW_EVENT(FilterCubeFace, DEC_SCENE_ITEMS);
+					GeometryShader->SetParameters(RHICmdList, CubeFace);
 
-				PixelShader->SetParameters(RHICmdList, CubeFace, MipIndex, Roughness, EffectiveSource);
+					PixelShader->SetParameters(RHICmdList, CubeFace, MipIndex, EffectiveSource);
 
-				DrawRectangle( 
-					RHICmdList,
-					ViewRect.Min.X, ViewRect.Min.Y, 
-					ViewRect.Width(), ViewRect.Height(),
-					ViewRect.Min.X, ViewRect.Min.Y, 
-					ViewRect.Width(), ViewRect.Height(),
-					FIntPoint(ViewRect.Width(), ViewRect.Height()),
-					FIntPoint(MipSize, MipSize),
-					*VertexShader,
-					EDRF_UseTriangleOptimization);
+					DrawRectangle( 
+						RHICmdList,
+						ViewRect.Min.X, ViewRect.Min.Y, 
+						ViewRect.Width(), ViewRect.Height(),
+						ViewRect.Min.X, ViewRect.Min.Y, 
+						ViewRect.Width(), ViewRect.Height(),
+						FIntPoint(ViewRect.Width(), ViewRect.Height()),
+						FIntPoint(MipSize, MipSize),
+						*VertexShader,
+						EDRF_UseTriangleOptimization);
 
-				RHICmdList.CopyToResolveTarget(EffectiveRT.TargetableTexture, EffectiveRT.ShaderResourceTexture, true, FResolveParams(FResolveRect(), (ECubeFace)CubeFace, MipIndex));
+					RHICmdList.CopyToResolveTarget(EffectiveRT.TargetableTexture, EffectiveRT.ShaderResourceTexture, true, FResolveParams(FResolveRect(), (ECubeFace)CubeFace, MipIndex));
+				}
 			}
-		}
 			else
 			{
 				for (int32 CubeFace = 0; CubeFace < CubeFace_MAX; CubeFace++)
@@ -519,7 +508,7 @@ void FilterReflectionEnvironment(FRHICommandList& RHICmdList, FSHVectorRGB3* Out
 						SetGlobalBoundShaderState(RHICmdList, BoundShaderState, GFilterVertexDeclaration.VertexDeclarationRHI, *VertexShader, PixelShader);
 					}
 
-					PixelShader->SetParameters(RHICmdList, CubeFace, MipIndex, Roughness, EffectiveSource);
+					PixelShader->SetParameters(RHICmdList, CubeFace, MipIndex, EffectiveSource);
 
 					DrawRectangle( 
 						RHICmdList,
diff --git a/Engine/Source/Runtime/Renderer/Private/TiledDeferredLightRendering.cpp b/Engine/Source/Runtime/Renderer/Private/TiledDeferredLightRendering.cpp
--- a/Engine/Source/Runtime/Renderer/Private/TiledDeferredLightRendering.cpp
+++ b/Engine/Source/Runtime/Renderer/Private/TiledDeferredLightRendering.cpp
@@ -167,6 +167,10 @@ public:
 
 				if (LightSceneInfo->Proxy->IsInverseSquared())
 				{
+					// Correction for lumen units
+					LightData.LightColorAndFalloffExponent[LightIndex].X *= 16.0f;
+					LightData.LightColorAndFalloffExponent[LightIndex].Y *= 16.0f;
+					LightData.LightColorAndFalloffExponent[LightIndex].Z *= 16.0f;
 					LightData.LightColorAndFalloffExponent[LightIndex].W = 0;
 				}
 
@@ -202,6 +206,12 @@ public:
 				LightData2.LightDirectionAndSpotlightMaskAndMinRoughness[LightIndex] = FVector4(FVector(1, 0, 0), 0);
 				LightData2.SpotAnglesAndSourceRadiusAndSimpleLighting[LightIndex] = FVector4(-2, 1, 0, 1);
 				LightData2.ShadowMapChannelMask[LightIndex] = FVector4(0, 0, 0, 0);
+
+				if( SimpleLight.Exponent == 0.0f )
+				{
+					// Correction for lumen units
+					LightData.LightColorAndFalloffExponent[LightIndex] *= 16.0f;
+				}
 			}
 		}
 
