Fixed clear coat intensity

diff --git a/Engine/Shaders/DeferredLightingCommon.usf b/Engine/Shaders/DeferredLightingCommon.usf
--- a/Engine/Shaders/DeferredLightingCommon.usf
+++ b/Engine/Shaders/DeferredLightingCommon.usf
@@ -137,7 +137,7 @@ float3 ClearCoatShading( FGBufferData GBuffer, float3 LobeRoughness, float3 Lobe
 	float Vis = Vis_Kelemen( VoH );
 	
 	// F_Schlick
-	float F0 = GBuffer.Specular * 0.08;
+	float F0 = 0.04;
 	float Fc = pow( 1 - VoH, 5 );
 	float F = Fc + (1 - Fc) * F0;
 	F *= ClearCoat;
@@ -169,7 +169,7 @@ float3 ClearCoatShading( FGBufferData GBuffer, float3 LobeRoughness, float3 Lobe
 	float Vis = Vis_Kelemen( VoH );
 	
 	// F_Schlick
-	float F0 = GBuffer.Specular * 0.08;
+	float F0 = 0.04;
 	float Fc = pow( 1 - VoH, 5 );
 	float F = Fc + (1 - Fc) * F0;
 
diff --git a/Engine/Shaders/PostProcessAmbient.usf b/Engine/Shaders/PostProcessAmbient.usf
--- a/Engine/Shaders/PostProcessAmbient.usf
+++ b/Engine/Shaders/PostProcessAmbient.usf
@@ -200,14 +200,14 @@ float3 ClearCoatShadingPerLight( Shared, float3 L, float3 V, half3 N )
 
 float3 ClearCoatShading( FGBufferData GBuffer, float Roughness, float3 L, float3 V, half3 N )
 {
-	float ClearCoat = 1;
-	float ClearCoatRoughness = 0.1;
+	const float ClearCoat			= GBuffer.CustomData.x;
+	const float ClearCoatRoughness	= GBuffer.CustomData.y;
 	float Film = 1;
 	float MetalSpec = 0.9;
 
 	Film *= ClearCoat;
 
-#if 0
+#if 1
 	float3 H = normalize(V + L);
 	float NoL = saturate( dot(N, L) );
 	float NoV = saturate( dot(N, V) );
@@ -219,7 +219,7 @@ float3 ClearCoatShading( FGBufferData GBuffer, float Roughness, float3 L, float3
 	float Vis = Vis_Kelemen( VoH );
 	
 	// F_Schlick
-	float F0 = GBuffer.Specular * 0.08;
+	float F0 = 0.04;
 	float Fc = pow( 1 - VoH, 5 );
 	float F = Fc + (1 - Fc) * F0;
 
@@ -257,7 +257,7 @@ float3 ClearCoatShading( FGBufferData GBuffer, float Roughness, float3 L, float3
 	float Vis = Vis_Kelemen( VoH );
 	
 	// F_Schlick
-	float F0 = GBuffer.Specular * 0.08;
+	float F0 = 0.04;
 	float Fc = pow( 1 - VoH, 5 );
 	float F = Fc + (1 - Fc) * F0;
 
@@ -305,7 +305,7 @@ float3 ClearCoatShading( FGBufferData GBuffer, float Roughness, float3 L, float3
 	float Vis = Vis_Kelemen( VoH );
 	
 	// F_Schlick
-	float F0 = GBuffer.Specular * 0.08;
+	float F0 = 0.04;
 	float Fc = pow( 1 - VoH, 5 );
 	float F = Fc + (1 - Fc) * F0;
 
@@ -430,8 +430,8 @@ float3 ImageBasedLightingMIS( FGBufferData GBuffer, float3 V, float3 N, uint2 Ra
 				}
 				float Weight = rcp( InvWeight );
 
-				//float3 Shading = ClearCoatShading( GBuffer, GBuffer.Roughness, L, V, N );
-				float3 Shading = StandardShading( GBuffer, GBuffer.Roughness, L, V, N );
+				float3 Shading = ClearCoatShading( GBuffer, GBuffer.Roughness, L, V, N );
+				//float3 Shading = StandardShading( GBuffer, GBuffer.Roughness, L, V, N );
 			
 				Lighting += SampleColor * Shading * ( NoL * Weight );
 			}
@@ -548,8 +548,8 @@ float3 ImageBasedLightingSplitSum( FGBufferData GBuffer, float3 V, float3 N, uin
 		return 0;
 	}
 
-	float ClearCoat = 1;
-	float ClearCoatRoughness = 0.1;
+	const float ClearCoat			= GBuffer.CustomData.x;
+	const float ClearCoatRoughness	= GBuffer.CustomData.y;
 	float Film = 1;
 	float MetalSpec = 0.9;
 
@@ -560,7 +560,7 @@ float3 ImageBasedLightingSplitSum( FGBufferData GBuffer, float3 V, float3 N, uin
 	float Lobe1Roughness = ClearCoatRoughness;
 	float Lobe2Roughness = GBuffer.Roughness;
 	
-#if 1
+#if 0
 	float3 Lobe1Color = FilterEnvMap( Random, Lobe1Roughness, N, V );
 	float3 Lobe2Color = FilterEnvMap( Random, Lobe2Roughness, N, V );
 #else
@@ -576,7 +576,7 @@ float3 ImageBasedLightingSplitSum( FGBufferData GBuffer, float3 V, float3 N, uin
 #endif
 
 	// F_Schlick
-	float F0 = GBuffer.Specular * 0.08;
+	float F0 = 0.04;
 	float Fc = pow( 1 - NoV, 5 );
 	float F = Fc + (1 - Fc) * F0;
 	F *= ClearCoat;
@@ -625,7 +625,7 @@ void MainPS(in float4 UVAndScreenPos : TEXCOORD0, out float4 OutColor : SV_Targe
 	// Diffuse
 	{
 	#if IMPORTANCE_SAMPLE
-		//DiffuseAndSpecular += DiffuseIBL( Random, GBuffer.DiffuseColor, GBuffer.WorldNormal );
+		DiffuseAndSpecular += DiffuseIBL( Random, GBuffer.DiffuseColor, GBuffer.WorldNormal );
 	#else
 		// we want to access the mip with the preconvolved diffuse lighting (coneangle=90 degree)
 		DiffuseAndSpecular += GBuffer.DiffuseColor * DiffuseLookup;
@@ -635,7 +635,7 @@ void MainPS(in float4 UVAndScreenPos : TEXCOORD0, out float4 OutColor : SV_Targe
 	// Specular
 	{
 	#if IMPORTANCE_SAMPLE
-		//DiffuseAndSpecular += SpecularIBL( Random, GBuffer.SpecularColor, GBuffer.Roughness, GBuffer.WorldNormal, -ScreenVector );
+		DiffuseAndSpecular += SpecularIBL( Random, GBuffer.SpecularColor, GBuffer.Roughness, GBuffer.WorldNormal, -ScreenVector );
 		//DiffuseAndSpecular += ApproximateSpecularIBL( Random, GBuffer.SpecularColor, GBuffer.Roughness, GBuffer.WorldNormal, -ScreenVector );
 	#else
 		float3 R = reflect( ScreenVector, GBuffer.WorldNormal );
@@ -648,13 +648,13 @@ void MainPS(in float4 UVAndScreenPos : TEXCOORD0, out float4 OutColor : SV_Targe
 	#endif
 	}
 
-	#if IMPORTANCE_SAMPLE
-		DiffuseAndSpecular += ImageBasedLightingMIS( GBuffer, -ScreenVector, GBuffer.WorldNormal, Random );
-	#endif
-
 	BRANCH if( GBuffer.ShadingModelID == SHADINGMODELID_CLEAR_COAT )
 	{
+	#if IMPORTANCE_SAMPLE
+		DiffuseAndSpecular = ImageBasedLightingMIS( GBuffer, -ScreenVector, GBuffer.WorldNormal, Random );
+	#else
 		DiffuseAndSpecular = ImageBasedLightingSplitSum( GBuffer, -ScreenVector, GBuffer.WorldNormal, Random );
+	#endif
 	}
 	
 	// apply darkening from ambient occlusion (does not use PostprocessInput1 to set white texture if SSAO is off)
diff --git a/Engine/Shaders/PostProcessTonemap.usf b/Engine/Shaders/PostProcessTonemap.usf
--- a/Engine/Shaders/PostProcessTonemap.usf
+++ b/Engine/Shaders/PostProcessTonemap.usf
@@ -8,6 +8,9 @@
 #include "PostProcessCommon.usf"
 #include "TonemapCommon.usf"
 
+#include "BRDF.usf"
+#include "MonteCarlo.usf"
+
 //
 // 64x64 tiling noise texture, optimized to have mostly high frequency content
 Texture2D NoiseTexture;
@@ -149,25 +152,242 @@ void MainVS(
 	OutFringe.zw = ScreenPosToUV(ScreenPos * FringeUVParams.g);
 }
 
+
+
+float2 IntegrateBRDF( uint2 Random, float Roughness, float NoV )
+{
+	float F0 = 0.04;
+	float Fc = pow( 1 - NoV, 5 );
+	float F = Fc + (1 - Fc) * F0;
+
+	float3 V;
+	V.x = sqrt( 1.0f - NoV * NoV );	// sin
+	V.y = 0;
+	V.z = NoV;						// cos
+
+	float A = 0;
+	float B = 0;
+
+	const uint NumSamples = 16;
+	for( uint i = 0; i < NumSamples; i++ )
+	{
+		float2 E = Hammersley( i, NumSamples, Random );
+		float3 H = ImportanceSampleGGX( E, Roughness ).xyz;
+		float3 L = 2 * dot( V, H ) * H - V;
+
+		float NoL = saturate( L.z );
+		float NoH = saturate( H.z );
+		float VoH = saturate( dot( V, H ) );
+
+		if( NoL > 0 )
+		{
+			float Vis = Vis_Schlick( Roughness, NoV, NoL );
+
+			// Incident light = NoL
+			// pdf = D * NoH / (4 * VoH)
+			// NoL * Vis / pdf
+			float NoL_Vis_PDF = NoL * Vis * (4 * VoH / NoH);
+			NoL_Vis_PDF *= (1 - F);
+
+			float Fc = pow( 1 - VoH, 5 );
+			A += (1 - Fc) * NoL_Vis_PDF;
+			B += Fc * NoL_Vis_PDF;
+		}
+	}
+
+	return float2( A, B ) / NumSamples;
+}
+
+float2 IntegrateClearCoat( uint2 Random, float Roughness, float NoV )
+{
+	float3 V;
+	V.x = sqrt( 1.0f - NoV * NoV );	// sin
+	V.y = 0;
+	V.z = NoV;						// cos
+
+	float A = 0;
+	float B = 0;
+
+	const uint NumSamples = 16;
+	for( uint i = 0; i < NumSamples; i++ )
+	{
+		float2 E = Hammersley( i, NumSamples, Random );
+		float3 H = ImportanceSampleGGX( E, Roughness ).xyz;
+		float3 L = 2 * dot( V, H ) * H - V;
+
+		float NoL = saturate( L.z );
+		float NoH = saturate( H.z );
+		float VoH = saturate( dot( V, H ) );
+
+		float RefractBlend = (0.22 * VoH + 0.7) * VoH + 0.745;
+		float RefractNoH = RefractBlend * NoH;
+		float NoL2 = saturate( RefractNoH - (1 / 1.5) * NoL );
+		float NoV2 = saturate( RefractNoH - (1 / 1.5) * NoV );
+
+		if( NoL > 0 )
+		{
+			float Vis = Vis_Schlick( Roughness, NoV2, NoL2 );
+
+			float NoL_Vis_PDF = NoL * Vis * (4 * VoH / NoH);
+
+			float Fc = pow( 1 - VoH, 5 );
+			A += (1 - Fc) * NoL_Vis_PDF;
+			B += Fc * NoL_Vis_PDF;
+		}
+	}
+
+	float F0 = 0.04;
+	float Fc = pow( 1 - NoV, 5 );
+	float F = Fc + (1 - Fc) * F0;
+
+	return (1 - Fc) * float2( A, B ) / NumSamples;
+}
+
+float G_Schlick( float Roughness, float NoV, float NoL )
+{
+	float k = Square( Roughness ) * 0.5;
+	float Vis_SchlickV = NoV * (1 - k) + k;
+	float Vis_SchlickL = NoL * (1 - k) + k;
+	return ( NoV * NoL ) / ( 4 * Vis_SchlickV * Vis_SchlickL );
+}
+
+float IntegrateAbsorption( uint2 Random, float Roughness, float NoV, float AbsorptionColor )
+{
+	float3 V;
+	V.x = sqrt( 1.0f - NoV * NoV );	// sin
+	V.y = 0;
+	V.z = NoV;						// cos
+
+	float A = 0;
+	float B = 0;
+
+	const uint NumSamples = 16;
+	for( uint i = 0; i < NumSamples; i++ )
+	{
+		float2 E = Hammersley( i, NumSamples, Random );
+		float3 H = ImportanceSampleGGX( E, Roughness ).xyz;
+		float3 L = 2 * dot( V, H ) * H - V;
+
+		float NoL = saturate( L.z );
+		float NoH = saturate( H.z );
+		float VoH = saturate( dot( V, H ) );
+
+		if( NoL > 0 )
+		{
+			float3 L2 = refract( -L, -H, 1 / 1.5 );
+			float3 V2 = refract( -V, -H, 1 / 1.5 );
+			float3 H2 = normalize(V2 + L2);
+
+			float NoL2 = saturate( L2.z );
+			float NoV2 = saturate( V2.z );
+			float NoH2 = saturate( H2.z );
+			float VoH2 = saturate( dot(V2, H2) );
+			
+			float AbsorptionDist = rcp(NoV2) + rcp(NoL2);
+			float Absorption = pow( AbsorptionColor, 0.5 * AbsorptionDist );
+
+			float F21 = 0.04 + 0.96 * ClampedPow( 1 - saturate( dot(V2, H) ), 5 );
+			float TotalInternalReflection = 1 - F21 * G_Schlick( Roughness, NoV2, NoL2 );
+			float3 LayerAttenuation = TotalInternalReflection * Absorption;
+
+			float D = D_GGX( Roughness, NoH2 );
+			float Vis = Vis_Schlick( Roughness, NoV2, NoL2 );
+
+			float PDF = D_GGX( Roughness, NoH2 ) * NoH / (4 * VoH);
+
+			float NoL_Vis_PDF = NoL * Vis * D / PDF;
+
+			NoL_Vis_PDF *= LayerAttenuation;
+
+			float Fc = pow( 1 - VoH2, 5 );
+			A += (1 - Fc) * NoL_Vis_PDF;
+			B += Fc * NoL_Vis_PDF;
+		}
+	}
+
+	float F0 = 0.04;
+	float Fc = pow( 1 - NoV, 5 );
+	float F = Fc + (1 - Fc) * F0;
+
+	return (1 - F) * float2( A, B ) / NumSamples;
+}
+
+float ApproxAbsorption( float Roughness, float NoV, float AbsorptionColor )
+{
+	float Toe = 1 - exp2( -9 * NoV - 3.5 );
+	return Toe * AbsorptionColor * ( (NoV - 1) * 0.85 * ( 1 - lerp( AbsorptionColor, Square(AbsorptionColor), -0.78 ) ) + lerp( 0.93, 0.96, AbsorptionColor * 2 - 1 ) );
+}
+
 // Function graphing
 float F0( float x )
 {
-	return x*saturate( (x - 0.5)/2 );
+	//return x*saturate( (x - 0.5)/2 );
+	//return IntegrateBRDF( uint2(0,0), 0.75, x ).x;
+	return IntegrateAbsorption( uint2(0,0), 0.4, x, 0.85 ) / ( IntegrateClearCoat( uint2(0,0), 0.4, x ).x );
 }
 
 float F1( float x )
 {
-	return x;
+	return IntegrateAbsorption( uint2(0,0), 0.4, x, 0.5 ) / ( IntegrateClearCoat( uint2(0,0), 0.4, x ).x );
+	//return IntegrateAbsorption( uint2(0,0), 0, x );
 }
 
 float F2( float x )
 {
-	return x;
+	//return x;
+	return IntegrateAbsorption( uint2(0,0), 0.4, x, 0.25 ) / ( IntegrateClearCoat( uint2(0,0), 0.4, x ).x );
+
+	float Roughness = 0.5;
+	float NoV = x;
+
+	float F0 = 0.04;
+	float Fc = pow( 1 - NoV, 5 );
+	float F = Fc + (1 - Fc) * F0;
+
+	// [ Lazarov 2013, "Getting More Physical in Call of Duty: Black Ops II" ]
+	// Adaptation to fit our G term.
+	half a004 = min( Square(1 - Roughness), exp2( -9.28 * NoV ) ) * (1 - Roughness) + Roughness * -0.0275 + 0.0425;
+	half AB = 1.04 * (1 - a004) + Roughness * -0.572;
+
+	AB = IntegrateBRDF( uint2(0,0), Roughness, x ).x;
+
+	//AB *= (1 - F);
+
+	float Scale = saturate( (NoV - 1) * lerp( 1.3, 0.5, Roughness ) + 1 );
+	//float Scale = 0.85;
+
+	return AB * Scale;
 }
 
 float F3( float x )
 {
-	return x;
+	//return x;
+	return IntegrateAbsorption( uint2(0,0), 0.4, x, 0.05 ) / ( IntegrateClearCoat( uint2(0,0), 0.4, x ).x );
+}
+
+float F4( float x )
+{
+	return ApproxAbsorption( 0.4, x, 0.85 );
+	//return (x - 1) * 0.12 + 0.95;
+}
+
+float F5( float x )
+{
+	//return (x - 1) * 0.25 + 0.93;
+
+	return ApproxAbsorption( 0.4, x, 0.5 );
+}
+
+float F6( float x )
+{
+	return ApproxAbsorption( 0.4, x, 0.25 );
+	//return (x - 1) * 0.45 + 0.9;
+}
+
+float F7( float x )
+{
+	return ApproxAbsorption( 0.4, x, 0.05 );
+	//return (x - 1) * 0.7 + 0.82;
 }
 
 float LineShade( float fx, float y, float dydx, float LineWidth )
@@ -178,7 +398,7 @@ float LineShade( float fx, float y, float dydx, float LineWidth )
 float3 Graph( float2 ScreenSpacePos )
 {
 	float2 WindowMin = float2( 0, 0 );
-	float2 WindowMax = float2( 10, 10 );
+	float2 WindowMax = float2( 1, 1 );
 	
 	float2 p = ( (ScreenSpacePos + 1) * 0.5 - WindowMin ) * ( WindowMax - WindowMin );
 	float LineWidth = dot( WindowMax - WindowMin, 0.0005 );
@@ -186,8 +406,14 @@ float3 Graph( float2 ScreenSpacePos )
 	float3 Color;
 	Color  = float3( 1, 0, 0 ) * LineShade( F0(p.x), p.y, ( F0(p.x + LineWidth) - F0(p.x - LineWidth) ) / (2 * LineWidth), LineWidth );
 	Color += float3( 0, 1, 0 ) * LineShade( F1(p.x), p.y, ( F1(p.x + LineWidth) - F1(p.x - LineWidth) ) / (2 * LineWidth), LineWidth );
-	//Color += float3( 1, 1, 0 ) * LineShade( F2(Graph.x), Graph.y, ( F2(p.x + LineWidth) - F2(p.x - LineWidth) ) / (2 * LineWidth), LineWidth );
-	//Color += float3( 0, 1, 1 ) * LineShade( F3(Graph.x), Graph.y, ( F3(p.x + LineWidth) - F3(p.x - LineWidth) ) / (2 * LineWidth), LineWidth );
+	Color += float3( 0, 0, 1 ) * LineShade( F2(p.x), p.y, ( F2(p.x + LineWidth) - F2(p.x - LineWidth) ) / (2 * LineWidth), LineWidth );
+	Color += float3( 1, 1, 0 ) * LineShade( F3(p.x), p.y, ( F3(p.x + LineWidth) - F3(p.x - LineWidth) ) / (2 * LineWidth), LineWidth );
+
+	Color += float3( 1, 0.5, 0 ) * LineShade( F4(p.x), p.y, ( F4(p.x + LineWidth) - F4(p.x - LineWidth) ) / (2 * LineWidth), LineWidth );
+	Color += float3( 0, 1, 1 ) * LineShade( F5(p.x), p.y, ( F5(p.x + LineWidth) - F5(p.x - LineWidth) ) / (2 * LineWidth), LineWidth );
+	Color += float3( 1, 0, 1 ) * LineShade( F6(p.x), p.y, ( F6(p.x + LineWidth) - F6(p.x - LineWidth) ) / (2 * LineWidth), LineWidth );
+	Color += float3( 1, 1, 1 ) * LineShade( F7(p.x), p.y, ( F7(p.x + LineWidth) - F7(p.x - LineWidth) ) / (2 * LineWidth), LineWidth );
+	
 	return Color;
 }
 
@@ -282,6 +508,10 @@ void MainPS(
 		//TonemappedColor = pow( TonemappedColor, 2.0 / 2.2 );
 
 		//TonemappedColor = Graph( ScreenSpacePos );
+		float Roughness = 0;//UV.x;
+		float NoV = UV.y;
+		//TonemappedColor.xy = IntegrateClearCoat( uint2(0,0), Roughness, NoV ).x / IntegrateBRDF( uint2(0,0), Roughness, NoV ).x;
+		//TonemappedColor.z = 0;
 	
 		half LuminanceForPostProcessAA  = dot(TonemappedColor, half3(0.299f, 0.587f, 0.114f));
 
diff --git a/Engine/Shaders/ReflectionEnvironmentComputeShaders.usf b/Engine/Shaders/ReflectionEnvironmentComputeShaders.usf
--- a/Engine/Shaders/ReflectionEnvironmentComputeShaders.usf
+++ b/Engine/Shaders/ReflectionEnvironmentComputeShaders.usf
@@ -697,7 +697,7 @@ void ReflectionEnvironmentTiledDeferredMain(
 			GBuffer.SpecularColor = GBuffer.SpecularColor * AB.x + AB.y * saturate( 50 * GBuffer.SpecularColor.g ) * (1 - ClearCoat);
 			
 			// F_Schlick
-			float F0 = GBuffer.Specular * 0.08;
+			float F0 = 0.04;
 			float Fc = pow( 1 - NoV, 5 );
 			float F = Fc + (1 - Fc) * F0;
 			F *= ClearCoat;
