Change the ShadingModelID to an int to prevent a type-promotion bug in the cross-compiler from preventing the Mac build from running due to the latest shader changes. This will be removed once the cross-compiler isn't failing.

diff --git a/Engine/Shaders/DeferredShadingCommon.usf b/Engine/Shaders/DeferredShadingCommon.usf
--- a/Engine/Shaders/DeferredShadingCommon.usf
+++ b/Engine/Shaders/DeferredShadingCommon.usf
@@ -289,7 +289,7 @@ struct FGBufferData
 	// 0..1 ambient occlusion  e.g.SSAO, wet surface mask, skylight mask, ...
 	float GBufferAO;
 	// 0..255 
-	uint ShadingModelID;
+	int ShadingModelID;
 	// 0..1 decal receiver mask
 	uint DecalMask;
 	// 0..1
