Fixed particle light attenuation.

diff --git a/Engine/Shaders/DeferredLightingCommon.usf b/Engine/Shaders/DeferredLightingCommon.usf
--- a/Engine/Shaders/DeferredLightingCommon.usf
+++ b/Engine/Shaders/DeferredLightingCommon.usf
@@ -655,7 +655,7 @@ float3 GetSimpleDynamicLighting(float3 WorldPosition, float3 CameraVector, FScre
 	}
 	else
 	{
-		DistanceAttenuation = RadialAttenuation(L * LightData.LightPositionAndInvRadius.w, LightData.LightColorAndFalloffExponent.w);
+		DistanceAttenuation = RadialAttenuation(ToLight * LightData.LightPositionAndInvRadius.w, LightData.LightColorAndFalloffExponent.w);
 	}
 
 	float3 OutLighting = 0;
@@ -666,7 +666,7 @@ float3 GetSimpleDynamicLighting(float3 WorldPosition, float3 CameraVector, FScre
 		const float3 LightColor = LightData.LightColorAndFalloffExponent.rgb;
 
 		// Apply SSAO to the direct lighting since we're not going to have any other shadowing
-		float3 Attenuation = DistanceAttenuation * ScreenSpaceData.AmbientOcclusion;
+		float Attenuation = DistanceAttenuation * ScreenSpaceData.AmbientOcclusion;
 
 		OutLighting += LightColor * (NoL * Attenuation) * SimpleShading( ScreenSpaceData.GBuffer, ScreenSpaceData.GBuffer.Roughness, L, V, N );
 	}
