* clamped material graph inputs to avoid negative colors or NaN creeping into the shading in the forward pass.
  can be reproduced in RenderTestMap

diff --git a/Engine/Shaders/BasePassPixelShader.usf b/Engine/Shaders/BasePassPixelShader.usf
--- a/Engine/Shaders/BasePassPixelShader.usf
+++ b/Engine/Shaders/BasePassPixelShader.usf
@@ -544,11 +544,22 @@ void Main(
 #if DIFFUSE_SPEC_INPUTS
 	half3 DiffuseColor = GetMaterialDiffuseColor( MaterialParameters );
 	half3 SpecularColor = GetMaterialSpecularColor( MaterialParameters );
+
+	// prevents negative colors and inconsistent results with lighting in the base pass (should be almost free on PC, on mobile it's more costly so we expect the content creator to fix it)
+	// could be optimized a bit later but it would be more fragile
+	DiffuseColor = saturate(DiffuseColor);
+	SpecularColor = saturate(SpecularColor);
 #else
 	// Store the results in local variables and reuse instead of calling the functions multiple times.
 	half3 BaseColor = GetMaterialBaseColor( MaterialParameters );
 	half  Metallic = GetMaterialMetallic( MaterialParameters );
 	half  Specular = GetMaterialSpecular( MaterialParameters );
+
+	// prevents negative colors and inconsistent results with lighting in the base pass (should be almost free on PC, on mobile it's more costly so we expect the content creator to fix it)
+	// could be optimized a bit later but it would be more fragile
+	BaseColor = saturate(BaseColor);
+	Metallic = saturate(Metallic);
+
 #endif
 
 	float MaterialAO = GetMaterialAmbientOcclusion(MaterialParameters);
