* minor optimization

diff --git a/Engine/Shaders/PostProcessTonemap.usf b/Engine/Shaders/PostProcessTonemap.usf
--- a/Engine/Shaders/PostProcessTonemap.usf
+++ b/Engine/Shaders/PostProcessTonemap.usf
@@ -200,6 +200,8 @@ void MainPS(
 	out float4 OutColor : SV_Target0
 	)
 {
+	OutColor = 0;
+
 	float2 UV = UVAndScreenPos.xy;
 	float2 ScreenSpacePos = UVAndScreenPos.zw;
 
@@ -217,22 +219,22 @@ void MainPS(
 		#if USE_GRAIN_JITTER
 			SceneUVJitter = SceneUV.xy - UV.xy;
 		#endif
-		half4 SceneColor = Texture2DSample(PostprocessInput0, PostprocessInput0Sampler, SceneUV);
+		half3 SceneColor = Texture2DSample(PostprocessInput0, PostprocessInput0Sampler, SceneUV).rgb;
 		half SceneColorG = Texture2DSample(PostprocessInput0, PostprocessInput0Sampler, FringeUV.xy + SceneUVJitter.xy).g;
 		half SceneColorB = Texture2DSample(PostprocessInput0, PostprocessInput0Sampler, FringeUV.zw + SceneUVJitter.xy).b;
 		SceneColor.g = SceneColorG;
 		SceneColor.b = SceneColorB;
 	#else
-		half4 SceneColor = Texture2DSample(PostprocessInput0, PostprocessInput0Sampler, SceneUV);
+		half3 SceneColor = Texture2DSample(PostprocessInput0, PostprocessInput0Sampler, SceneUV).rgb;
 	#endif
 
 	#if USE_GAMMA_ONLY
 
-		OutColor = pow(SceneColor, InverseGamma.x);
+		OutColor.rgb = pow(SceneColor, InverseGamma.x);
 
 	#else 
 
-		half3 LinearColor = SceneColor.rgb * ColorScale0.rgb;
+		half3 LinearColor = SceneColor * ColorScale0.rgb;
 
 		#if USE_BLOOM
 			float4 CombinedBloom = Texture2DSample(PostprocessInput1, PostprocessInput1Sampler, UV);
