Minor optimization.

diff --git a/Engine/Shaders/ReflectionEnvironmentComputeShaders.usf b/Engine/Shaders/ReflectionEnvironmentComputeShaders.usf
--- a/Engine/Shaders/ReflectionEnvironmentComputeShaders.usf
+++ b/Engine/Shaders/ReflectionEnvironmentComputeShaders.usf
@@ -705,8 +705,11 @@ void ReflectionEnvironmentTiledDeferredMain(
 			float LayerAttenuation = (1 - F);
 
 		#if USE_LIGHTMAPS
-			float  Lobe0Reflectance = F;
-			float3 Lobe1Reflectance = GBuffer.SpecularColor * LayerAttenuation;
+			Lobes[0].Color.a *= F;
+			Lobes[1].Color.a *= LayerAttenuation;
+
+			float3 Lobe0Reflectance = 1;
+			float3 Lobe1Reflectance = GBuffer.SpecularColor;
 		#else
 			float3 Lobe0Reflectance = lerp( GBuffer.SpecularColor, F, ClearCoat );
 			float3 Lobe1Reflectance = ( GBuffer.DiffuseColor + GBuffer.SpecularColor * ClearCoat ) * LayerAttenuation;
