Cleaning up clear coat code.

diff --git a/Engine/Shaders/PostProcessAmbient.usf b/Engine/Shaders/PostProcessAmbient.usf
--- a/Engine/Shaders/PostProcessAmbient.usf
+++ b/Engine/Shaders/PostProcessAmbient.usf
@@ -14,7 +14,6 @@
 
 #define IMPORTANCE_SAMPLE		0
 
-#if 1
 float3 DiffuseIBL( uint2 Random, float3 DiffuseColor, float3 N )
 {
 	float3 DiffuseLighting = 0;
@@ -94,110 +93,6 @@ float3 StandardShading( FGBufferData GBuffer, float Roughness, float3 L, float3
 	return Diffuse_Lambert( GBuffer.DiffuseColor ) + (D * Vis) * F;
 }
 
-#if 0
-void ClearCoatShadingShared( FGBufferData GBuffer, float3 V, half3 N )
-{
-	float NoV = saturate( dot(N, V) );
-	Shared.NoV = NoV;
-
-	// D_GGX, Vis_Schlick
-	float m = Roughness * Roughness;
-	Shared.m2 = m * m;
-	Shared.SpecularMul = (0.25 / PI) * ( Shared.m2 / (NoV * NoV) );
-	
-	// F_Schlick
-	float F0 = GBuffer.Specular * 0.08;
-	float Fc = exp2( (-5.55473 * NoV - 6.98316) * NoV );	// 1 mad, 1 mul, 1 exp
-	float F = Fc + (1 - Fc) * F0;							// 1 add, 3 mad
-	Shared.SpecularMul *= F;
-
-	Shared.AbsorptionFactor = 0.5 * log2( GBuffer.BaseColor );
-
-
-	float AbsorptionFactor = 0.72134752 * rcp( NoL2 ) + ( 0.39674113 + 0.72134752 * rcp( NoV2 ) );
-	float3 Absorption = exp2( AbsorptionFactor * GBuffer.BaseColor - AbsorptionFactor );
-
-	
-
-	// Diffuse_Lambert
-	Shared.DiffuseMul = DiffuseColor * (1.0 / PI);
-
-	// D_GGX, Vis_Schlick
-	float m = Roughness * Roughness;
-	Shared.m2 = m * m;
-	Shared.k = Square( Roughness ) * 0.5;
-	//float Vis_SchlickV = NoV * (1 - Shared.k) + Shared.k;
-	Shared.SpecularMul = (0.25 / PI) * Shared.m2;
-	
-	// F_Schlick
-	Shared.SpecularMul *= saturate( 50.0 * SpecularColor.g );
-}
-
-float3 ClearCoatShadingPerLight( Shared, float3 L, float3 V, half3 N )
-{
-	float ClearCoatRoughness = 0.1;
-
-	float Thickness = 1.5;
-	float3 WaterAbsorption = { 0.98, 0.4, 0.28 };
-
-	GBuffer.BaseColor = exp( -WaterAbsorption * Thickness * 2 );
-
-	const float3 n = { 1.94, 1, 0.6 };
-	const float3 k = { 8.21, 6.69, 4.86 };
-	GBuffer.SpecularColor = ( Square(n - 1) + k*k ) / ( Square(n + 1) + k*k );
-
-	// TODO control SpecularColor
-
-
-	float3 H = normalize(V + L);
-	float NoL = saturate( dot(N, L) );
-	float NoV = saturate( dot(N, V) );
-	float NoH = saturate( dot(N, H) );
-	float VoH = saturate( dot(V, H) );
-	
-	// simplify for low roughness
-	// VoH == NoV
-	// EnvBRDF == F_Schlick()
-
-	// Generalized microfacet specular
-	float D = D_GGX( ClearCoatRoughness, NoH );
-	float Vis = Vis_Kelemen( VoH );
-	//float Vis = rcp( 4 * NoV * NoV );
-	
-	// F_Schlick
-	float F0 = GBuffer.Specular * 0.08;
-	float Fc = exp2( (-5.55473 * VoH - 6.98316) * VoH );	// 1 mad, 1 mul, 1 exp
-	float F = Fc + (1 - Fc) * F0;							// 1 add, 3 mad
-
-	float Fr1 = D * Vis * F;
-
-	
-	// Approximation
-	float RefractBlend = (0.22 * VoH + 0.7) * VoH + 0.745;	// 2 mad
-
-	// dot products distribute
-	float RefractNoH = RefractBlend * NoH;					// 1 mul
-	float NoL2 = saturate( RefractNoH - (1 / 1.5) * NoL );	// 1 mad
-	float NoV2 = saturate( RefractNoH - (1 / 1.5) * NoV );	// 1 mad
-
-	float3 AbsorptionFactor = 0.5 * log2( GBuffer.BaseColor );
-	float3 Absorption = exp2( AbsorptionFactor * ( 1 / NoV2 + 1 / NoL2 ) );
-	//float3 Absorption = pow( GBuffer.BaseColor, 0.5 * ( rcp(NoV2) + rcp(NoL2) ) );
-	
-	float3 LayerAttenuation = (1 - F) * Absorption;
-	
-	// Generalized microfacet specular
-	float D2 = D_GGX( Roughness, NoH );
-	float Vis2 = Vis_Schlick( Roughness, NoV2, NoL2 );
-	float3 F2 = F_Schlick( SpecularColor, VoH );
-	//float3 F2 = GBuffer.SpecularColor;
-
-	float3 Fr2 = Diffuse_Lambert( GBuffer.DiffuseColor ) + (D2 * Vis2) * F2;
-	
-	return Fr1 + Fr2 * LayerAttenuation;
-}
-#endif
-
 float3 ClearCoatShading( FGBufferData GBuffer, float Roughness, float3 L, float3 V, half3 N )
 {
 	const float ClearCoat			= GBuffer.CustomData.x;
@@ -207,7 +102,6 @@ float3 ClearCoatShading( FGBufferData GBuffer, float Roughness, float3 L, float3
 
 	Film *= ClearCoat;
 
-#if 1
 	float3 H = normalize(V + L);
 	float NoL = saturate( dot(N, L) );
 	float NoV = saturate( dot(N, V) );
@@ -236,140 +130,11 @@ float3 ClearCoatShading( FGBufferData GBuffer, float Roughness, float3 L, float3
 	float3 Fr2 = Diffuse_Lambert( GBuffer.DiffuseColor ) + (D2 * Vis2) * F2;
 	
 	return Fr1 + Fr2 * LayerAttenuation;
-#else
-	float3 AbsorptionColor = (1 - Film) + GBuffer.BaseColor * ( Film * (1 / MetalSpec) );
-	GBuffer.BaseColor = lerp( GBuffer.BaseColor, MetalSpec, Film );
-	
-	GBuffer.DiffuseColor = GBuffer.BaseColor - GBuffer.BaseColor * GBuffer.Metallic;
-	GBuffer.SpecularColor = lerp( 0.08 * GBuffer.Specular.xxx, GBuffer.BaseColor, GBuffer.Metallic );
-
-#if 0	// Ref
-	float3 H = normalize(V + L);
-	float NoL = saturate( dot(N, L) );
-	float NoV = saturate( dot(N, V) );
-	float NoH = saturate( dot(N, H) );
-	float VoH = saturate( dot(V, H) );
-	
-	// Hard coded IOR of 1.5
-
-	// Generalized microfacet specular
-	float D = D_GGX( ClearCoatRoughness, NoH );
-	float Vis = Vis_Kelemen( VoH );
-	
-	// F_Schlick
-	float F0 = 0.04;
-	float Fc = pow( 1 - VoH, 5 );
-	float F = Fc + (1 - Fc) * F0;
-
-	float Fr1 = D * Vis * F;
-
-	// Refract rays
-	float3 L2 = refract( -L, -H, 1 / 1.5 );
-	float3 V2 = refract( -V, -H, 1 / 1.5 );
-	float3 H2 = normalize(V2 + L2);
-
-	float NoL2 = saturate( dot(N, L2) );
-	float NoV2 = saturate( dot(N, V2) );
-	float NoH2 = saturate( dot(N, H2) );
-	float VoH2 = saturate( dot(V2, H2) );
-
-	float  AbsorptionDist = rcp(NoV2) + rcp(NoL2);
-	float3 Absorption = pow( AbsorptionColor, 0.5 * AbsorptionDist );
-	
-	float F21 = Fresnel( 1 / 1.5, saturate( dot(V2, H) ) );
-	float TotalInternalReflection = 1 - F21 * G_Schlick( GBuffer.Roughness, NoV2, NoL2 );
-	float3 LayerAttenuation = ( (1 - F) * TotalInternalReflection ) * Absorption;
-
-	// Treat SpecularColor as relative to IOR. Artist compensates.
-	//GBuffer.SpecularColor = ChangeBaseMedium( GBuffer.SpecularColor, 1.5 );
-	
-	// Generalized microfacet specular
-	float D2 = D_GGX( Roughness, NoH2 );
-	float Vis2 = Vis_Schlick( GBuffer.Roughness, NoV2, NoL2 );
-	float3 F2 = F_Schlick( GBuffer.SpecularColor, VoH2 );
-
-	float3 Fr2 = Diffuse_Lambert( GBuffer.DiffuseColor ) + (D2 * Vis2) * F2;
-	
-	return Fr1 + Fr2 * LayerAttenuation;
-#else
-	float3 H = normalize(V + L);
-	float NoL = saturate( dot(N, L) );
-	float NoV = saturate( dot(N, V) );
-	float NoH = saturate( dot(N, H) );
-	float VoH = saturate( dot(V, H) );
-	
-	// Hard coded IOR of 1.5
-
-	// Generalized microfacet specular
-	float D = D_GGX( ClearCoatRoughness, NoH );
-	float Vis = Vis_Kelemen( VoH );
-	
-	// F_Schlick
-	float F0 = 0.04;
-	float Fc = pow( 1 - VoH, 5 );
-	float F = Fc + (1 - Fc) * F0;
-
-	float Fr1 = D * Vis * F;
-
-	// Refract rays
-	//float3 L2 = refract( -L, -H, 1 / 1.5 );
-	//float3 V2 = refract( -V, -H, 1 / 1.5 );
-
-	// LoH == VoH
-	//float RefractBlend = sqrt( 4 * VoH*VoH + 5 ) / 3 + 2.0 / 3 * VoH;
-	//float3 L2 = RefractBlend * H - L / 1.5;
-	//float3 V2 = RefractBlend * H - V / 1.5;
-	//float NoL2 = saturate( dot(N, L2) );
-	//float NoV2 = saturate( dot(N, V2) );
-	
-	// Approximation
-	float RefractBlend = (0.22 * VoH + 0.7) * VoH + 0.745;	// 2 mad
-	// Dot products distribute. No need for L2 and V2.
-	float RefractNoH = RefractBlend * NoH;					// 1 mul
-	float NoL2 = saturate( RefractNoH - (1 / 1.5) * NoL );	// 1 mad
-	float NoV2 = saturate( RefractNoH - (1 / 1.5) * NoV );	// 1 mad
-	// Should refract H too but unimportant
-
-	float  AbsorptionDist = rcp(NoV2) + rcp(NoL2);
-	float3 Absorption = pow( AbsorptionColor, 0.5 * AbsorptionDist );
-
-	// Approximation
-	//float  AbsorptionDist = ( NoV2 + NoL2 ) / ( NoV2 * NoL2 );
-	//float3 Absorption = AbsorptionColor * ( AbsorptionColor * (AbsorptionDist * 0.5 - 1) + (2 - 0.5 * AbsorptionDist) );
-	//float3 Absorption = AbsorptionColor + AbsorptionColor * (AbsorptionColor - 1) * (AbsorptionDist * 0.5 - 1);	// use for shared version
-	
-	//float F21 = Fresnel( 1 / 1.5, saturate( dot(V2, H) ) );
-	//float TotalInternalReflection = 1 - F21 * G_Schlick( Roughness, NoV2, NoL2 );
-	//float3 LayerAttenuation = ( (1 - F12) * TotalInternalReflection ) * Absorption;
-
-	// Approximation
-	float3 LayerAttenuation = (1 - F) * Absorption;
-
-	// Treat SpecularColor as relative to IOR. Artist compensates.
-	//SpecularColor = ChangeBaseMedium( SpecularColor, 1.5 );
-	
-	// Generalized microfacet specular
-	float D2 = D_GGX( Roughness, NoH );
-	float Vis2 = Vis_Schlick( Roughness, NoV2, NoL2 );
-	float3 F2 = F_Schlick( GBuffer.SpecularColor, VoH );
-
-	float3 Fr2 = Diffuse_Lambert( GBuffer.DiffuseColor ) + (D2 * Vis2) * F2;
-	
-	return Fr1 + Fr2 * LayerAttenuation;
-#endif
-
-#endif
 }
 
 
 float3 ImageBasedLightingMIS( FGBufferData GBuffer, float3 V, float3 N, uint2 Random )
 {
-	BRANCH
-	if( GBuffer.ShadingModelID == 0 )
-	{
-		return 0;
-	}
-
 	float3 Lighting = 0;
 
 	float Roughness1 = GBuffer.Roughness;
@@ -540,66 +305,6 @@ float3 ApproximateSpecularIBL( uint2 Random, float3 SpecularColor, float Roughne
 	return PrefilteredColor * ( SpecularColor * AB.x + AB.y );
 }
 
-float3 ImageBasedLightingSplitSum( FGBufferData GBuffer, float3 V, float3 N, uint2 Random )
-{
-	BRANCH
-	if( GBuffer.ShadingModelID == 0 )
-	{
-		return 0;
-	}
-
-	const float ClearCoat			= GBuffer.CustomData.x;
-	const float ClearCoatRoughness	= GBuffer.CustomData.y;
-	float Film = 1;
-	float MetalSpec = 0.9;
-
-	Film *= ClearCoat;
-
-	float NoV = saturate( dot( N, V ) );
-	
-	float Lobe1Roughness = ClearCoatRoughness;
-	float Lobe2Roughness = GBuffer.Roughness;
-	
-#if 0
-	float3 Lobe1Color = FilterEnvMap( Random, Lobe1Roughness, N, V );
-	float3 Lobe2Color = FilterEnvMap( Random, Lobe2Roughness, N, V );
-#else
-	float3 R = 2 * dot( V, N ) * N - V;
-	//float3 Lobe1Color = PrefilterEnvMap( Random, Lobe1Roughness, R );
-	//float3 Lobe2Color = PrefilterEnvMap( Random, Lobe2Roughness, R );
-
-	float Lobe1Mip = ComputeCubemapMipFromRoughness( Lobe1Roughness, AmbientCubemapMipAdjust.w );
-	float Lobe2Mip = ComputeCubemapMipFromRoughness( Lobe2Roughness, AmbientCubemapMipAdjust.w );
-
-	float3 Lobe1Color = TextureCubeSampleLevel( AmbientCubemap, AmbientCubemapSampler, R, Lobe1Mip ).rgb;
-	float3 Lobe2Color = TextureCubeSampleLevel( AmbientCubemap, AmbientCubemapSampler, R, Lobe2Mip ).rgb;
-#endif
-
-	// F_Schlick
-	float F0 = 0.04;
-	float Fc = pow( 1 - NoV, 5 );
-	float F = Fc + (1 - Fc) * F0;
-	F *= ClearCoat;
-
-	float3 Fr1 = Lobe1Color * F;
-
-	float LayerAttenuation = (1 - F);
-
-	// Generalized microfacet specular
-	//float Vis2 = Vis_Schlick( Roughness, NoV2, NoL2 );
-	//float3 F2 = F_Schlick( GBuffer.SpecularColor, VoH );
-	
-	//float3 Fr2 = Diffuse_Lambert( GBuffer.DiffuseColor ) + (D2 * Vis2) * F2;
-
-	float2 AB = PreIntegratedGF.SampleLevel( PreIntegratedGFSampler, float2( NoV, GBuffer.Roughness ), 0 ).rg;
-
-	float3 Fr2 = Lobe2Color * ( GBuffer.SpecularColor * AB.x + AB.y * (1 - ClearCoat) );
-	//float3 Fr2 = Lobe2Color * EnvBRDF( GBuffer.SpecularColor, GBuffer.Roughness, NoV );
-	
-	return Fr1 + Fr2 * LayerAttenuation;
-}
-#endif
-
 void MainPS(in float4 UVAndScreenPos : TEXCOORD0, out float4 OutColor : SV_Target0)
 {
 	float2 UV = UVAndScreenPos.xy;
@@ -613,6 +318,8 @@ void MainPS(in float4 UVAndScreenPos : TEXCOORD0, out float4 OutColor : SV_Targe
 	int2 PixelPos = int2(UVAndScreenPos.zw * ScreenPosToPixel.xy + ScreenPosToPixel.zw + 0.5f);
 
 	float3 ScreenVector = normalize(mul(float4(ScreenSpacePos, 1, 0), View.ScreenToWorld).xyz);
+	float NoV = saturate( dot( GBuffer.WorldNormal, -ScreenVector ) );
+	float3 R = reflect( ScreenVector, GBuffer.WorldNormal );
 	
 	uint2 Random = ScrambleTEA( PixelPos );
 	Random.x ^= View.Random;
@@ -624,38 +331,44 @@ void MainPS(in float4 UVAndScreenPos : TEXCOORD0, out float4 OutColor : SV_Targe
 
 	// Diffuse
 	{
-	#if IMPORTANCE_SAMPLE
-		DiffuseAndSpecular += DiffuseIBL( Random, GBuffer.DiffuseColor, GBuffer.WorldNormal );
-	#else
 		// we want to access the mip with the preconvolved diffuse lighting (coneangle=90 degree)
 		DiffuseAndSpecular += GBuffer.DiffuseColor * DiffuseLookup;
-	#endif
 	}
 
 	// Specular
 	{
-	#if IMPORTANCE_SAMPLE
-		DiffuseAndSpecular += SpecularIBL( Random, GBuffer.SpecularColor, GBuffer.Roughness, GBuffer.WorldNormal, -ScreenVector );
-		//DiffuseAndSpecular += ApproximateSpecularIBL( Random, GBuffer.SpecularColor, GBuffer.Roughness, GBuffer.WorldNormal, -ScreenVector );
-	#else
-		float3 R = reflect( ScreenVector, GBuffer.WorldNormal );
-		
-		float AbsoluteSpecularMip = ComputeCubemapMipFromRoughness( GBuffer.Roughness, AmbientCubemapMipAdjust.w );
-		float3 SampleColor = TextureCubeSampleLevel( AmbientCubemap, AmbientCubemapSampler, R, AbsoluteSpecularMip ).rgb;
+		float Mip = ComputeCubemapMipFromRoughness( GBuffer.Roughness, AmbientCubemapMipAdjust.w );
+		float3 SampleColor = TextureCubeSampleLevel( AmbientCubemap, AmbientCubemapSampler, R, Mip ).rgb;
 
-		float NoV = saturate( dot( GBuffer.WorldNormal, -ScreenVector ) );
+		//float2 AB = PreIntegratedGF.SampleLevel( PreIntegratedGFSampler, float2( NoV, GBuffer.Roughness ), 0 ).rg;
+		//GBuffer.SpecularColor = GBuffer.SpecularColor * AB.x + AB.y * (1 - ClearCoat);
+		
 		DiffuseAndSpecular += SampleColor * EnvBRDF( GBuffer.SpecularColor, GBuffer.Roughness, NoV );
-	#endif
 	}
 
 	BRANCH if( GBuffer.ShadingModelID == SHADINGMODELID_CLEAR_COAT )
 	{
-	#if IMPORTANCE_SAMPLE
-		DiffuseAndSpecular = ImageBasedLightingMIS( GBuffer, -ScreenVector, GBuffer.WorldNormal, Random );
-	#else
-		DiffuseAndSpecular = ImageBasedLightingSplitSum( GBuffer, -ScreenVector, GBuffer.WorldNormal, Random );
-	#endif
+		const float ClearCoat			= GBuffer.CustomData.x;
+		const float ClearCoatRoughness	= GBuffer.CustomData.y;
+
+		float Mip = ComputeCubemapMipFromRoughness( ClearCoatRoughness, AmbientCubemapMipAdjust.w );
+		float3 SampleColor = TextureCubeSampleLevel( AmbientCubemap, AmbientCubemapSampler, R, Mip ).rgb;
+
+		// F_Schlick
+		float F0 = 0.04;
+		float Fc = pow( 1 - NoV, 5 );
+		float F = Fc + (1 - Fc) * F0;
+		F *= ClearCoat;
+
+		float LayerAttenuation = (1 - F);
+
+		DiffuseAndSpecular *= LayerAttenuation;
+		DiffuseAndSpecular += SampleColor * F;
 	}
+
+#if IMPORTANCE_SAMPLE
+	DiffuseAndSpecular = ImageBasedLightingMIS( GBuffer, -ScreenVector, GBuffer.WorldNormal, Random );
+#endif
 	
 	// apply darkening from ambient occlusion (does not use PostprocessInput1 to set white texture if SSAO is off)
 	float AmbientOcclusion = GBuffer.GBufferAO * ScreenSpaceData.AmbientOcclusion;
@@ -680,11 +393,5 @@ void MainPS(in float4 UVAndScreenPos : TEXCOORD0, out float4 OutColor : SV_Targe
 	OutColor = float4(DiffuseAndSpecular * AmbientCubemapColor.rgb, 0);
 
 	OutColor *= AmbientOcclusion;
-
-	// to visualize the cubemap in the distance
-// 	{
-//  		float SceneDepth = CalcSceneDepth(UV);
-//  		OutColor.rgb = lerp(OutColor.rgb, TextureCubeSample(AmbientCubemap, AmbientCubemapSampler, ScreenVector).rgb, saturate(SceneDepth * 0.001f - 90));
-//  	}
 }
 
