Fixed CalcTangentToWorld calculation for InstancedStaticMeshes/Foliage
TTP#340063

diff --git a/Engine/Shaders/LocalVertexFactory.usf b/Engine/Shaders/LocalVertexFactory.usf
--- a/Engine/Shaders/LocalVertexFactory.usf
+++ b/Engine/Shaders/LocalVertexFactory.usf
@@ -362,8 +362,9 @@ half3x3 CalcTangentToLocal(FVertexFactoryInput Input)
 half3x3 CalcTangentToWorld(FVertexFactoryInput Input, half3x3 TangentToLocal)
 {
 #if USE_INSTANCING
-	half3x3 InstanceTransform = transpose(float3x3(Input.InstanceInverseTransform1.xyz, Input.InstanceInverseTransform2.xyz, Input.InstanceInverseTransform3.xyz));
-	half3x3 TangentToWorld = CalcTangentToWorldNoScale(mul(TangentToLocal, InstanceTransform));
+	half3x3 TangentToPrimitive = mul(TangentToLocal, transpose(float3x3(Input.InstanceTransform1.xyz, Input.InstanceTransform2.xyz, Input.InstanceTransform3.xyz)));
+	half3x3 TangentToWorldWithScaling = mul(TangentToPrimitive, (float3x3)Primitive.LocalToWorld); 
+	half3x3 TangentToWorld = float3x3(normalize(TangentToWorldWithScaling[0]), normalize(TangentToWorldWithScaling[1]), normalize(TangentToWorldWithScaling[2]));
 #else
 	half3x3 TangentToWorld = CalcTangentToWorldNoScale(TangentToLocal);
 #endif
