* fix screen space subsurface scattering (experimental feature)

diff --git a/Engine/Shaders/ReflectionEnvironmentComputeShaders.usf b/Engine/Shaders/ReflectionEnvironmentComputeShaders.usf
--- a/Engine/Shaders/ReflectionEnvironmentComputeShaders.usf
+++ b/Engine/Shaders/ReflectionEnvironmentComputeShaders.usf
@@ -744,8 +744,9 @@ void ReflectionEnvironmentTiledDeferredMain(
 		// Transform NaNs to black, transform negative colors to black.
 		OutColor.rgb = -min(-OutColor.rgb, 0.0);
 
-		// TODO platforms that allow RW just add
-		OutColor.rgb += InSceneColor.Load( int3(PixelPos, 0) ).rgb;
+		// alpha channel is also added to keep the alpha channel for screen space subsurface scattering
+		OutColor += InSceneColor.Load( int3(PixelPos, 0) );
+
 		RWOutSceneColor[PixelPos.xy] = OutColor;
 	}
 }
