* fixed ScreenSpaceSubsurfaceScattering - no SSS color but uses caucasian skin table, next: artist control over the table

diff --git a/Engine/Shaders/DeferredLightingCommon.usf b/Engine/Shaders/DeferredLightingCommon.usf
--- a/Engine/Shaders/DeferredLightingCommon.usf
+++ b/Engine/Shaders/DeferredLightingCommon.usf
@@ -89,7 +89,8 @@ void StandardShadingPerLight( Shared, float3 L, float3 V, half3 N )
 }
 #endif
 
-float3 StandardShading( FGBufferData GBuffer, float3 LobeRoughness, float3 LobeEnergy, float3 L, float3 V, half3 N )
+// @param DiffSpecMask .r: diffuse, .g:specular e.g. float2(1,1) for both, float2(1,0) for diffuse only
+float3 StandardShading( FGBufferData GBuffer, float3 LobeRoughness, float3 LobeEnergy, float3 L, float3 V, half3 N, float2 DiffSpecMask )
 {
 	float3 H = normalize(V + L);
 	float NoL = saturate( dot(N, L) );
@@ -102,7 +103,7 @@ float3 StandardShading( FGBufferData GBuffer, float3 LobeRoughness, float3 LobeE
 	float Vis = Vis_Schlick( LobeRoughness[1], NoV, NoL );
 	float3 F = F_Schlick( GBuffer.SpecularColor, VoH );
 
-	return Diffuse_Lambert( GBuffer.DiffuseColor ) * LobeEnergy[2] + (D * Vis) * F;
+	return Diffuse_Lambert(GBuffer.DiffuseColor) * (LobeEnergy[2] * DiffSpecMask.r) + (D * Vis * DiffSpecMask.g) * F;
 }
 
 float3 SimpleShading( FGBufferData GBuffer, float Roughness, float3 L, float3 V, half3 N )
@@ -259,7 +260,8 @@ float3 SubsurfaceShadingPreintegratedSkin( FGBufferData GBuffer, float3 L, float
 	return PreintegratedBRDF * SubsurfaceColor;
 }
 
-float3 SurfaceShading( FGBufferData GBuffer, float3 LobeRoughness, float3 LobeEnergy, float3 L, float3 V, half3 N )
+// @param DiffSpecMask .r: diffuse, .g:specular e.g. float2(1,1) for both, float2(1,0) for diffuse only
+float3 SurfaceShading( FGBufferData GBuffer, float3 LobeRoughness, float3 LobeEnergy, float3 L, float3 V, half3 N, float2 DiffSpecMask )
 {
 	switch( GBuffer.ShadingModelID )
 	{
@@ -267,8 +269,9 @@ float3 SurfaceShading( FGBufferData GBuffer, float3 LobeRoughness, float3 LobeEn
 		case SHADINGMODELID_DEFAULT_LIT:
 		case SHADINGMODELID_SUBSURFACE:
 		case SHADINGMODELID_PREINTEGRATED_SKIN:
-			return StandardShading( GBuffer, LobeRoughness, LobeEnergy, L, V, N );
+			return StandardShading( GBuffer, LobeRoughness, LobeEnergy, L, V, N, DiffSpecMask);
 		case SHADINGMODELID_CLEAR_COAT:
+			// this path does not support DiffSpecMask yet
 			return ClearCoatShading( GBuffer, LobeRoughness, LobeEnergy, L, V, N );
 		default:
 			return 0;
@@ -497,6 +500,8 @@ void GetShadowTerms(FScreenSpaceData ScreenSpaceData, FDeferredLightData LightDa
 /** Calculates lighting for a given position, normal, etc with a fully featured lighting model designed for quality. */
 float4 GetDynamicLighting(float3 WorldPosition, float3 CameraVector, float2 InUV, FScreenSpaceData ScreenSpaceData, uint ShadingModelID, FDeferredLightData LightData, float4 LightAttenuation, uint2 Random)
 {
+	FLightAccumulator LightAccumulator = (FLightAccumulator)0;
+
 	float3 V = -CameraVector;
 	float3 N = ScreenSpaceData.GBuffer.WorldNormal;
 	float3 ToLight = LightData.LightDirection;
@@ -554,11 +559,7 @@ float4 GetDynamicLighting(float3 WorldPosition, float3 CameraVector, float2 InUV
 		}
 	}
 
-	// RGB accumulated RGB HDR color, A: specular luminance for screenspace subsurface scattering
-	float4 OutLighting = 0;
-
-	// for VISUALIZE_LIGHT_CULLING
-	float EstimatedCost = 0.3f;	// running the PixelShader at all has cost
+	LightAccumulator.EstimatedCost += 0.3f;		// running the PixelShader at all has a cost
 
 	BRANCH
 	if (LightRadiusMask > 0 && SpotFalloff > 0)
@@ -582,55 +583,46 @@ float4 GetDynamicLighting(float3 WorldPosition, float3 CameraVector, float2 InUV
 		float SurfaceAttenuation	= (DistanceAttenuation * LightRadiusMask * SpotFalloff) * SurfaceShadow;
 		float SubsurfaceAttenuation	= (DistanceAttenuation * LightRadiusMask * SpotFalloff) * SubsurfaceShadow;
 
-		EstimatedCost += 0.3f;	// add the cost of getting the shadow terms
+		LightAccumulator.EstimatedCost += 0.3f;		// add the cost of getting the shadow terms
 
 #if LIGHTCULLING_FINE == 1
 		// optimization that should help if there are a lot of shadowed pixels
 		BRANCH if (ShadowedAttenuation > 0 || ShadingModelID == SHADINGMODELID_SUBSURFACE)
 #endif
+
 		{
 			const float3 LightColor = LightData.LightColorAndFalloffExponent.rgb;
 
-			float3 LobeRoughness = float3( 0.1, ScreenSpaceData.GBuffer.Roughness, 1 );
-			float3 LobeEnergy = AreaLightSpecular( LightData, LobeRoughness, ToLight, L, V, N );
+			float3 LobeRoughness = float3(0.1, ScreenSpaceData.GBuffer.Roughness, 1);
+			float3 LobeEnergy = AreaLightSpecular(LightData, LobeRoughness, ToLight, L, V, N);
 
-			float3 SurfaceLighting = SurfaceShading( ScreenSpaceData.GBuffer, LobeRoughness, LobeEnergy, L, V, N );
-			float3 SubsurfaceLighting = SubsurfaceShading( ScreenSpaceData.GBuffer, L, V, N );
-			
-#if 1
-			// TODO fix SSSSS skin
-			OutLighting.rgb += LightColor * ( (NoL * SurfaceAttenuation) * SurfaceLighting + SubsurfaceAttenuation * SubsurfaceLighting );
+			// accumulate diffuse and specular
+			{
+#if 1	// for testing if there is a perf impact
+				// correct screen space subsurface scattering
+				float3 SurfaceLightingDiff = SurfaceShading(ScreenSpaceData.GBuffer, LobeRoughness, LobeEnergy, L, V, N, float2(1, 0));
+				float3 SurfaceLightingSpec = SurfaceShading(ScreenSpaceData.GBuffer, LobeRoughness, LobeEnergy, L, V, N, float2(0, 1));
+				LightAccumulator_Add(LightAccumulator, SurfaceLightingDiff, SurfaceLightingSpec, LightColor * (NoL * SurfaceAttenuation));
 #else
+				// wrong screen space subsurface scattering but potentially faster
+				float3 SurfaceLighting = SurfaceShading(ScreenSpaceData.GBuffer, LobeRoughness, LobeEnergy, L, V, N, float2(1, 1));
+				LightAccumulator_Add(LightAccumulator, SurfaceLighting, 0, LightColor * (NoL * SurfaceAttenuation));
+#endif			
+				LightAccumulator_Add(LightAccumulator, SurfaceLightingDiff, SurfaceLightingSpec, LightColor * (NoL * SurfaceAttenuation));
+			}
 
-#define RAY_TRACE 0
-#if RAY_TRACE
-			SpecularLighting = PointLightSpecularMIS( ScreenSpaceData, LightData, ToLight, V, N, Random );
-#endif
-
-			float4 LightColor4 = float4(LightColor, Luminance(LightColor));
-			float4 DiffuseLighting4 = float4(DiffuseLighting, 0);
-			float4 SpecularLighting4 = float4(SpecularLighting, Luminance(SpecularLighting));
-			float4 SubsurfaceLighting4 = float4(SubsurfaceLighting, 0);
-
-#if RAY_TRACE
-			OutLighting += LightColor4 * ( (NoL * ShadowedAttenuation) * DiffuseLighting4 + SubsurfaceLighting4 * SubsurfaceAttenuation );
-			OutLighting += LightColor4 * SpecularLighting4 * (LightRadiusMask * SpotFalloff * OpaqueShadowTerm);
-#else
-			OutLighting += LightColor4 * ( (NoL * SurfaceAttenuation) * (DiffuseLighting4 + SpecularLighting4) + SubsurfaceLighting4 * SubsurfaceAttenuation );
-#endif
+			// accumulate subsurface
+			{
+				float3 SubsurfaceLighting = SubsurfaceShading(ScreenSpaceData.GBuffer, L, V, N);
 
-#endif
+				LightAccumulator_Add(LightAccumulator, SubsurfaceLighting, 0, LightColor * SubsurfaceAttenuation);
 
-			EstimatedCost += 0.4f;	// add the cost of the lighting computations (should sum up to 1)
+				LightAccumulator.EstimatedCost += 0.4f;		// add the cost of the lighting computations (should sum up to 1 form one light)
+			}
 		}
 	}
 
-#if VISUALIZE_LIGHT_CULLING == 1
-	// similar to VISUALIZE_LIGHT_CULLING in tile based deferred lighting
-	OutLighting = 0.1f * float4(1.0f, 0.25f, 0.075f, 0) * EstimatedCost;
-#endif
-
-	return OutLighting;
+	return LightAccumulator_GetResult(LightAccumulator);
 }
 
 /** 
diff --git a/Engine/Shaders/DeferredShadingCommon.usf b/Engine/Shaders/DeferredShadingCommon.usf
--- a/Engine/Shaders/DeferredShadingCommon.usf
+++ b/Engine/Shaders/DeferredShadingCommon.usf
@@ -7,6 +7,8 @@
 #ifndef __DEFERRED_SHADING_COMMON__
 #define __DEFERRED_SHADING_COMMON__
 
+#include "LightAccumulator.usf"
+
 // TODO: for CustomGBufferResolvePS() MSAA_SAMPLE_COUNT is defined by C++ code as 2 or 4
 // bot not for any other shaders!
 #ifndef MSAA_SAMPLE_COUNT
@@ -679,4 +681,4 @@ FScreenSpaceData GetScreenSpaceData(float2 UV, bool bGetNormalizedNormal = true)
 #endif
 
 
-#endif
+#endif // __DEFERRED_SHADING_COMMON__
diff --git a/Engine/Shaders/LightAccumulator.usf b/Engine/Shaders/LightAccumulator.usf
new file mode 100644
--- /dev/null
+++ b/Engine/Shaders/LightAccumulator.usf
@@ -0,0 +1,81 @@
+// Copyright 1998-2014 Epic Games, Inc. All Rights Reserved.
+
+/*=============================================================================
+	LightAccumulator.usf: FLightAccumulator "class" and it's methods, useful for screen space subsurface scattering
+=============================================================================*/
+
+#ifndef __LIGHTACCUMULATOR_COMMON__
+#define __LIGHTACCUMULATOR_COMMON__
+
+// set by c++, not set for LPV
+// 0 / 1
+#ifndef VISUALIZE_LIGHT_CULLING
+	#define VISUALIZE_LIGHT_CULLING 0
+#endif
+
+// for ScreenSpaceSubsurfaceScattering
+// 0 : fastest (speculars leak in SSSSS)
+// 1 : luminance specular works well (needs 64bit SceneColor, stores luminance of specular in alpha, can be optimized further in LightAccumulator_Add() ) 
+// 2 : colored specular works well (best, requires another RT for light accumulation)
+#define SUBSURFACE_CHANNEL_MODE 1
+
+struct FLightAccumulator
+{
+	float3 TotalSum;
+
+	// only actually used SUBSURFACE_CHANNEL_MODE != 0
+	// assumed to be compiled out otherwise (not compiled out with #if so we can use if() instead of #if for better readability and compiler error checking)
+	// input for ScreenSpaceSubsurfaceScattering
+	float3 NonSpecularSum;
+
+	// only used for development (not compiled out with #if so we can use if() instead of #if for better readability and compiler error checking)
+	// assumed to be compiled out otherwise
+	float EstimatedCost;
+};
+
+//
+// accumulate light, can be called multiple times
+void LightAccumulator_Add(inout FLightAccumulator In, float3 NonSpecularContribution, float3 SpecularContribution, float3 CommonMultiplier)
+{
+	if (SUBSURFACE_CHANNEL_MODE == 0)
+	{
+		// 3 add, 3 mad
+		In.TotalSum += (NonSpecularContribution + SpecularContribution) * CommonMultiplier;
+	}
+	else
+	{
+		// 3 add, 3 mad
+		In.TotalSum += (NonSpecularContribution + SpecularContribution) * CommonMultiplier;
+		// 3 mad
+		In.NonSpecularSum += NonSpecularContribution * CommonMultiplier;
+	}
+}
+
+//
+// compute final value to store in the MRT0
+float4 LightAccumulator_GetResult(FLightAccumulator In)
+{
+	float4 Ret;
+
+	if (VISUALIZE_LIGHT_CULLING == 1)
+	{
+		// similar to VISUALIZE_LIGHT_CULLING in tile based deferred lighting
+		Ret = 0.1f * float4(1.0f, 0.25f, 0.075f, 0) * In.EstimatedCost;
+	}
+	else
+	{
+		Ret = float4(In.TotalSum, 0);
+
+		if (SUBSURFACE_CHANNEL_MODE != 0)
+		{
+			// RGB accumulated RGB HDR color, A: specular luminance for screenspace subsurface scattering
+			// 3 add,  1 mul, 2 mad, can be optimized to use 2 less temporary during accumulation and remove the 3 add
+			Ret.a = Luminance(In.TotalSum - In.NonSpecularSum);
+			// todo, need second MRT for SUBSURFACE_CHANNEL_MODE==2
+		}
+	}
+
+	return Ret;
+}
+
+#endif // __LIGHTACCUMULATOR_COMMON__
\ No newline at end of file
diff --git a/Engine/Shaders/PostProcessAmbient.usf b/Engine/Shaders/PostProcessAmbient.usf
--- a/Engine/Shaders/PostProcessAmbient.usf
+++ b/Engine/Shaders/PostProcessAmbient.usf
@@ -325,14 +325,15 @@ void MainPS(in float4 UVAndScreenPos : TEXCOORD0, out float4 OutColor : SV_Targe
 	Random.x ^= View.Random;
 	Random.y ^= View.Random;
 
-	float3 DiffuseAndSpecular = 0;
+	float3 NonSpecularContribution = 0;
+	float3 SpecularContribution = 0;
 
 	float3 DiffuseLookup =  TextureCubeSampleLevel(AmbientCubemap, AmbientCubemapSampler, GBuffer.WorldNormal, AbsoluteDiffuseMip).rgb;
 
 	// Diffuse
 	{
 		// we want to access the mip with the preconvolved diffuse lighting (coneangle=90 degree)
-		DiffuseAndSpecular += GBuffer.DiffuseColor * DiffuseLookup;
+		NonSpecularContribution += GBuffer.DiffuseColor * DiffuseLookup;
 	}
 
 	// Specular
@@ -343,7 +344,7 @@ void MainPS(in float4 UVAndScreenPos : TEXCOORD0, out float4 OutColor : SV_Targe
 		//float2 AB = PreIntegratedGF.SampleLevel( PreIntegratedGFSampler, float2( NoV, GBuffer.Roughness ), 0 ).rg;
 		//GBuffer.SpecularColor = GBuffer.SpecularColor * AB.x + AB.y * (1 - ClearCoat);
 		
-		DiffuseAndSpecular += SampleColor * EnvBRDF( GBuffer.SpecularColor, GBuffer.Roughness, NoV );
+		SpecularContribution += SampleColor * EnvBRDF(GBuffer.SpecularColor, GBuffer.Roughness, NoV);
 	}
 
 	BRANCH if( GBuffer.ShadingModelID == SHADINGMODELID_CLEAR_COAT )
@@ -362,12 +363,14 @@ void MainPS(in float4 UVAndScreenPos : TEXCOORD0, out float4 OutColor : SV_Targe
 
 		float LayerAttenuation = (1 - F);
 
-		DiffuseAndSpecular *= LayerAttenuation;
-		DiffuseAndSpecular += SampleColor * F;
+		NonSpecularContribution *= LayerAttenuation;
+		SpecularContribution *= LayerAttenuation;
+		SpecularContribution += SampleColor * F;
 	}
 
 #if IMPORTANCE_SAMPLE
-	DiffuseAndSpecular = ImageBasedLightingMIS( GBuffer, -ScreenVector, GBuffer.WorldNormal, Random );
+	NonSpecularContribution = 0;
+	SpecularContribution = ImageBasedLightingMIS( GBuffer, -ScreenVector, GBuffer.WorldNormal, Random );
 #endif
 	
 	// apply darkening from ambient occlusion (does not use PostprocessInput1 to set white texture if SSAO is off)
@@ -380,17 +383,18 @@ void MainPS(in float4 UVAndScreenPos : TEXCOORD0, out float4 OutColor : SV_Targe
 		float DependentSplit = 0.5f;
 
 		float3 SubsurfaceColor = DecodeSubsurfaceColor( GBuffer.CustomData );
-		float3 SubsurfaceLighting = 0;
 
 		// view independent (shared lookup for diffuse for better performance
-		SubsurfaceLighting += DiffuseLookup * (DependentSplit);
+		NonSpecularContribution += DiffuseLookup * SubsurfaceColor * (DependentSplit);
 		// view dependent (blurriness is hard coded)
-		SubsurfaceLighting += TextureCubeSampleLevel(AmbientCubemap, AmbientCubemapSampler, ScreenVector, AbsoluteDiffuseMip - 2.5f).rgb * AmbientOcclusion * (1.0f - DependentSplit);
-
-		DiffuseAndSpecular += SubsurfaceColor * SubsurfaceLighting;
+		SpecularContribution += TextureCubeSampleLevel(AmbientCubemap, AmbientCubemapSampler, ScreenVector, AbsoluteDiffuseMip - 2.5f).rgb * SubsurfaceColor * (AmbientOcclusion * (1.0f - DependentSplit));
 	}
-	
-	OutColor = float4(DiffuseAndSpecular * AmbientCubemapColor.rgb, 0);
+
+	FLightAccumulator LightAccumulator = (FLightAccumulator)0;
+
+	LightAccumulator_Add(LightAccumulator, NonSpecularContribution, SpecularContribution, AmbientCubemapColor.rgb);
+
+	OutColor = LightAccumulator_GetResult(LightAccumulator);
 
 	OutColor *= AmbientOcclusion;
 }
diff --git a/Engine/Shaders/PostProcessSubsurface.usf b/Engine/Shaders/PostProcessSubsurface.usf
--- a/Engine/Shaders/PostProcessSubsurface.usf
+++ b/Engine/Shaders/PostProcessSubsurface.usf
@@ -8,6 +8,33 @@
 #include "PostProcessCommon.usf"
 #include "DeferredShadingCommon.usf"
 
+/** RGBA8 linear texture containing random normals */
+Texture2D SSProfilesTexture;
+SamplerState SSProfilesTextureSampler;
+
+// ------------------------------------------
+
+// setup for "SeparableSSS.usf"
+
+#define SSSS_HLSL_4 1
+//#define SSSS_STREGTH_SOURCE 1	// todo: masking fom LightingModel or SSSColor
+#define SSSS_STREGTH_SOURCE (GetScreenSpaceData(texcoord).GBuffer.ShadingModelID >= SHADINGMODELID_SUBSURFACE)
+
+// 0:speculars leak into SSS / 1:requires alpha channel for scenecolor
+// #define SSSS_SPECULAR_CORRECTION 1
+
+// 0: faster
+// 1: no color bleeding in z direction
+#define SSSS_FOLLOW_SURFACE 1
+
+// from [REDACTED_PROFILE_URL]/separable-sss/tree/master/Demo
+// [REDACTED_PERSON] [REDACTED_PROFILE_URL]
+#include "SeparableSSS.usf"
+
+// ------------------------------------------
+
+
+
 #if SSSS_SPECULAR_CORRECTION == 1
 	// requires alpha channel for scenecolor
 #elif SSSS_SPECULAR_CORRECTION == 0
@@ -122,61 +149,31 @@ void MainPS(float4 UVAndScreenPos : TEXCOORD0, out float4 OutColor : SV_Target0)
 	float2 UV = UVAndScreenPos.xy;
 
 	FScreenSpaceData ScreenSpaceData = GetScreenSpaceData(UV);
-	
-	float2 Offset = PostprocessInput0Size.zw;
 
-	const int Count = 20;
-	float AccumSum = 0.0001f;
+	// call into "SeparableSSS.usf"
 
-	// scale in world space, todo: expose scale
-	float Radius = SSSParams.x * saturate(45 / ScreenSpaceData.GBuffer.Depth);
-
-
-	LOOP for(int i = -Count; i <= Count; ++i)
-	{
+	// in world units
+	float sssWidth = SSSParams.x * 0.01f;
 #if METHOD == 0
-		int x = i, y = 0;
+	// first pass
+	float2 Direction = float2(1, 0);
 #else
-		int x = 0, y = i;
+	// second pass
+	float2 Direction = float2(0, 1);
 #endif
+	OutColor = SSSSBlurPS(UV, sssWidth, Direction, false);
 
-		{
-			float2 NormDelta = float2(x, y) / Count;
-
-			float r = length(NormDelta);
-
-			//		float2 LocalUV = UV + float2(x, y) * Offset;
-			float2 LocalUV = UV + NormDelta * Radius * Offset;
+	OutColor.a = 1.0f;
 
-			FScreenSpaceData LocalScreenSpaceData = GetScreenSpaceData(LocalUV);
+	// Lighting
 
-			// hack
-			//		float Subsurface = ScreenSpaceData.GBuffer.Roughness > 0.4f;
-			//		float Subsurface = (1 - ScreenSpaceData.GBuffer.DecalMask) > 0.5f && ScreenSpaceData.GBuffer.Depth < 1000.0f;
-			FLATTEN if(LocalScreenSpaceData.GBuffer.ShadingModelID == SHADINGMODELID_SUBSURFACE)
-			{
-				float3 Color = ColorOverDistance(DecodeSubsurfaceColor(LocalScreenSpaceData.GBuffer.CustomData), r);
-				float ColorLum = Luminance(Color);
+#if METHOD > 0
+	// on second pass we recombine with the scene color
 
-				OutColor += float4(Color, 1) * float4(SampleDiffuseLighting(LocalUV), 1);	
-				AccumSum += ColorLum;
-			}
-		}
-	}
-
-	OutColor.rgb /= AccumSum;
-	OutColor.a /= (2 * Count + 1);
-
-
-#if METHOD >= 1
-	
 	float4 SceneColor4 = Texture2DSample(PostprocessInput1, PostprocessInput1Sampler, UV);
 	float4 SSSColor = OutColor;
 
-	SSSColor.a = (ScreenSpaceData.GBuffer.ShadingModelID == SHADINGMODELID_SUBSURFACE);
-
-	// partly translucent surface (0.9 means 10% of the opaque is blended in)
-	SSSColor.a *= 0.3f;
+	SSSColor.a = (ScreenSpaceData.GBuffer.ShadingModelID >= SHADINGMODELID_SUBSURFACE);
 
 #if METHOD > 1
 	// we took the specular highlights out, now we add them back in
diff --git a/Engine/Shaders/SeparableSSS.usf b/Engine/Shaders/SeparableSSS.usf
new file mode 100644
--- /dev/null
+++ b/Engine/Shaders/SeparableSSS.usf
@@ -0,0 +1,501 @@
+/**
+ * Copyright (C) 2012 [REDACTED_PERSON] ([REDACTED_EMAIL])
+ * Copyright (C) 2012 Diego Gutierrez ([REDACTED_EMAIL])
+ * All rights reserved.
+ *
+ * Redistribution and use in source and binary forms, with or without
+ * modification, are permitted provided that the following conditions are met:
+ * 
+ *    1. Redistributions of source code must retain the above copyright notice,
+ *       this list of conditions and the following disclaimer.
+ *
+ *    2. Redistributions in binary form must reproduce the following disclaimer
+ *       in the documentation and/or other materials provided with the 
+ *       distribution:
+ *
+ *       "Uses Separable SSS. Copyright (C) 2012 by [REDACTED_PERSON] and Diego
+ *        Gutierrez."
+ *
+ * THIS SOFTWARE IS PROVIDED BY THE COPYRIGHT HOLDERS AND CONTRIBUTORS ``AS 
+ * IS'' AND ANY EXPRESS OR IMPLIED WARRANTIES, INCLUDING, BUT NOT LIMITED TO, 
+ * THE IMPLIED WARRANTIES OF MERCHANTABILITY AND FITNESS FOR A PARTICULAR 
+ * PURPOSE ARE DISCLAIMED. IN NO EVENT SHALL COPYRIGHT HOLDERS OR CONTRIBUTORS 
+ * BE LIABLE FOR ANY DIRECT, INDIRECT, INCIDENTAL, SPECIAL, EXEMPLARY, OR 
+ * CONSEQUENTIAL DAMAGES (INCLUDING, BUT NOT LIMITED TO, PROCUREMENT OF 
+ * SUBSTITUTE GOODS OR SERVICES; LOSS OF USE, DATA, OR PROFITS; OR BUSINESS 
+ * INTERRUPTION) HOWEVER CAUSED AND ON ANY THEORY OF LIABILITY, WHETHER IN 
+ * CONTRACT, STRICT LIABILITY, OR TORT (INCLUDING NEGLIGENCE OR OTHERWISE) 
+ * ARISING IN ANY WAY OUT OF THE USE OF THIS SOFTWARE, EVEN IF ADVISED OF THE 
+ * POSSIBILITY OF SUCH DAMAGE.
+ *
+ * The views and conclusions contained in the software and documentation are 
+ * those of the authors and should not be interpreted as representing official
+ * policies, either expressed or implied, of the copyright holders.
+ */
+
+
+/**
+ *                  _______      _______      _______      _______
+ *                 /       |    /       |    /       |    /       |
+ *                |   (----    |   (----    |   (----    |   (----
+ *                 \   \        \   \        \   \        \   \
+ *              ----)   |    ----)   |    ----)   |    ----)   |
+ *             |_______/    |_______/    |_______/    |_______/
+ *
+ *        S E P A R A B L E   S U B S U R F A C E   S C A T T E R I N G
+ *
+ *                           [REDACTED_PROFILE_URL]
+ *
+ * Hi, thanks for your interest in Separable SSS!
+ *
+ * It's a simple shader composed of two components:
+ *
+ * 1) A transmittance function, 'SSSSTransmittance', which allows to calculate
+ *    light transmission in thin slabs, useful for ears and nostrils. It should
+ *    be applied during the main rendering pass as follows:
+ *
+ *        float3 t = albedo.rgb * lights[i].color * attenuation * spot;
+ *        color.rgb += t * SSSSTransmittance(...)
+ *
+ *    (See 'Main.fx' for more details).
+ *
+ * 2) A simple two-pass reflectance post-processing shader, 'SSSSBlur*', which
+ *    softens the skin appearance. It should be applied as a regular
+ *    post-processing effect like bloom (the usual framebuffer ping-ponging):
+ *
+ *    a) The first pass (horizontal) must be invoked by taking the final color
+ *       framebuffer as input, and storing the results into a temporal
+ *       framebuffer.
+ *    b) The second pass (vertical) must be invoked by taking the temporal
+ *       framebuffer as input, and storing the results into the original final
+ *       color framebuffer.
+ *
+ *    Note that This SSS filter should be applied *before* tonemapping.
+ *
+ * Before including SeparableSSS.h you'll have to setup the target. The
+ * following targets are available:
+ *         SMAA_HLSL_3
+ *         SMAA_HLSL_4
+ *         SMAA_GLSL_3
+ *
+ * For more information of what's under the hood, you can check the following
+ * URLs (but take into account that the shader has evolved a little bit since
+ * these publications):
+ *
+ * 1) Reflectance: [REDACTED_PROFILE_URL]
+ * 2) Transmittance: [REDACTED_PROFILE_URL]
+ *
+ * If you've got any doubts, just contact us!
+ */
+
+
+//-----------------------------------------------------------------------------
+// Configurable Defines
+
+/**
+ * SSSS_FOV must be set to the value used to render the scene.
+ */
+#ifndef SSSS_FOVY
+#define SSSS_FOVY 20.0
+#endif
+
+/**
+ * Light diffusion should occur on the surface of the object, not in a screen 
+ * oriented plane. Setting SSSS_FOLLOW_SURFACE to 1 will ensure that diffusion
+ * is more accurately calculated, at the expense of more memory accesses.
+ */
+#ifndef SSSS_FOLLOW_SURFACE
+#define SSSS_FOLLOW_SURFACE 0
+#endif
+
+/**
+ * This define allows to specify a different source for the SSS strength
+ * (instead of using the alpha channel of the color framebuffer). This is
+ * useful when the alpha channel of the mian color buffer is used for something
+ * else.
+ */
+#ifndef SSSS_STREGTH_SOURCE
+#define SSSS_STREGTH_SOURCE (colorM.a)
+#endif
+
+/**
+ * If SSSS_N_KERNELWEIGHTCOUNT is defined at this point, a custom filter kernel must be
+ * set by the runtime.
+ */
+#ifdef SSSS_N_KERNELWEIGHTCOUNT
+/**
+ * Filter kernel layout is as follows:
+ *   - Weights in the RGB channels.
+ *   - Offsets in the A channel.
+ */
+float4 kernel[SSSS_N_KERNELWEIGHTCOUNT]; 
+#else
+/**
+ * Here you have ready-to-use kernels for quickstarters. Three kernels are 
+ * readily available, with varying quality.
+ * To create new kernels take a look into SSS::calculateKernel, or simply
+ * push CTRL+C in the demo to copy the customized kernel into the clipboard.
+ *
+ * Note: these preset kernels are not used by the demo. They are calculated on
+ * the fly depending on the selected values in the interface, by directly using
+ * SSS::calculateKernel.
+ *
+ * Quality ranges from 0 to 2, being 2 the highest quality available.
+ * The quality is with respect to 1080p; for 720p Quality=0 suffices.
+ */
+#define SSSS_QUALITY 2	// todo: expose
+
+#if SSSS_QUALITY == 2
+#define SSSS_N_KERNELWEIGHTCOUNT 13
+static float4 kernel[] = {
+    float4(0.530605, 0.613514, 0.739601, 0),
+    float4(0.0211412, 0.0459286, 0.0378196, 0.0208333),
+    float4(0.0402784, 0.0657244, 0.04631, 0.0833333),
+    float4(0.0493588, 0.0367726, 0.0219485, 0.1875),
+    float4(0.0410172, 0.0199899, 0.0118481, 0.333333),
+    float4(0.0263642, 0.0119715, 0.00684598, 0.520833),
+    float4(0.017924, 0.00711691, 0.00347194, 0.75),
+    float4(0.0128496, 0.00356329, 0.00132016, 1.02083),
+    float4(0.0094389, 0.00139119, 0.000416598, 1.33333),
+    float4(0.00700976, 0.00049366, 0.000151938, 1.6875),
+    float4(0.00500364, 0.00020094, 5.28848e-005, 2.08333),
+    float4(0.00333804, 7.85443e-005, 1.2945e-005, 2.52083),
+    float4(0.000973794, 1.11862e-005, 9.43437e-007, 3),
+};
+#elif SSSS_QUALITY == 1
+#define SSSS_N_KERNELWEIGHTCOUNT 9
+static float4 kernel[] = {
+    float4(0.536343, 0.624624, 0.748867, 0),
+    float4(0.0324462, 0.0656718, 0.0532821, 0.03125),
+    float4(0.0582416, 0.0659959, 0.0411329, 0.125),
+    float4(0.0571056, 0.0287432, 0.0172844, 0.28125),
+    float4(0.0347317, 0.0151085, 0.00871983, 0.5),
+    float4(0.0216301, 0.00794618, 0.00376991, 0.78125),
+    float4(0.0144609, 0.00317269, 0.00106399, 1.125),
+    float4(0.0100386, 0.000914679, 0.000275702, 1.53125),
+    float4(0.00317394, 0.000134823, 3.77269e-005, 2),
+};
+#elif SSSS_QUALITY == 0
+#define SSSS_N_KERNELWEIGHTCOUNT 6
+static float4 kernel[] = {
+    float4(0.560479, 0.669086, 0.784728, 0),
+    float4(0.0771802, 0.113491, 0.0793803, 0.08),
+    float4(0.0821904, 0.0358608, 0.0209261, 0.32),
+    float4(0.03639, 0.0130999, 0.00643685, 0.72),
+    float4(0.0192831, 0.00282018, 0.00084214, 1.28),
+    float4(0.00471691, 0.000184771, 5.07565e-005, 2),
+};
+#else
+#error Quality must be one of {0,1,2}
+#endif
+#endif
+
+float4 Quantize8Bit(float4 In)
+{
+	return floor(In * 255.999f) / 255.0f;
+}
+
+float4 GetKernel(uint Index)
+{
+	const float4 TableMax = float4(0.784728f, 0.784728f, 0.784728f, 3.0f);
+	
+	// profiled on NV670 fullscreen, one pass (total: 2 pass), samples: 13+1+13, large head on screen, r.SSSS 200
+
+	// texture lookup
+//	float4 Value = SSProfilesTexture.Load(int3(Index, 0, 0)) * TableMax;		// 0.88ms
+	// simulated  8 bit lookup
+//	float4 Value = Quantize8Bit(kernel[Index] / TableMax) * TableMax;			// 1.33ms
+	// float reference
+	float4 Value = kernel[Index];												// 0.85ms
+
+	return Value;
+}
+
+//-----------------------------------------------------------------------------
+// Porting Functions
+
+#if SSSS_HLSL_3 == 1
+#define SSSSTexture2D sampler2D
+#define SSSSSampleLevelZero(tex, coord) tex2Dlod(tex, float4(coord, 0.0, 0.0))
+#define SSSSSampleLevelZeroPoint(tex, coord) tex2Dlod(tex, float4(coord, 0.0, 0.0))
+#define SSSSSample(tex, coord) tex2D(tex, coord)
+#define SSSSSamplePoint(tex, coord) tex2D(tex, coord)
+#define SSSSSampleLevelZeroOffset(tex, coord, offset) tex2Dlod(tex, float4(coord + offset * SSSS_PIXEL_SIZE, 0.0, 0.0))
+#define SSSSSampleOffset(tex, coord, offset) tex2D(tex, coord + offset * SSSS_PIXEL_SIZE)
+#define SSSSLerp(a, b, t) lerp(a, b, t)
+#define SSSSSaturate(a) saturate(a)
+#define SSSSMad(a, b, c) mad(a, b, c)
+#define SSSSMul(v, m) mul(v, m)
+#define SSSS_FLATTEN [flatten]
+#define SSSS_BRANCH [branch]
+#define SSSS_UNROLL [unroll]
+#endif
+#if SSSS_HLSL_4 == 1
+#define SSSSTexture2D Texture2D
+#define SSSSSampleLevelZero(tex, coord) tex.SampleLevel(LinearSampler, coord, 0)
+#define SSSSSampleLevelZeroPoint(tex, coord) tex.SampleLevel(PointSampler, coord, 0)
+
+// EPIC changed start
+#define SSSSSampleShadowmap(coord)	float4(0,0,0,0)	// todo: not used yet
+#define SSSSSampleMasking(texcoord)	(GetScreenSpaceData(texcoord).GBuffer.ShadingModelID >= SHADINGMODELID_SUBSURFACE)	// todo: masking doesn't work with filtering ideally put in alpa to get it filtered
+#define SSSSSampleSceneColor(coord) (SSSSSampleMasking(coord) * Texture2DSample(PostprocessInput0, PostprocessInput0Sampler, coord))
+#define SSSSSampleSceneColorPoint(coord) ( Texture2DSample(PostprocessInput0, PostprocessInput0Sampler, coord))		// todo: could be made point
+#define SSSSSampleSceneDepth(coord) CalcSceneDepth(coord)		// todo: could be made linear
+#define SSSSSampleSceneDepthPoint(coord) CalcSceneDepth(coord)
+// EPIC changed end
+
+#define SSSSSampleLevelZeroOffset(tex, coord, offset) tex.SampleLevel(LinearSampler, coord, 0, offset)
+#define SSSSSampleOffset(tex, coord, offset) SSSSSampleLevelZeroOffset(tex, coord, offset)
+#define SSSSLerp(a, b, t) lerp(a, b, t)
+#define SSSSSaturate(a) saturate(a)
+#define SSSSMad(a, b, c) mad(a, b, c)
+#define SSSSMul(v, m) mul(v, m)
+#define SSSS_FLATTEN [flatten]
+#define SSSS_BRANCH [branch]
+#define SSSS_UNROLL [unroll]
+#endif
+#if SSSS_GLSL_3 == 1
+#define SSSSTexture2D sampler2D
+#define SSSSSampleLevelZero(tex, coord) textureLod(tex, coord, 0.0)
+#define SSSSSampleLevelZeroPoint(tex, coord) textureLod(tex, coord, 0.0)
+#define SSSSSample(tex, coord) texture(tex, coord)
+#define SSSSSamplePoint(tex, coord) texture(tex, coord)
+#define SSSSSampleLevelZeroOffset(tex, coord, offset) textureLodOffset(tex, coord, 0.0, offset)
+#define SSSSSampleOffset(tex, coord, offset) texture(tex, coord, offset)
+#define SSSSLerp(a, b, t) mix(a, b, t)
+#define SSSSSaturate(a) clamp(a, 0.0, 1.0)
+#define SSSSMad(a, b, c) (a * b + c)
+#define SSSSMul(v, m) (m * v)
+#define SSSS_FLATTEN
+#define SSSS_BRANCH
+#define SSSS_UNROLL
+#define float2 vec2
+#define float3 vec3
+#define float4 vec4
+#define int2 ivec2
+#define int3 ivec3
+#define int4 ivec4
+#define float4x4 mat4x4
+#endif
+
+
+//-----------------------------------------------------------------------------
+// Separable SSS Transmittance Function
+
+float3 SSSSTransmittance(
+        /**
+         * This parameter allows to control the transmittance effect. Its range
+         * should be 0..1. Higher values translate to a stronger effect.
+         */
+        float translucency,
+
+        /**
+         * This parameter should be the same as the 'SSSSBlurPS' one. See below
+         * for more details.
+         */
+        float sssWidth,
+
+        /**
+         * Position in world space.
+         */
+        float3 worldPosition,
+
+        /**
+         * Normal in world space.
+         */
+        float3 worldNormal,
+
+        /**
+         * Light vector: lightWorldPosition - worldPosition.
+         */
+        float3 light,
+		
+// EPIC changed start
+        /**
+         * Linear 0..1 shadow map.
+         */
+//        SSSSTexture2D shadowMap,
+// EPIC changed end
+
+        /**
+         * Regular world to light space matrix.
+         */
+        float4x4 lightViewProjection,
+
+        /**
+         * Far plane distance used in the light projection matrix.
+         */
+        float lightFarPlane) {
+    /**
+     * Calculate the scale of the effect.
+     */
+    float scale = 8.25 * (1.0 - translucency) / sssWidth;
+       
+    /**
+     * First we shrink the position inwards the surface to avoid artifacts:
+     * (Note that this can be done once for all the lights)
+     */
+    float4 shrinkedPos = float4(worldPosition - 0.005 * worldNormal, 1.0);
+
+    /**
+     * Now we calculate the thickness from the light point of view:
+     */
+    float4 shadowPosition = SSSSMul(shrinkedPos, lightViewProjection);
+    float d1 = SSSSSampleShadowmap(shadowPosition.xy / shadowPosition.w).r; // 'd1' has a range of 0..1
+    float d2 = shadowPosition.z; // 'd2' has a range of 0..'lightFarPlane'
+    d1 *= lightFarPlane; // So we scale 'd1' accordingly:
+    float d = scale * abs(d1 - d2);
+
+    /**
+     * Armed with the thickness, we can now calculate the color by means of the
+     * precalculated transmittance profile.
+     * (It can be precomputed into a texture, for maximum performance):
+     */
+    float dd = -d * d;
+    float3 profile = float3(0.233, 0.455, 0.649) * exp(dd / 0.0064) +
+                     float3(0.1,   0.336, 0.344) * exp(dd / 0.0484) +
+                     float3(0.118, 0.198, 0.0)   * exp(dd / 0.187)  +
+                     float3(0.113, 0.007, 0.007) * exp(dd / 0.567)  +
+                     float3(0.358, 0.004, 0.0)   * exp(dd / 1.99)   +
+                     float3(0.078, 0.0,   0.0)   * exp(dd / 7.41);
+
+    /** 
+     * Using the profile, we finally approximate the transmitted lighting from
+     * the back of the object:
+     */
+    return profile * SSSSSaturate(0.3 + dot(light, -worldNormal));
+}
+
+//-----------------------------------------------------------------------------
+// Separable SSS Reflectance Vertex Shader
+
+void SSSSBlurVS(float4 position,
+                out float4 svPosition,
+                inout float2 texcoord) {
+    svPosition = position;
+}
+
+//-----------------------------------------------------------------------------
+// Separable SSS Reflectance Pixel Shader
+
+float4 SSSSBlurPS(
+        /**
+         * The usual quad texture coordinates.
+         */
+        float2 texcoord,
+
+// EPIC changed start
+        /**
+         * This is a SRGB or HDR color input buffer, which should be the final
+         * color frame, resolved in case of using multisampling. The desired
+         * SSS strength should be stored in the alpha channel (1 for full
+         * strength, 0 for disabling SSS). If this is not possible, you an
+         * customize the source of this value using SSSS_STREGTH_SOURCE.
+         *
+         * When using non-SRGB buffers, you
+         * should convert to linear before processing, and back again to gamma
+         * space before storing the pixels (see Chapter 24 of GPU Gems 3 for
+         * more info)
+         *
+         * IMPORTANT: WORKING IN A NON-LINEAR SPACE WILL TOTALLY RUIN SSS!
+         */
+//        SSSSTexture2D colorTex,
+		/**
+         * The linear depth buffer of the scene, resolved in case of using
+         * multisampling. The resolve should be a simple average to avoid
+         * artifacts in the silhouette of objects.
+         */
+//        SSSSTexture2D depthTex,
+// EPIC changed end
+
+        /**
+         * This parameter specifies the global level of subsurface scattering
+         * or, in other words, the width of the filter. It's specified in
+         * world space units.
+         */
+        float sssWidth,
+
+        /**
+         * Direction of the blur:
+         *   - First pass:   float2(1.0, 0.0)
+         *   - Second pass:  float2(0.0, 1.0)
+         */
+        float2 dir,
+
+        /**
+         * This parameter indicates whether the stencil buffer should be
+         * initialized. Should be set to 'true' for the first pass if not
+         * previously initialized, to enable optimization of the second
+         * pass.
+         */
+        bool initStencil) {
+
+    // Fetch color of current pixel:
+    float4 colorM = SSSSSampleSceneColorPoint(texcoord);
+
+    // Initialize the stencil buffer in case it was not already available:
+    if (initStencil) // (Checked in compile time, it's optimized away)
+        if (SSSS_STREGTH_SOURCE == 0.0) discard;
+
+    // Fetch linear depth of current pixel:
+    float depthM = SSSSSampleSceneDepthPoint(texcoord).r;
+
+    // Calculate the sssWidth scale (1.0 for a unit plane sitting on the
+    // projection window):
+    float distanceToProjectionWindow = 1.0 / tan(0.5 * radians(SSSS_FOVY));
+    float scale = distanceToProjectionWindow / depthM;
+
+    // Calculate the final step to fetch the surrounding pixels:
+    float2 finalStep = sssWidth * scale * dir;
+    finalStep *= SSSS_STREGTH_SOURCE; // Modulate it using the alpha channel.
+    finalStep *= 1.0 / 3.0; // Divide by 3 as the kernels range from -3 to 3.
+
+    // Accumulate the center sample:
+    float4 colorBlurred = colorM;
+	colorBlurred.rgb *= GetKernel(0).rgb;
+
+	// hack
+   // colorBlurred.rgb = 0;
+
+	// Accumulate the other samples:
+    SSSS_UNROLL
+	for (int i = 1; i < SSSS_N_KERNELWEIGHTCOUNT; i++) 
+	{
+		half4 Kernel = GetKernel(i);
+
+		float3 LocalAccum = 0;
+
+		// The kernel is symtrical, we want to use that property.
+		// Half the GetKernel() calls (more expensive if done by texture)
+		// Half the weighting computations (saves 3mul per lookup sample)
+		// We don't sum up with (x+=y;x+=z;) but with (x+=(y+z);), order allows for better texture lookup grouping
+		SSSS_UNROLL
+		for (int Side = -1; Side <= 1; Side += 2)
+		{
+			// Fetch color and depth for current sample:
+			float2 offset = texcoord + Kernel.a * finalStep * Side;
+			float4 color = SSSSSampleSceneColor(offset);
+
+#if SSSS_FOLLOW_SURFACE == 1
+			// If the difference in depth is huge, we lerp color back to "colorM":
+			float depth = SSSSSampleSceneDepth(offset).r;
+			float s = SSSSSaturate(6000.0f / 400000 * distanceToProjectionWindow *
+				//        float s = SSSSSaturate(300.0f/400000 * distanceToProjectionWindow *
+				sssWidth * abs(depthM - depth));
+			color.rgb = SSSSLerp(color.rgb, colorM.rgb, s);
+#endif
+			// Accumulate left and right 
+			LocalAccum += color.rgb;
+		}
+
+		// Accumulate to final value
+		colorBlurred.rgb += Kernel.rgb * LocalAccum;
+	}
+
+	// hack - normalize
+//	colorBlurred.rgb /= 1.0f - kernel[0].rgb;
+
+    return colorBlurred;
+}
+
+//-----------------------------------------------------------------------------
\ No newline at end of file
