Tweaked clear coat

diff --git a/Engine/Shaders/BasePassPixelShader.usf b/Engine/Shaders/BasePassPixelShader.usf
--- a/Engine/Shaders/BasePassPixelShader.usf
+++ b/Engine/Shaders/BasePassPixelShader.usf
@@ -733,7 +733,7 @@ void Main(
 			float NoV = saturate( dot( MaterialParameters.WorldNormal, MaterialParameters.CameraVector ) );
 
 			// Approximation of refraction's effect on EnvBRDF
-			float RefractionScale = saturate( (NoV - 1) * lerp( 1.3, 0.5, Roughness ) + 1 );
+			float RefractionScale = ( (NoV * 0.5 + 0.5) * NoV - 1 ) * saturate( 1.25 - 1.25 * Roughness ) + 1;
 			RefractionScale = lerp( 1, RefractionScale, ClearCoat );
 
 #if DIFFUSE_SPEC_INPUTS
@@ -747,6 +747,8 @@ void Main(
 
 			// Approximation of absorption integral, tuned for Roughness=0.4
 			float3 Absorption = AbsorptionColor * ( (NoV - 1) * 0.85 * ( 1 - lerp( AbsorptionColor, Square(AbsorptionColor), -0.78 ) ) + 1 );
+			//float Toe = 1 - exp2( -9 * NoV - 3.5 );
+			//float3 Absorption = Toe * AbsorptionColor * ( (NoV - 1) * 0.85 * ( 1 - lerp( AbsorptionColor, Square(AbsorptionColor), -0.78 ) ) + lerp( 0.93, 0.96, AbsorptionColor * 2 - 1 ) );
 
 #if DIFFUSE_SPEC_INPUTS
 			GBuffer.DiffuseColor *= Absorption;
diff --git a/Engine/Shaders/PostProcessAmbient.usf b/Engine/Shaders/PostProcessAmbient.usf
--- a/Engine/Shaders/PostProcessAmbient.usf
+++ b/Engine/Shaders/PostProcessAmbient.usf
@@ -97,11 +97,10 @@ float3 ClearCoatShading( FGBufferData GBuffer, float Roughness, float3 L, float3
 {
 	const float ClearCoat			= GBuffer.CustomData.x;
 	const float ClearCoatRoughness	= GBuffer.CustomData.y;
-	float Film = 1;
-	float MetalSpec = 0.9;
-
-	Film *= ClearCoat;
+	const float Film = 1 * ClearCoat;
+	const float MetalSpec = 0.9;
 
+#if 0
 	float3 H = normalize(V + L);
 	float NoL = saturate( dot(N, L) );
 	float NoV = saturate( dot(N, V) );
@@ -116,6 +115,7 @@ float3 ClearCoatShading( FGBufferData GBuffer, float Roughness, float3 L, float3
 	float F0 = 0.04;
 	float Fc = pow( 1 - VoH, 5 );
 	float F = Fc + (1 - Fc) * F0;
+	F *= ClearCoat;
 
 	float Fr1 = D * Vis * F;
 
@@ -124,12 +124,102 @@ float3 ClearCoatShading( FGBufferData GBuffer, float Roughness, float3 L, float3
 	// Generalized microfacet specular
 	float D2 = D_GGX( Roughness, NoH );
 	float Vis2 = Vis_Schlick( Roughness, NoV, NoL );
-	//float3 F2 = F_Schlick( GBuffer.SpecularColor, VoH );
-	float3 F2 = (1 - Fc) * GBuffer.SpecularColor;
+	float3 F2 = saturate( 50.0 * GBuffer.SpecularColor.g ) * Fc + (1 - Fc) * GBuffer.SpecularColor;
+
+	float3 Fr2 = Diffuse_Lambert( GBuffer.DiffuseColor ) + (D2 * Vis2) * F2;
+	
+	return Fr1 + Fr2 * LayerAttenuation;
+#else
+	float3 H = normalize(V + L);
+	float NoL = saturate( dot(N, L) );
+	float NoV = saturate( dot(N, V) );
+	float NoH = saturate( dot(N, H) );
+	float VoH = saturate( dot(V, H) );
+	
+	// Hard coded IOR of 1.5
+
+	// Generalized microfacet specular
+	float D = D_GGX( ClearCoatRoughness, NoH );
+	float Vis = Vis_Kelemen( VoH );
+	
+	// F_Schlick
+	float F0 = 0.04;
+	float Fc = pow( 1 - VoH, 5 );
+	float F = Fc + (1 - Fc) * F0;
+
+	float Fr1 = D * Vis * F;
+
+#if 1
+	// Refract rays
+	float3 L2 = refract( -L, -H, 1 / 1.5 );
+	float3 V2 = refract( -V, -H, 1 / 1.5 );
+
+	// LoH == VoH
+	float RefractBlend = sqrt( 4 * VoH*VoH + 5 ) / 3 + 2.0 / 3 * VoH;
+	//float3 L2 = RefractBlend * H - L / 1.5;
+	//float3 V2 = RefractBlend * H - V / 1.5;
+
+	float3 H2 = normalize( V2 + L2 );
+	float NoL2 = saturate( dot(N, L2) );
+	float NoV2 = saturate( dot(N, V2) );
+	float NoH2 = saturate( dot(N, H2) );
+	float VoH2 = saturate( dot(V2, H2) );
+#else
+	// Approximation
+	float RefractBlend = (0.22 * VoH + 0.7) * VoH + 0.745;	// 2 mad
+	// Dot products distribute. No need for L2 and V2.
+	float RefractNoH = RefractBlend * NoH;					// 1 mul
+	float NoL2 = saturate( RefractNoH - (1 / 1.5) * NoL );	// 1 mad
+	float NoV2 = saturate( RefractNoH - (1 / 1.5) * NoV );	// 1 mad
+	// Should refract H too but unimportant
+#endif
+
+	NoL2 = max( 0.001, NoL2 );
+	NoV2 = max( 0.001, NoV2 );
+
+	float3 AbsorptionColor = (1 - Film) + GBuffer.BaseColor * ( Film * (1 / MetalSpec) );
+	float  AbsorptionDist = rcp(NoV2) + rcp(NoL2);
+	float3 Absorption = pow( AbsorptionColor, 0.5 * AbsorptionDist );
+
+	// Approximation
+	//float  AbsorptionDist = ( NoV2 + NoL2 ) / ( NoV2 * NoL2 );
+	//float3 Absorption = AbsorptionColor * ( AbsorptionColor * (AbsorptionDist * 0.5 - 1) + (2 - 0.5 * AbsorptionDist) );
+	//float3 Absorption = AbsorptionColor + AbsorptionColor * (AbsorptionColor - 1) * (AbsorptionDist * 0.5 - 1);	// use for shared version
+	
+#if 1
+	float F21 = F_Schlick( 0.04, saturate( dot(V2, H) ) );
+
+	float k = Square( Roughness ) * 0.5;
+	float G_SchlickV2 = NoV2 / ( NoV2 * (1 - k) + k );
+	float G_SchlickL2 = NoL2 / ( NoL2 * (1 - k) + k );
+	float TotalInternalReflection = 1 - F21 * G_SchlickV2 * G_SchlickL2;
+	float3 LayerAttenuation = ( (1 - F) * TotalInternalReflection ) * Absorption;
+#else
+	// Approximation
+	float3 LayerAttenuation = (1 - F) * Absorption;
+#endif
+
+	// Approximation for IOR == 1.5
+	//SpecularColor = ChangeBaseMedium( SpecularColor, 1.5 );
+	//SpecularColor = saturate( ( 0.55 * SpecularColor + (0.45 * 1.08) ) * SpecularColor - (0.45 * 0.08) );
+	// Treat SpecularColor as relative to IOR. Artist compensates.
+	
+	// Generalized microfacet specular
+#if 1
+	float D2 = D_GGX( Roughness, NoH2 );
+	float Vis2 = Vis_Schlick( Roughness, NoV2, NoL2 );
+	float3 F2 = F_Schlick( MetalSpec, VoH2 );
+#else
+	// Approximation
+	float D2 = D_GGX( Roughness, NoH );
+	float Vis2 = Vis_Schlick( Roughness, NoV2, NoL2 );
+	float3 F2 = F_Schlick( GBuffer.SpecularColor, VoH );
+#endif
 
 	float3 Fr2 = Diffuse_Lambert( GBuffer.DiffuseColor ) + (D2 * Vis2) * F2;
 	
 	return Fr1 + Fr2 * LayerAttenuation;
+#endif
 }
 
 
@@ -195,8 +285,15 @@ float3 ImageBasedLightingMIS( FGBufferData GBuffer, float3 V, float3 N, uint2 Ra
 				}
 				float Weight = rcp( InvWeight );
 
-				float3 Shading = ClearCoatShading( GBuffer, GBuffer.Roughness, L, V, N );
-				//float3 Shading = StandardShading( GBuffer, GBuffer.Roughness, L, V, N );
+				float3 Shading = 0;
+				BRANCH if( GBuffer.ShadingModelID == SHADINGMODELID_CLEAR_COAT )
+				{
+					Shading = ClearCoatShading( GBuffer, GBuffer.Roughness, L, V, N );
+				}
+				else
+				{
+					Shading = StandardShading( GBuffer, GBuffer.Roughness, L, V, N );
+				}
 			
 				Lighting += SampleColor * Shading * ( NoL * Weight );
 			}
@@ -341,10 +438,16 @@ void MainPS(in float4 UVAndScreenPos : TEXCOORD0, out float4 OutColor : SV_Targe
 		float Mip = ComputeCubemapMipFromRoughness( GBuffer.Roughness, AmbientCubemapMipAdjust.w );
 		float3 SampleColor = TextureCubeSampleLevel( AmbientCubemap, AmbientCubemapSampler, R, Mip ).rgb;
 
-		//float2 AB = PreIntegratedGF.SampleLevel( PreIntegratedGFSampler, float2( NoV, GBuffer.Roughness ), 0 ).rg;
-		//GBuffer.SpecularColor = GBuffer.SpecularColor * AB.x + AB.y * (1 - ClearCoat);
-		
-		SpecularContribution += SampleColor * EnvBRDF(GBuffer.SpecularColor, GBuffer.Roughness, NoV);
+		BRANCH if( GBuffer.ShadingModelID == SHADINGMODELID_CLEAR_COAT )
+		{
+			float ClearCoat = GBuffer.CustomData.x;
+			float2 AB = PreIntegratedGF.SampleLevel( PreIntegratedGFSampler, float2( NoV, GBuffer.Roughness ), 0 ).rg;
+			SpecularContribution += SampleColor * GBuffer.SpecularColor * AB.x + AB.y * (1 - ClearCoat);
+		}
+		else
+		{
+			SpecularContribution += SampleColor * EnvBRDF( GBuffer.SpecularColor, GBuffer.Roughness, NoV );
+		}
 	}
 
 	BRANCH if( GBuffer.ShadingModelID == SHADINGMODELID_CLEAR_COAT )
@@ -369,8 +472,11 @@ void MainPS(in float4 UVAndScreenPos : TEXCOORD0, out float4 OutColor : SV_Targe
 	}
 
 #if IMPORTANCE_SAMPLE
-	NonSpecularContribution = 0;
-	SpecularContribution = ImageBasedLightingMIS( GBuffer, -ScreenVector, GBuffer.WorldNormal, Random );
+	if( GBuffer.ShadingModelID > 0 )
+	{
+		NonSpecularContribution = 0;
+		SpecularContribution = ImageBasedLightingMIS( GBuffer, -ScreenVector, GBuffer.WorldNormal, Random );
+	}
 #endif
 	
 	// apply darkening from ambient occlusion (does not use PostprocessInput1 to set white texture if SSAO is off)
