UE4 - Fix opengl issue

diff --git a/Engine/Shaders/SeparableSSS.usf b/Engine/Shaders/SeparableSSS.usf
--- a/Engine/Shaders/SeparableSSS.usf
+++ b/Engine/Shaders/SeparableSSS.usf
@@ -147,7 +147,7 @@ float4 kernel[SSSS_N_KERNELWEIGHTCOUNT];
 
 #if SSSS_QUALITY == 2
 #define SSSS_N_KERNELWEIGHTCOUNT 13
-static float4 kernel[] = {
+const static float4 kernel[] = {
     float4(0.530605, 0.613514, 0.739601, 0),
     float4(0.0211412, 0.0459286, 0.0378196, 0.0208333),
     float4(0.0402784, 0.0657244, 0.04631, 0.0833333),
@@ -164,7 +164,7 @@ static float4 kernel[] = {
 };
 #elif SSSS_QUALITY == 1
 #define SSSS_N_KERNELWEIGHTCOUNT 9
-static float4 kernel[] = {
+const static float4 kernel[] = {
     float4(0.536343, 0.624624, 0.748867, 0),
     float4(0.0324462, 0.0656718, 0.0532821, 0.03125),
     float4(0.0582416, 0.0659959, 0.0411329, 0.125),
@@ -177,7 +177,7 @@ static float4 kernel[] = {
 };
 #elif SSSS_QUALITY == 0
 #define SSSS_N_KERNELWEIGHTCOUNT 6
-static float4 kernel[] = {
+const static float4 kernel[] = {
     float4(0.560479, 0.669086, 0.784728, 0),
     float4(0.0771802, 0.113491, 0.0793803, 0.08),
     float4(0.0821904, 0.0358608, 0.0209261, 0.32),
