* fixed shader compile warning

diff --git a/Engine/Shaders/GpuSkinVertexFactory.usf b/Engine/Shaders/GpuSkinVertexFactory.usf
--- a/Engine/Shaders/GpuSkinVertexFactory.usf
+++ b/Engine/Shaders/GpuSkinVertexFactory.usf
@@ -415,7 +415,7 @@ float3x3 SkinTangents( FVertexFactoryInput Input, FVertexFactoryIntermediates In
 	// normal morph offset. 
 	// since normal delta can be in the range of (-2, 2), 
 	// we packed twice more with less precision, so scale back up by 2 
-	float3 DeltaTangentZ = TangentBias(Input.DeltaTangentZ)*2;
+	float3 DeltaTangentZ = TangentBias(Input.DeltaTangentZ.xyz)*2;
 
 	// calc new normal by offseting it with the delta
 	Tangents[1] = normalize( TangentZ.xyz + Input.DeltaTangentZ.w * DeltaTangentZ);
