* fixed shader compiler warning

diff --git a/Engine/Shaders/LocalVertexFactory.usf b/Engine/Shaders/LocalVertexFactory.usf
--- a/Engine/Shaders/LocalVertexFactory.usf
+++ b/Engine/Shaders/LocalVertexFactory.usf
@@ -315,7 +315,7 @@ float4 CalcWorldPosition(float4 Position)
 {
 #if USE_INSTANCING
 	float4x4 InstanceTransform = transpose(float4x4(InstanceTransform1, InstanceTransform2, InstanceTransform3, float4(0, 0, 0, 1)));
-	return TransformLocalToTranslatedWorld(mul(Position, InstanceTransform));
+	return TransformLocalToTranslatedWorld(mul(Position, InstanceTransform).xyz);
 #elif USE_SPLINEDEFORM
 	// Make transform for this point along spline
 	float4x3 SliceTransform = CalcSliceTransform(dot(Position.xyz, SplineMeshDir));
