* optimized ScreenSpaceSubsurfaceScattering shader cost

diff --git a/Engine/Shaders/LightAccumulator.usf b/Engine/Shaders/LightAccumulator.usf
--- a/Engine/Shaders/LightAccumulator.usf
+++ b/Engine/Shaders/LightAccumulator.usf
@@ -21,12 +21,20 @@
 
 struct FLightAccumulator
 {
+	// view independent lighting + view dependent lighting
 	float3 TotalSum;
 
-	// only actually used SUBSURFACE_CHANNEL_MODE != 0
+	// only actually used SUBSURFACE_CHANNEL_MODE == 1
 	// assumed to be compiled out otherwise (not compiled out with #if so we can use if() instead of #if for better readability and compiler error checking)
 	// input for ScreenSpaceSubsurfaceScattering
-	float3 NonSpecularSum;
+	// view independent lighting luminance
+	float ViewIndependentLightingLuminanceSum;
+
+	// only actually used SUBSURFACE_CHANNEL_MODE == 2
+	// assumed to be compiled out otherwise (not compiled out with #if so we can use if() instead of #if for better readability and compiler error checking)
+	// input for ScreenSpaceSubsurfaceScattering
+	// view independent lighting color
+	float3 ViewIndependentLightingSum;
 
 	// only used for development (not compiled out with #if so we can use if() instead of #if for better readability and compiler error checking)
 	// assumed to be compiled out otherwise
@@ -37,40 +45,47 @@ struct FLightAccumulator
 // accumulate light, can be called multiple times
 void LightAccumulator_Add(inout FLightAccumulator In, float3 NonSpecularContribution, float3 SpecularContribution, float3 CommonMultiplier)
 {
-	if (SUBSURFACE_CHANNEL_MODE == 0)
+	// 3 add, 3 mad
+	In.TotalSum += (NonSpecularContribution + SpecularContribution) * CommonMultiplier;
+
+	if (SUBSURFACE_CHANNEL_MODE == 1)
 	{
-		// 3 add, 3 mad
-		In.TotalSum += (NonSpecularContribution + SpecularContribution) * CommonMultiplier;
+		// 3 mad
+		In.ViewIndependentLightingLuminanceSum += Luminance(NonSpecularContribution * CommonMultiplier);
 	}
-	else
+	else if (SUBSURFACE_CHANNEL_MODE == 2)
 	{
-		// 3 add, 3 mad
-		In.TotalSum += (NonSpecularContribution + SpecularContribution) * CommonMultiplier;
 		// 3 mad
-		In.NonSpecularSum += NonSpecularContribution * CommonMultiplier;
+		In.ViewIndependentLightingSum += NonSpecularContribution * CommonMultiplier;
 	}
 }
 
 //
 // compute final value to store in the MRT0
+// @retrun RGB:SceneColor Specular and Diffuse, A:Non Specular SceneColor Luminance
 float4 LightAccumulator_GetResult(FLightAccumulator In)
 {
 	float4 Ret;
 
 	if (VISUALIZE_LIGHT_CULLING == 1)
 	{
-		// similar to VISUALIZE_LIGHT_CULLING in tile based deferred lighting
+		// a soft gradient from dark red to bright white, can be changed to be different
 		Ret = 0.1f * float4(1.0f, 0.25f, 0.075f, 0) * In.EstimatedCost;
 	}
 	else
 	{
 		Ret = float4(In.TotalSum, 0);
 
-		if (SUBSURFACE_CHANNEL_MODE != 0)
+		if (SUBSURFACE_CHANNEL_MODE == 1)
+		{
+			// RGB accumulated RGB HDR color, A: specular luminance for screenspace subsurface scattering
+			Ret.a = In.ViewIndependentLightingLuminanceSum;
+		}
+		else if (SUBSURFACE_CHANNEL_MODE == 2)
 		{
 			// RGB accumulated RGB HDR color, A: specular luminance for screenspace subsurface scattering
 			// 3 add,  1 mul, 2 mad, can be optimized to use 2 less temporary during accumulation and remove the 3 add
-			Ret.a = Luminance(In.TotalSum - In.NonSpecularSum);
+			Ret.a = Luminance(In.ViewIndependentLightingSum);
 			// todo, need second MRT for SUBSURFACE_CHANNEL_MODE==2
 		}
 	}
@@ -78,4 +93,23 @@ float4 LightAccumulator_GetResult(FLightAccumulator In)
 	return Ret;
 }
 
+// @param SceneColor4 RGB:SceneColor Specular and Diffuse, A:Non Specular SceneColor Luminance
+float3 LightAccumulator_ReconstructDiffuseLighting(float4 SceneColor4)
+{
+	// add small bias to avoid div by 0 (faster than max(), assuming we don't have negatives and NaN in there)
+	float TotalSceneColorLum = Luminance(SceneColor4.rgb) + 0.000001f;
+
+	return SceneColor4.rgb * (SceneColor4.a / TotalSceneColorLum);
+}
+
+// @param SceneColor4 RGB:SceneColor Specular and Diffuse, A:Non Specular SceneColor Luminance
+float3 LightAccumulator_ReconstructNonDiffuseLighting(float4 SceneColor4)
+{
+	// add small bias to avoid div by 0 (faster than max(), assuming we don't have negatives and NaN in there)
+	float TotalSceneColorLum = Luminance(SceneColor4.rgb) + 0.000001f;
+
+	return SceneColor4.rgb * saturate(1 - SceneColor4.a / TotalSceneColorLum);
+}
+
+
 #endif // __LIGHTACCUMULATOR_COMMON__
\ No newline at end of file
diff --git a/Engine/Shaders/PostProcessSubsurface.usf b/Engine/Shaders/PostProcessSubsurface.usf
--- a/Engine/Shaders/PostProcessSubsurface.usf
+++ b/Engine/Shaders/PostProcessSubsurface.usf
@@ -60,22 +60,6 @@ SamplerState SSProfilesTextureSampler;
 float4 SSSParams;
 
 
-float3 ReconstructDiffuseLighting(float4 SceneColor4)
-{
-	float SceneColorLum = Luminance(SceneColor4.rgb);
-
-	// max to avoid division by 0
-	return SceneColor4.rgb * max(0, SceneColorLum - SceneColor4.a) / max(0.0001f, SceneColorLum);
-}
-
-float3 ReconstructNonDiffuseLighting(float4 SceneColor4)
-{
-	float SceneColorLum = Luminance(SceneColor4.rgb);
-
-	// max to avoid division by 0
-	return SceneColor4.rgb * SceneColor4.a / max(0.0001f, SceneColorLum);
-}
-
 void SetupPS(in float4 UVAndScreenPos : TEXCOORD0, out float4 OutColor : SV_Target0)
 {
 	float2 UV = UVAndScreenPos.xy;
@@ -90,7 +74,7 @@ void SetupPS(in float4 UVAndScreenPos : TEXCOORD0, out float4 OutColor : SV_Targ
 
 #if SSSS_SPECULAR_CORRECTION == 1
 		// we take out the specular highlights
-		float3 Lighting = ReconstructDiffuseLighting(SceneColor4);
+		float3 Lighting = LightAccumulator_ReconstructDiffuseLighting(SceneColor4);
 #else
 		float3 Lighting = SceneColor4.rgb;
 #endif
@@ -177,7 +161,7 @@ void MainPS(float4 UVAndScreenPos : TEXCOORD0, out float4 OutColor : SV_Target0)
 
 #if METHOD > 1
 	// we took the specular highlights out, now we add them back in
-	SSSColor.rgb += ReconstructNonDiffuseLighting(SceneColor4);
+	SSSColor.rgb += LightAccumulator_ReconstructNonDiffuseLighting(SceneColor4);
 #endif
 
 	OutColor = float4(lerp(SceneColor4.rgb, SSSColor.rgb, SSSColor.a), 1);
diff --git a/Engine/Shaders/TiledDeferredLightShaders.usf b/Engine/Shaders/TiledDeferredLightShaders.usf
--- a/Engine/Shaders/TiledDeferredLightShaders.usf
+++ b/Engine/Shaders/TiledDeferredLightShaders.usf
@@ -95,7 +95,7 @@ void TiledDeferredLightingMain(
 	float4 HomogeneousWorldPosition = mul(float4(ScreenPosition * SceneDepth, SceneDepth, 1), View.ScreenToWorld);
 	float3 WorldPosition = HomogeneousWorldPosition.xyz / HomogeneousWorldPosition.w;
 
-	// RGB accumulated RGB HDR color, A: specular luminance for screenspace subsurface scattering
+	// RGB accumulated RGB HDR color, A: luminance for screenspace subsurface scattering
 	float4 CompositedLighting = 0;
 
 #define DO_CULLING_AND_SHADING 1
