Fix for compilation error when using world space normals with material distortion

#ttp 333846 - UE4: RENDERING: Material editor: materials that use world space normals and refraction do not compile (see attached image)

diff --git a/Engine/Shaders/DistortAccumulatePS.usf b/Engine/Shaders/DistortAccumulatePS.usf
--- a/Engine/Shaders/DistortAccumulatePS.usf
+++ b/Engine/Shaders/DistortAccumulatePS.usf
@@ -32,7 +32,7 @@ void Main(
 #if MATERIAL_TANGENTSPACENORMAL
 	half3 ViewNormal = normalize(TransformTangentVectorToView(MaterialParameters, Normal));
 #else
-	half3 ViewNormal = normalize(TransformWorldVectorToLocal(MaterialParameters, Normal));
+	half3 ViewNormal = normalize(TransformWorldVectorToLocal(Normal));
 #endif
 
 	// X = IOR, Y = Refraction Bias
