Fixed Distance Field AO normal calculation reading outside of the valid area for an object - only manifested in artifacts on Nvidia cards for some reason

diff --git a/Engine/Shaders/DistanceFieldSurfaceCacheLighting.usf b/Engine/Shaders/DistanceFieldSurfaceCacheLighting.usf
--- a/Engine/Shaders/DistanceFieldSurfaceCacheLighting.usf
+++ b/Engine/Shaders/DistanceFieldSurfaceCacheLighting.usf
@@ -174,12 +174,16 @@ void ComputeDistanceFieldNormalPS(
 		float3 VolumeShadingPosition = mul(float4(OpaqueWorldPosition, 1), WorldToVolume).xyz;
 		float3 LocalShadingUV = DistanceFieldVolumePositionToUV(VolumeShadingPosition, UVScaleAndVolumeScale.xyz, UVAdd);
 
-		float R = Texture3DSampleLevel(DistanceFieldTexture, DistanceFieldSampler, LocalShadingUV + float3(DistanceFieldAtlasTexelSize.x, 0, 0), 0).x;
-		float L = Texture3DSampleLevel(DistanceFieldTexture, DistanceFieldSampler, LocalShadingUV + float3(-DistanceFieldAtlasTexelSize.x, 0, 0), 0).x;
-		float F = Texture3DSampleLevel(DistanceFieldTexture, DistanceFieldSampler, LocalShadingUV + float3(0, DistanceFieldAtlasTexelSize.y, 0), 0).x;
-		float B = Texture3DSampleLevel(DistanceFieldTexture, DistanceFieldSampler, LocalShadingUV + float3(0, -DistanceFieldAtlasTexelSize.y, 0), 0).x;
-		float U = Texture3DSampleLevel(DistanceFieldTexture, DistanceFieldSampler, LocalShadingUV + float3(0, 0, DistanceFieldAtlasTexelSize.z), 0).x;
-		float D = Texture3DSampleLevel(DistanceFieldTexture, DistanceFieldSampler, LocalShadingUV + float3(0, 0, -DistanceFieldAtlasTexelSize.z), 0).x;
+		// Used to clamp UVs inside valid space of this object's distance field
+		float3 UVMin = DistanceFieldVolumePositionToUV(-LocalPositionExtent, UVScaleAndVolumeScale.xyz, UVAdd);
+		float3 UVMax = DistanceFieldVolumePositionToUV(LocalPositionExtent, UVScaleAndVolumeScale.xyz, UVAdd);
+		
+		float R = Texture3DSampleLevel(DistanceFieldTexture, DistanceFieldSampler, float3(min(LocalShadingUV.x + DistanceFieldAtlasTexelSize.x, UVMax.x), LocalShadingUV.y, LocalShadingUV.z), 0).x;
+		float L = Texture3DSampleLevel(DistanceFieldTexture, DistanceFieldSampler, float3(max(LocalShadingUV.x - DistanceFieldAtlasTexelSize.x, UVMin.x), LocalShadingUV.y, LocalShadingUV.z), 0).x;
+		float F = Texture3DSampleLevel(DistanceFieldTexture, DistanceFieldSampler, float3(LocalShadingUV.x, min(LocalShadingUV.y + DistanceFieldAtlasTexelSize.y, UVMax.y), LocalShadingUV.z), 0).x;
+		float B = Texture3DSampleLevel(DistanceFieldTexture, DistanceFieldSampler, float3(LocalShadingUV.x, max(LocalShadingUV.y - DistanceFieldAtlasTexelSize.y, UVMin.y), LocalShadingUV.z), 0).x;
+		float U = Texture3DSampleLevel(DistanceFieldTexture, DistanceFieldSampler, float3(LocalShadingUV.x, LocalShadingUV.y, min(LocalShadingUV.z + DistanceFieldAtlasTexelSize.z, UVMax.z)), 0).x;
+		float D = Texture3DSampleLevel(DistanceFieldTexture, DistanceFieldSampler, float3(LocalShadingUV.x, LocalShadingUV.y, max(LocalShadingUV.z - DistanceFieldAtlasTexelSize.z, UVMin.z)), 0).x;
 		
 		float3 Gradient = .5f * float3(R - L, F - B, U - D);
 		float3 LocalNormal = normalize(Gradient);
@@ -317,12 +321,16 @@ void ComputeDistanceFieldNormalCS(
 				float3 VolumeShadingPosition = mul(float4(OpaqueWorldPosition, 1), WorldToVolume).xyz;
 				float3 LocalShadingUV = DistanceFieldVolumePositionToUV(VolumeShadingPosition, UVScaleAndVolumeScale.xyz, UVAdd);
 
-				float R = Texture3DSampleLevel(DistanceFieldTexture, DistanceFieldSampler, LocalShadingUV + float3(DistanceFieldAtlasTexelSize.x, 0, 0), 0).x;
-				float L = Texture3DSampleLevel(DistanceFieldTexture, DistanceFieldSampler, LocalShadingUV + float3(-DistanceFieldAtlasTexelSize.x, 0, 0), 0).x;
-				float F = Texture3DSampleLevel(DistanceFieldTexture, DistanceFieldSampler, LocalShadingUV + float3(0, DistanceFieldAtlasTexelSize.y, 0), 0).x;
-				float B = Texture3DSampleLevel(DistanceFieldTexture, DistanceFieldSampler, LocalShadingUV + float3(0, -DistanceFieldAtlasTexelSize.y, 0), 0).x;
-				float U = Texture3DSampleLevel(DistanceFieldTexture, DistanceFieldSampler, LocalShadingUV + float3(0, 0, DistanceFieldAtlasTexelSize.z), 0).x;
-				float D = Texture3DSampleLevel(DistanceFieldTexture, DistanceFieldSampler, LocalShadingUV + float3(0, 0, -DistanceFieldAtlasTexelSize.z), 0).x;
+				// Used to clamp UVs inside valid space of this object's distance field
+				float3 UVMin = DistanceFieldVolumePositionToUV(-LocalPositionExtent, UVScaleAndVolumeScale.xyz, UVAdd);
+				float3 UVMax = DistanceFieldVolumePositionToUV(LocalPositionExtent, UVScaleAndVolumeScale.xyz, UVAdd);
+		
+				float R = Texture3DSampleLevel(DistanceFieldTexture, DistanceFieldSampler, float3(min(LocalShadingUV.x + DistanceFieldAtlasTexelSize.x, UVMax.x), LocalShadingUV.y, LocalShadingUV.z), 0).x;
+				float L = Texture3DSampleLevel(DistanceFieldTexture, DistanceFieldSampler, float3(max(LocalShadingUV.x - DistanceFieldAtlasTexelSize.x, UVMin.x), LocalShadingUV.y, LocalShadingUV.z), 0).x;
+				float F = Texture3DSampleLevel(DistanceFieldTexture, DistanceFieldSampler, float3(LocalShadingUV.x, min(LocalShadingUV.y + DistanceFieldAtlasTexelSize.y, UVMax.y), LocalShadingUV.z), 0).x;
+				float B = Texture3DSampleLevel(DistanceFieldTexture, DistanceFieldSampler, float3(LocalShadingUV.x, max(LocalShadingUV.y - DistanceFieldAtlasTexelSize.y, UVMin.y), LocalShadingUV.z), 0).x;
+				float U = Texture3DSampleLevel(DistanceFieldTexture, DistanceFieldSampler, float3(LocalShadingUV.x, LocalShadingUV.y, min(LocalShadingUV.z + DistanceFieldAtlasTexelSize.z, UVMax.z)), 0).x;
+				float D = Texture3DSampleLevel(DistanceFieldTexture, DistanceFieldSampler, float3(LocalShadingUV.x, LocalShadingUV.y, max(LocalShadingUV.z - DistanceFieldAtlasTexelSize.z, UVMin.z)), 0).x;
 		
 				float3 Gradient = .5f * float3(R - L, F - B, U - D);
 				float3 LocalNormal = normalize(Gradient);
