* fixed Skylight/indirect lighting for ScreenSpaceSubsurfaceScattering
* fixed negative colors from SkyLight computations in base pass

diff --git a/Engine/Shaders/BasePassPixelShader.usf b/Engine/Shaders/BasePassPixelShader.usf
--- a/Engine/Shaders/BasePassPixelShader.usf
+++ b/Engine/Shaders/BasePassPixelShader.usf
@@ -13,6 +13,7 @@
 #include "ReflectionEnvironmentShared.usf"
 #include "BRDF.usf"
 #include "Random.usf"
+#include "LightAccumulator.usf"
 
 #include "DeferredShadingCommon.usf"
 
@@ -685,8 +686,13 @@ void Main(
 		half3 FoggedColor = lerp(float3(1, 1, 1), Color, VertexFog.aaa * VertexFog.aaa);
 		OutColor = half4(FoggedColor, Opacity);
 	#else
-		// Scene color alpha is used to specify how much specular is in the color, for efficiency we approximate with 0
-		OutColor = RETURN_COLOR(float4(Color, 0)); 
+		// Scene color alpha is used for ScreenSpaceSubsurfaceScattering (if that is not needed it can be disabled with SUBSURFACE_CHANNEL_MODE)
+		{
+			FLightAccumulator LightAccumulator = (FLightAccumulator)0;
+
+			LightAccumulator_Add(LightAccumulator, Color, 0, 1.0f);
+			OutColor = RETURN_COLOR(LightAccumulator_GetResult(LightAccumulator));
+		}
 	#endif
 
 	float4 PrecomputedShadowFactors = GetPrecomputedShadowMasks(Interpolants);
diff --git a/Engine/Shaders/LightAccumulator.usf b/Engine/Shaders/LightAccumulator.usf
--- a/Engine/Shaders/LightAccumulator.usf
+++ b/Engine/Shaders/LightAccumulator.usf
@@ -83,7 +83,7 @@ float4 LightAccumulator_GetResult(FLightAccumulator In)
 		}
 		else if (SUBSURFACE_CHANNEL_MODE == 2)
 		{
-			// RGB accumulated RGB HDR color, A: specular luminance for screenspace subsurface scattering
+			// RGB accumulated RGB HDR color, A: view independent (diffuse) luminance for screenspace subsurface scattering
 			// 3 add,  1 mul, 2 mad, can be optimized to use 2 less temporary during accumulation and remove the 3 add
 			Ret.a = Luminance(In.ViewIndependentLightingSum);
 			// todo, need second MRT for SUBSURFACE_CHANNEL_MODE==2
diff --git a/Engine/Shaders/ReflectionEnvironmentShared.usf b/Engine/Shaders/ReflectionEnvironmentShared.usf
--- a/Engine/Shaders/ReflectionEnvironmentShared.usf
+++ b/Engine/Shaders/ReflectionEnvironmentShared.usf
@@ -73,5 +73,6 @@ float3 GetSkySHDiffuse(float3 Normal)
 	float vC = NormalVector.x * NormalVector.x - NormalVector.y * NormalVector.y;
 	Intermediate2 = View.SkyIrradianceEnvironmentMap[6].xyz * vC;
 
-	return Intermediate0 + Intermediate1 + Intermediate2;
+	// max to not get negative colors, maybe we can optimize to use saturate() and scale later or clamp this way: saturate(A/100)*100
+	return max(0, Intermediate0 + Intermediate1 + Intermediate2);
 }
