* fixed ScreenSpaceSubsurfaceScattering for dynamic Skylight

diff --git a/Engine/Shaders/SkyLighting.usf b/Engine/Shaders/SkyLighting.usf
--- a/Engine/Shaders/SkyLighting.usf
+++ b/Engine/Shaders/SkyLighting.usf
@@ -85,5 +85,11 @@ void SkyLightDiffusePS(
 		//Lighting = Texture2DSampleLevel(BentNormalAOTexture, BentNormalAOSampler, UV * BentNormalUVScale, 0).xyz;
 	}
 
-	OutColor = float4(Lighting, 0);
+	// Scene color alpha is used for ScreenSpaceSubsurfaceScattering (if that is not needed it can be disabled with SUBSURFACE_CHANNEL_MODE)
+	{
+		FLightAccumulator LightAccumulator = (FLightAccumulator)0;
+
+		LightAccumulator_Add(LightAccumulator, Lighting, 0, 1.0f);
+		OutColor = LightAccumulator_GetResult(LightAccumulator);
+	}
 }
\ No newline at end of file
