Remove scaling when calculating tangents for instanced mesh particles

#ttp 335591 - RENDERING: PARTICLE: Mesh emitters which utilize a material with a normal map are distorted

diff --git a/Engine/Shaders/MeshParticleVertexFactory.usf b/Engine/Shaders/MeshParticleVertexFactory.usf
--- a/Engine/Shaders/MeshParticleVertexFactory.usf
+++ b/Engine/Shaders/MeshParticleVertexFactory.usf
@@ -372,8 +372,9 @@ void CalcTangentToWorld(FVertexFactoryInput Input, FVertexFactoryIntermediates I
 #endif
 	float3x3 TangentToWorld = mul(TangentToLocal, transpose(Transform));
 
-	TangentToWorld0 = TangentToWorld[0];
-	TangentToWorld2 = float4(TangentToWorld[2], TangentBias(Input.TangentZ.w) * Primitive.LocalToWorldDeterminantSign);
+	// Normalize to remove scaling. Incorrect but faster than computing the inverse-transpose.
+	TangentToWorld0 = normalize(TangentToWorld[0]);
+	TangentToWorld2 = normalize(float4(TangentToWorld[2], TangentBias(Input.TangentZ.w) * Primitive.LocalToWorldDeterminantSign));
 }
 
 FVertexFactoryInterpolantsVSToPS VertexFactoryGetInterpolantsVSToPS(FVertexFactoryInput Input, FVertexFactoryIntermediates Intermediates, FMaterialVertexParameters VertexParameters)
