Remove LocalToWorld scaling when calculating tangents for skinned meshes

#ttp 338139 - LIVE: RENDER: Tessellated Skeletal Meshes canot be scaled without spiking the Tessellation outward

diff --git a/Engine/Shaders/GpuSkinVertexFactory.usf b/Engine/Shaders/GpuSkinVertexFactory.usf
--- a/Engine/Shaders/GpuSkinVertexFactory.usf
+++ b/Engine/Shaders/GpuSkinVertexFactory.usf
@@ -540,7 +540,15 @@ float4 VertexFactoryGetRasterizedWorldPosition(FVertexFactoryInput Input, FVerte
 
 void CalcTangentToWorld(FVertexFactoryInput Input, FVertexFactoryIntermediates Intermediates, out float3 TangentToWorld0, out float4 TangentToWorld2)
 {
-	float3x3 TangentToWorld = mul(Intermediates.TangentToLocal, GetLocalToWorld3x3()); 
+	float3x3 LocalToWorld = GetLocalToWorld3x3();
+
+	// Remove scaling.
+	half3 InvScale = Primitive.InvNonUniformScale.xyz;
+	LocalToWorld[0] *= InvScale.x;
+	LocalToWorld[1] *= InvScale.y;
+	LocalToWorld[2] *= InvScale.z;
+
+	float3x3 TangentToWorld = mul(Intermediates.TangentToLocal, LocalToWorld);
 
 	TangentToWorld0 = TangentToWorld[0];
 	TangentToWorld2 = float4(TangentToWorld[2], TangentBias(Input.TangentZ.w) * Primitive.LocalToWorldDeterminantSign);
