UE4 - Initial check-in for instanced static meshes (emulated) on ES2. Clamps at 64 instances.

diff --git a/Engine/Shaders/LocalVertexFactory.usf b/Engine/Shaders/LocalVertexFactory.usf
--- a/Engine/Shaders/LocalVertexFactory.usf
+++ b/Engine/Shaders/LocalVertexFactory.usf
@@ -1,4 +1,4 @@
-// Copyright 1998-2014 Epic Games, Inc. All Rights Reserved.
+// Copyright 1998-2014 Epic Games, Inc. All Rights Reserved...
 
 /*=============================================================================
 	LocalVertexFactory.usf: Local vertex factory shader code.
@@ -13,7 +13,19 @@ float4 ShadowMapCoordinateScaleBias;
 
 #if USE_INSTANCING
 float4 InstancingFadeOutParams;
-#endif
+
+#ifndef USE_INSTANCING_EMULATED
+#define USE_INSTANCING_EMULATED 0
+#endif	// USE_INSTANCING_EMULATED
+
+	#if USE_INSTANCING_EMULATED
+// Required for CPU emulation of Instancing
+float4 CPUInstanceShadowMapBias;
+float3x4 CPUInstanceTransform;
+float3x4 CPUInstanceInverseTransform;
+	#endif	// USE_INSTANCING_EMULATED
+
+#endif	// USE_INSTANCING
 
 #if USE_SPLINEDEFORM
 float3 SplineStartPos;
@@ -37,7 +49,7 @@ float SplineMeshScaleZ;
 float3 SplineMeshDir;
 float3 SplineMeshX;
 float3 SplineMeshY;
-#endif
+#endif	// USE_SPLINEDEFORM
 
 struct FVertexFactoryInput
 {
@@ -52,7 +64,7 @@ struct FVertexFactoryInput
 	// TangentZ.w contains sign of tangent basis determinant
 	half4	TangentZ	: ATTRIBUTE2;
 	half4	Color		: ATTRIBUTE3;
-#endif
+#endif	// METAL_PROFILE
 
 #if NUM_MATERIAL_TEXCOORDS_VERTEX
 
@@ -73,7 +85,7 @@ struct FVertexFactoryInput
 	float2	TexCoords[1] : ATTRIBUTE4;
 #endif
 
-#if USE_INSTANCING
+#if USE_INSTANCING && !USE_INSTANCING_EMULATED
 	float4 InstanceShadowMapBias : ATTRIBUTE8;
 	float4 InstanceTransform1 : ATTRIBUTE9;
 	float4 InstanceTransform2 : ATTRIBUTE10;
@@ -84,7 +96,7 @@ struct FVertexFactoryInput
 	// note, the per-instance lightmap bias is stored in the W of InverseTransform1 and 2 and per-instance random in 3.
 #else
 
-#endif //!USE_INSTANCING
+#endif //USE_INSTANCING && !USE_INSTANCING_EMULATED
 
 #if NEEDS_LIGHTMAP_COORDINATE
 	float2	LightMapCoordinate : ATTRIBUTE15;
@@ -95,11 +107,11 @@ struct FPositionOnlyVertexFactoryInput
 {
 	float4	Position	: ATTRIBUTE0;
 	
-#if USE_INSTANCING
+#if USE_INSTANCING && !USE_INSTANCING_EMULATED
 	float4 InstanceTransform1 : ATTRIBUTE9;
 	float4 InstanceTransform2 : ATTRIBUTE10;
 	float4 InstanceTransform3 : ATTRIBUTE11;
-#endif
+#endif	// USE_INSTANCING && !USE_INSTANCING_EMULATED
 };
 
 struct FVertexFactoryIntermediates
@@ -111,9 +123,47 @@ struct FVertexFactoryIntermediates
 	half4 Color;
 #if USE_INSTANCING
 	float2 PerInstanceParams;
-#endif
+#endif	// USE_INSTANCING
 };
 
+#if USE_INSTANCING
+float4x4 GetInstanceTransform(FVertexFactoryInput Input)
+{
+	#if !USE_INSTANCING_EMULATED
+	return float4x4(Input.InstanceTransform1, Input.InstanceTransform2, Input.InstanceTransform3, float4(0, 0, 0, 1));
+	#else
+	return float4x4(CPUInstanceTransform[0], CPUInstanceTransform[1], CPUInstanceTransform[2], float4(0,0,0,1));
+	#endif	// !USE_INSTANCING_EMULATED
+}
+
+float4x4 GetInstanceTransform(FPositionOnlyVertexFactoryInput Input)
+{
+	#if !USE_INSTANCING_EMULATED
+	return float4x4(Input.InstanceTransform1, Input.InstanceTransform2, Input.InstanceTransform3, float4(0, 0, 0, 1));
+	#else
+	return float4x4(CPUInstanceTransform[0], CPUInstanceTransform[1], CPUInstanceTransform[2], float4(0,0,0,1));
+#endif	// !USE_INSTANCING_EMULATED
+}
+
+float3x3 GetInstanceInverseTransform(FVertexFactoryInput Input)
+{
+	#if !USE_INSTANCING_EMULATED
+	return float3x3(Input.InstanceInverseTransform1.xyz, Input.InstanceInverseTransform2.xyz, Input.InstanceInverseTransform3.xyz);
+	#else
+	return float3x3(CPUInstanceInverseTransform[0].xyz, CPUInstanceInverseTransform[1].xyz, CPUInstanceInverseTransform[2].xyz);
+#endif	// !USE_INSTANCING_EMULATED
+}
+
+float4 GetInstanceShadowMapBias(FVertexFactoryInput Input)
+{
+	#if !USE_INSTANCING_EMULATED
+	return Input.InstanceShadowMapBias;
+	#else
+	return CPUInstanceShadowMapBias;
+	#endif	// !USE_INSTANCING_EMULATED
+}
+#endif	// USE_INSTANCING
+
 /** Converts from vertex factory specific interpolants to a FMaterialPixelParameters, which is used by material inputs. */
 FMaterialPixelParameters GetMaterialPixelParameters(FVertexFactoryInterpolantsVSToPS Interpolants, float4 PixelPosition)
 {
@@ -126,13 +176,13 @@ FMaterialPixelParameters GetMaterialPixelParameters(FVertexFactoryInterpolantsVS
 	{
 		Result.TexCoords[CoordinateIndex] = GetUV(Interpolants, CoordinateIndex);
 	}
-#endif
+#endif	// NUM_MATERIAL_TEXCOORDS
 
 #if USE_PARTICLE_SUBUVS
 	// Output TexCoord0 for when previewing materials that use ParticleSubUV.
 	Result.Particle.SubUVCoords[0] = GetUV(Interpolants, 0);
 	Result.Particle.SubUVCoords[1] = GetUV(Interpolants, 0);
-#endif
+#endif	// USE_PARTICLE_SUBUVS
 
 	half3 TangentToWorld0 = GetTangentToWorld0(Interpolants).xyz;
 	half4 TangentToWorld2 = GetTangentToWorld2(Interpolants);
@@ -155,9 +205,9 @@ FMaterialPixelParameters GetMaterialPixelParameters(FVertexFactoryInterpolantsVS
 		Result.LightmapUVs = float2(0, 0);
 	#else
 		Result.LightmapUVs = Interpolants.LightMapCoordinate.xy;
-	#endif
-#endif
-#endif
+	#endif	// ES2_PROFILE
+#endif	// NEEDS_LIGHTMAP_COORDINATE
+#endif	// LIGHTMAP_UV_ACCESS
 
 	Result.TwoSidedSign = 1;
 	return Result;
@@ -184,11 +234,11 @@ FMaterialVertexParameters GetMaterialVertexParameters( FVertexFactoryInput Input
 	Result.TangentToWorld = Intermediates.TangentToWorld;
 
 #if USE_INSTANCING
-	Result.InstanceLocalToWorld = mul(transpose(float4x4(Input.InstanceTransform1, Input.InstanceTransform2, Input.InstanceTransform3, float4(0, 0, 0, 1))), Primitive.LocalToWorld);
-	Result.InstanceWorldToLocal = mul(float3x3(Primitive.WorldToLocal[0].xyz,Primitive.WorldToLocal[1].xyz,Primitive.WorldToLocal[2].xyz), transpose(float3x3(Input.InstanceInverseTransform1.xyz, Input.InstanceInverseTransform2.xyz, Input.InstanceInverseTransform3.xyz)));
+	Result.InstanceLocalToWorld = mul(transpose(GetInstanceTransform(Input)), Primitive.LocalToWorld);
+	Result.InstanceWorldToLocal = mul(float3x3(Primitive.WorldToLocal[0].xyz,Primitive.WorldToLocal[1].xyz,Primitive.WorldToLocal[2].xyz), transpose(GetInstanceInverseTransform(Input)));
 	Result.InstanceLocalPosition = Input.Position.xyz;
 	Result.PerInstanceParams = Intermediates.PerInstanceParams;
-#endif
+#endif	// USE_INSTANCING
 
 #if NUM_MATERIAL_TEXCOORDS_VERTEX
 	#if NUM_MATERIAL_TEXCOORDS_VERTEX > 1
@@ -201,11 +251,11 @@ FMaterialVertexParameters GetMaterialVertexParameters( FVertexFactoryInput Input
 				Result.TexCoords[CoordinateIndex+1] = Input.PackedTexCoords4[CoordinateIndex/2].zw;
 			}
 		}
-	#endif
+	#endif	// NUM_MATERIAL_TEXCOORDS_VERTEX > 1
 	#if NUM_MATERIAL_TEXCOORDS_VERTEX % 2 == 1
 		Result.TexCoords[NUM_MATERIAL_TEXCOORDS_VERTEX-1] = Input.PackedTexCoords2;
-	#endif
-#endif
+	#endif	// NUM_MATERIAL_TEXCOORDS_VERTEX % 2 == 1
+#endif	// NUM_MATERIAL_TEXCOORDS_VERTEX
 	return Result;
 }
 
@@ -305,17 +355,16 @@ FMaterialVertexParameters GetMaterialVertexParameters( FVertexFactoryInput Input
 
 		return SliceTransform;
 	}
-#endif
+#endif	// USE_SPLINEDEFORM
 
 #if USE_INSTANCING
-float4 CalcWorldPosition(float4 Position, float4 InstanceTransform1, float4 InstanceTransform2, float4 InstanceTransform3)
+float4 CalcWorldPosition(float4 Position, float4x4 InstanceTransform)
 #else
 float4 CalcWorldPosition(float4 Position)
-#endif
+#endif	// USE_INSTANCING
 {
 #if USE_INSTANCING
-	float4x4 InstanceTransform = transpose(float4x4(InstanceTransform1, InstanceTransform2, InstanceTransform3, float4(0, 0, 0, 1)));
-	return TransformLocalToTranslatedWorld(mul(Position, InstanceTransform).xyz);
+	return TransformLocalToTranslatedWorld(mul(Position, transpose(InstanceTransform)).xyz);
 #elif USE_SPLINEDEFORM
 	// Make transform for this point along spline
 	float4x3 SliceTransform = CalcSliceTransform(dot(Position.xyz, SplineMeshDir));
@@ -345,7 +394,7 @@ half3x3 CalcTangentToLocal(FVertexFactoryInput Input)
 	// pass-thru the tangent
 	half3 TangentX = TangentBias(Input.TangentX);
 	// pass-thru the normal
-#endif
+#endif	// USE_SPLINEDEFORM
 
 	// derive the binormal by getting the cross product of the normal and tangent
 	half3 TangentY = cross(TangentZ.xyz, TangentX) * TangentZ.w;
@@ -362,12 +411,12 @@ half3x3 CalcTangentToLocal(FVertexFactoryInput Input)
 half3x3 CalcTangentToWorld(FVertexFactoryInput Input, half3x3 TangentToLocal)
 {
 #if USE_INSTANCING
-	half3x3 TangentToPrimitive = mul(TangentToLocal, transpose(float3x3(Input.InstanceTransform1.xyz, Input.InstanceTransform2.xyz, Input.InstanceTransform3.xyz)));
-	half3x3 TangentToWorldWithScaling = mul(TangentToPrimitive, (float3x3)Primitive.LocalToWorld); 
+	half3x3 TangentToPrimitive = mul(TangentToLocal, transpose((float3x3)GetInstanceTransform(Input)));
+	half3x3 TangentToWorldWithScaling = mul(TangentToPrimitive, (float3x3)Primitive.LocalToWorld);
 	half3x3 TangentToWorld = float3x3(normalize(TangentToWorldWithScaling[0]), normalize(TangentToWorldWithScaling[1]), normalize(TangentToWorldWithScaling[2]));
 #else
 	half3x3 TangentToWorld = CalcTangentToWorldNoScale(TangentToLocal);
-#endif
+#endif	// USE_INSTANCING
 	return TangentToWorld;
 }
 
@@ -382,11 +431,17 @@ FVertexFactoryIntermediates GetVertexFactoryIntermediates(FVertexFactoryInput In
 	// Swizzle vertex color.
 	Intermediates.Color = Input.Color FCOLOR_COMPONENT_SWIZZLE;
 #if USE_INSTANCING
+	#if !USE_INSTANCING_EMULATED
 	float3 InstanceLocation = TransformLocalToTranslatedWorld(float3(Input.InstanceTransform1.w, Input.InstanceTransform2.w, Input.InstanceTransform3.w)).xyz;
 	// x = per-instance random, y = per-instance fade out factor
 	Intermediates.PerInstanceParams.x = Input.InstanceInverseTransform3.w;
+	#else
+	float3 InstanceLocation = TransformLocalToTranslatedWorld(float3(CPUInstanceTransform[0].w, CPUInstanceTransform[1].w, CPUInstanceTransform[2].w)).xyz;
+	// x = per-instance random, y = per-instance fade out factor
+	Intermediates.PerInstanceParams.x = CPUInstanceInverseTransform[2].w;
+	#endif // !USE_INSTANCING_EMULATED
 	Intermediates.PerInstanceParams.y = 1.0 - saturate((length(InstanceLocation) - InstancingFadeOutParams.x) * InstancingFadeOutParams.y);
-#endif
+#endif	// USE_INSTANCING
 	return Intermediates;
 }
 
@@ -406,13 +461,13 @@ half3x3 VertexFactoryGetTangentToLocal( FVertexFactoryInput Input, FVertexFactor
 float4 VertexFactoryGetWorldPosition(FVertexFactoryInput Input, FVertexFactoryIntermediates Intermediates)
 {
 #if USE_INSTANCING
-	float SelectedValue = trunc(Input.InstanceShadowMapBias.z * (1.0/256.0));
+	float SelectedValue = trunc(GetInstanceShadowMapBias(Input).z * (1.0 / 256.0));
 	// InstancingFadeOutParams.z,w are RenderSelected and RenderDeselected respectively.
 	float SelectionMask = InstancingFadeOutParams.z * SelectedValue + InstancingFadeOutParams.w * (1-SelectedValue);
-	return CalcWorldPosition(Input.Position, Input.InstanceTransform1, Input.InstanceTransform2, Input.InstanceTransform3) * SelectionMask;
+	return CalcWorldPosition(Input.Position, GetInstanceTransform(Input)) * SelectionMask;
 #else
 	return CalcWorldPosition(Input.Position);
-#endif
+#endif	// USE_INSTANCING
 }
 
 float4 VertexFactoryGetRasterizedWorldPosition(FVertexFactoryInput Input, FVertexFactoryIntermediates Intermediates, float4 InWorldPosition)
@@ -441,41 +496,40 @@ FVertexFactoryInterpolantsVSToPS VertexFactoryGetInterpolantsVSToPS(FVertexFacto
 
 #elif USE_PARTICLE_SUBUVS
 	SetUV(Interpolants, 0, Input.TexCoords[0]);
-#endif
+#endif	// NUM_MATERIAL_TEXCOORDS
 
 #if NEEDS_LIGHTMAP_COORDINATE
-#if USE_INSTANCING
+	#if USE_INSTANCING
 	//float2 InstanceLightMapBias = float2(Input.InstanceInverseTransform1.w, Input.InstanceInverseTransform2.w);
 	//Interpolants.LightMapCoordinate.xy = Input.LightMapCoordinate * LightMapCoordinateScaleBias.xy + InstanceLightMapBias;
-	#if ES2_PROFILE
+		#if ES2_PROFILE
 	SetLightMapCoordinate(Interpolants, (float4)0, (float2)0);
-	#else
+		#else
 	Interpolants.LightMapCoordinate = 0;
-	#endif
-#else
+		#endif	// ES2_PROFILE
+	#else
 	
 	float4 LightMapCoordinate = 0;
 	float2 ShadowMapCoordinate = 0;
 
-	#if ES2_PROFILE
+		#if ES2_PROFILE
 		LightMapCoordinate.xy = Input.LightMapCoordinate * LightMapCoordinateScaleBias.xy + LightMapCoordinateScaleBias.zw;
 		LightMapCoordinate.y *= 0.5;
 		LightMapCoordinate.zw = LightMapCoordinate.xy;
 		LightMapCoordinate.w += 0.5;
-		#if STATICLIGHTING_TEXTUREMASK
+			#if STATICLIGHTING_TEXTUREMASK
 			ShadowMapCoordinate = Input.LightMapCoordinate * ShadowMapCoordinateScaleBias.xy + ShadowMapCoordinateScaleBias.zw;
-		#endif
-	#else
+			#endif	// STATICLIGHTING_TEXTUREMASK
+		#else
 		LightMapCoordinate.xy = Input.LightMapCoordinate * LightMapCoordinateScaleBias.xy + LightMapCoordinateScaleBias.zw;
-		#if STATICLIGHTING_TEXTUREMASK
+			#if STATICLIGHTING_TEXTUREMASK
 			ShadowMapCoordinate = Input.LightMapCoordinate * ShadowMapCoordinateScaleBias.xy + ShadowMapCoordinateScaleBias.zw;
-		#endif
-	#endif
+			#endif	// STATICLIGHTING_TEXTUREMASK
+		#endif	// ES2_PROFILE
 
 	SetLightMapCoordinate(Interpolants, LightMapCoordinate, ShadowMapCoordinate);
-
-#endif
-#endif
+	#endif	// USE_INSTANCING
+#endif	// NEEDS_LIGHTMAP_COORDINATE
 
 	SetTangents(Interpolants, Intermediates.TangentToWorld[0], Intermediates.TangentToWorld[2], Intermediates.TangentToWorldSign);
 	SetColor(Interpolants, Intermediates.Color);
@@ -489,17 +543,17 @@ FVertexFactoryInterpolantsVSToPS VertexFactoryGetInterpolantsVSToPS(FVertexFacto
 float4 VertexFactoryGetWorldPosition(FPositionOnlyVertexFactoryInput Input)
 {
 #if USE_INSTANCING
-	return CalcWorldPosition(Input.Position, Input.InstanceTransform1, Input.InstanceTransform2, Input.InstanceTransform3);
+	return CalcWorldPosition(Input.Position, GetInstanceTransform(Input));
 #else
 	return CalcWorldPosition(Input.Position);
-#endif
+#endif	// USE_INSTANCING
 }
 
 // @return previous translated world position
 float4 VertexFactoryGetPreviousWorldPosition(FVertexFactoryInput Input, FVertexFactoryIntermediates Intermediates)
 {
 #if USE_INSTANCING
-	float4x4 InstanceTransform = transpose(float4x4(Input.InstanceTransform1, Input.InstanceTransform2, Input.InstanceTransform3, float4(0, 0, 0, 1)));
+	float4x4 InstanceTransform = transpose(GetInstanceTransform(Input));
 	return mul(mul(Input.Position, InstanceTransform), PreviousLocalToWorld);
 #elif USE_SPLINEDEFORM
 	// Just like CalcWorldPosition...
@@ -511,7 +565,7 @@ float4 VertexFactoryGetPreviousWorldPosition(FVertexFactoryInput Input, FVertexF
 	return mul(LocalPos, PreviousLocalToWorld);
 #else
 	return mul(Input.Position, PreviousLocalToWorld);
-#endif
+#endif	// USE_INSTANCING
 }
 
 #if USING_TESSELLATION
@@ -548,14 +602,14 @@ float4 VertexFactoryGetPreviousWorldPosition(FVertexFactoryInput Input, FVertexF
 				Result.TexCoords[CoordinateIndex + 1] = Interpolants.InterpolantsVSToPS.TexCoords[CoordinateIndex/2].wz;
 			}
 		}
-	#endif
+	#endif	// NUM_MATERIAL_TEXCOORDS
 
 		half3 TangentToWorld0 = Interpolants.InterpolantsVSToPS.TangentToWorld0.xyz;
 		half4 TangentToWorld2 = Interpolants.InterpolantsVSToPS.TangentToWorld2;
 
-#if INTERPOLATE_VERTEX_COLOR
+	#if INTERPOLATE_VERTEX_COLOR
 		Result.VertexColor = Interpolants.InterpolantsVSToPS.Color;
-#endif
+	#endif	// INTERPOLATE_VERTEX_COLOR
 
 		Result.TangentToWorld = AssembleTangentToWorld( TangentToWorld0, TangentToWorld2 );
 
@@ -620,10 +674,11 @@ float4 VertexFactoryGetPreviousWorldPosition(FVertexFactoryInput Input, FVertexF
 float4 VertexFactoryGetInstanceHitProxyId(FVertexFactoryInput Input, FVertexFactoryIntermediates Intermediates)
 {
 	// Unpack HitProxy color from InstanceShadowMapBias z and w
-	float SelectedMask = trunc(Input.InstanceShadowMapBias.z * (1.0/256.0));
-	float R = Input.InstanceShadowMapBias.z - 256.0 * SelectedMask;
-	float G = trunc(Input.InstanceShadowMapBias.w * (1.0/256.0));
-	float B = Input.InstanceShadowMapBias.w - 256.0 * G;
+	float4 InstanceShadowMapBias = GetInstanceShadowMapBias(Input);
+	float SelectedMask = trunc(InstanceShadowMapBias.z * (1.0 / 256.0));
+	float R = InstanceShadowMapBias.z - 256.0 * SelectedMask;
+	float G = trunc(InstanceShadowMapBias.w * (1.0/256.0));
+	float B = InstanceShadowMapBias.w - 256.0 * G;
 	return float4(R/255.0, G/255.0, B/255.0, 0);
 }
-#endif
\ No newline at end of file
+#endif	// USE_INSTANCING
diff --git a/Engine/Source/Runtime/Engine/Private/InstancedStaticMesh.cpp b/Engine/Source/Runtime/Engine/Private/InstancedStaticMesh.cpp
--- a/Engine/Source/Runtime/Engine/Private/InstancedStaticMesh.cpp
+++ b/Engine/Source/Runtime/Engine/Private/InstancedStaticMesh.cpp
@@ -145,6 +145,11 @@ public:
 		return NumInstances;
 	}
 
+	const void* GetRawData() const
+	{
+		return InstanceData->GetDataPointer();
+	}
+
 	// FRenderResource interface.
 	virtual void InitRHI() override;
 	virtual FString GetFriendlyName() const { return TEXT("Static-mesh instances"); }
@@ -200,7 +205,7 @@ void FStaticMeshInstanceBuffer::Init(UInstancedStaticMeshComponent* InComponent,
 	// so we make a TArray of the right type, then assign it
 	TArray<FVector4> RawData;
 	check( GetStride() % sizeof(FVector4) == 0 );
-	RawData.Empty(NumInstances * GetStride() / sizeof(FVector));
+	RawData.Empty(NumInstances * GetStride() / sizeof(FVector4));
 
 	// @todo: Make LD-customizable per component?
 	const float RandomInstanceIDBase = 0.0f;
@@ -303,7 +308,8 @@ void FStaticMeshInstanceBuffer::AllocateData()
 	// Clear any old VertexData before allocating.
 	CleanUp();
 
-	const bool bNeedsCPUAccess=false;
+	const bool bInstanced = RHISupportsInstancing(GRHIShaderPlatform);
+	const bool bNeedsCPUAccess = !bInstanced;
 	InstanceData = new FStaticMeshInstanceData(bNeedsCPUAccess);
 	// Calculate the vertex stride.
 	Stride = InstanceData->GetStride();
@@ -315,8 +321,16 @@ void FStaticMeshInstanceBuffer::AllocateData()
 	FInstancedStaticMeshVertexFactory
 -----------------------------------------------------------------------------*/
 
-struct FInstancingUserData 
+struct FInstancingUserData
 {
+	struct FInstanceStream
+	{
+		FVector4 InstanceShadowmapUVBias;
+		FVector4 InstanceTransform[3];
+		FVector4 InstanceInverseTransform[3];
+	};
+	class FInstancedStaticMeshRenderData* RenderData;
+
 	int32 StartCullDistance;
 	int32 EndCullDistance;
 
@@ -354,11 +368,9 @@ public:
 	 */
 	static void ModifyCompilationEnvironment(EShaderPlatform Platform, const FMaterial* Material, FShaderCompilerEnvironment& OutEnvironment)
 	{
+		OutEnvironment.SetDefine(TEXT("USE_INSTANCING"),TEXT("1"));
 		const bool bInstanced = RHISupportsInstancing(Platform);
-		if (bInstanced)
-		{
-			OutEnvironment.SetDefine(TEXT("USE_INSTANCING"),TEXT("1"));
-		}
+		OutEnvironment.SetDefine(TEXT("USE_INSTANCING_EMULATED"), bInstanced ? TEXT("0") : TEXT("1"));
 	}
 
 	/**
@@ -381,6 +393,16 @@ public:
 
 	static FVertexFactoryShaderParameters* ConstructShaderParameters(EShaderFrequency ShaderFrequency);
 
+	/**
+	* Get a bitmask representing the visibility of each FMeshBatch element.
+	*/
+	virtual uint64 GetStaticBatchElementVisibility(const class FSceneView& View, const struct FMeshBatch* Batch) const override
+	{
+		//@todo-rco: Only draw the first 64 instances
+		int32 NumElements = FMath::Min(Batch->Elements.Num(), 64);
+		return (1ULL << (uint64)NumElements) - 1ULL;
+	}
+
 private:
 	DataType Data;
 };
@@ -512,60 +534,29 @@ class FInstancedStaticMeshVertexFactoryShaderParameters : public FVertexFactoryS
 	virtual void Bind(const FShaderParameterMap& ParameterMap) override
 	{
 		InstancingFadeOutParamsParameter.Bind(ParameterMap,TEXT("InstancingFadeOutParams"));
+		CPUInstanceShadowMapBias.Bind(ParameterMap, TEXT("CPUInstanceShadowMapBias"));
+		CPUInstanceTransform.Bind(ParameterMap, TEXT("CPUInstanceTransform"));
+		CPUInstanceInverseTransform.Bind(ParameterMap, TEXT("CPUInstanceInverseTransform"));
 	}
 
-	virtual void SetMesh(FRHICommandList& RHICmdList, FShader* VertexShader,const class FVertexFactory* VertexFactory,const class FSceneView& View,const struct FMeshBatchElement& BatchElement,uint32 DataFlags) const override
-	{
-		if( InstancingFadeOutParamsParameter.IsBound() )
-		{
-			FVector4 InstancingFadeOutParams(0.f,0.f,1.f,1.f);
-
-			const FInstancingUserData* InstancingUserData = (FInstancingUserData*)BatchElement.UserData;
-			if (InstancingUserData)
-			{
-				InstancingFadeOutParams.X = InstancingUserData->StartCullDistance;
-				if( InstancingUserData->EndCullDistance > 0 )
-				{
-					if( InstancingUserData->EndCullDistance > InstancingUserData->StartCullDistance )
-					{
-						InstancingFadeOutParams.Y = 1.f / (float)(InstancingUserData->EndCullDistance - InstancingUserData->StartCullDistance);
-					}
-					else
-					{
-						InstancingFadeOutParams.Y = 1.f;
-					}
-				}
-				else
-				{
-					InstancingFadeOutParams.Y = 0.f;
-				}
-
-				InstancingFadeOutParams.Z = InstancingUserData->bRenderSelected ? 1.f : 0.f;
-				InstancingFadeOutParams.W = InstancingUserData->bRenderUnselected ? 1.f : 0.f;
-			}
-			SetShaderValue(RHICmdList, VertexShader->GetVertexShader(), InstancingFadeOutParamsParameter, InstancingFadeOutParams );
-		}
-
-		// upload SpeedTree buffer, if available
-		if (View.Family != NULL && View.Family->Scene != NULL)
-		{
-			FUniformBufferRHIParamRef SpeedTreeUniformBuffer = View.Family->Scene->GetSpeedTreeUniformBuffer(VertexFactory);
-			if (SpeedTreeUniformBuffer != NULL)
-			{
-				SetUniformBufferParameter(RHICmdList, VertexShader->GetVertexShader(), VertexShader->GetUniformBufferParameter<FSpeedTreeUniformParameters>(), SpeedTreeUniformBuffer);
-			}
-		}
-	}
+	virtual void SetMesh(FRHICommandList& RHICmdList, FShader* VertexShader,const class FVertexFactory* VertexFactory,const class FSceneView& View,const struct FMeshBatchElement& BatchElement,uint32 DataFlags) const override;
 
 	void Serialize(FArchive& Ar)
 	{
 		Ar << InstancingFadeOutParamsParameter;
+		Ar << CPUInstanceShadowMapBias;
+		Ar << CPUInstanceTransform;
+		Ar << CPUInstanceInverseTransform;
 	}
 
 	virtual uint32 GetSize() const { return sizeof(*this); }
 
 private:
 	FShaderParameter InstancingFadeOutParamsParameter;
+
+	FShaderParameter CPUInstanceShadowMapBias;
+	FShaderParameter CPUInstanceTransform;
+	FShaderParameter CPUInstanceInverseTransform;
 };
 
 FVertexFactoryShaderParameters* FInstancedStaticMeshVertexFactory::ConstructShaderParameters(EShaderFrequency ShaderFrequency)
@@ -893,11 +884,16 @@ public:
 			}
 		}
 
+		check(InstancedRenderData.InstanceBuffer.GetStride() == sizeof(FInstancingUserData::FInstanceStream));
+
+		const bool bInstanced = RHISupportsInstancing(GRHIShaderPlatform);
+
 		// Copy the parameters for LOD - all instances
 		UserData_AllInstances.StartCullDistance = InComponent->InstanceStartCullDistance;
 		UserData_AllInstances.EndCullDistance = InComponent->InstanceEndCullDistance;
 		UserData_AllInstances.bRenderSelected = true;
 		UserData_AllInstances.bRenderUnselected = true;
+		UserData_AllInstances.RenderData = bInstanced ? nullptr : &InstancedRenderData;
 
 		// selected only
 		UserData_SelectedInstances = UserData_AllInstances;
@@ -927,6 +923,9 @@ public:
 		return Result;
 	}
 
+	/** Common path for the Get*MeshElement functions */
+	inline void SetupInstancedMeshBatch(int32 LODIndex, FMeshBatch& OutMeshBatch) const;
+
 	/** Draw the scene proxy as a dynamic element */
 	virtual void DrawDynamicElements(FPrimitiveDrawInterface* PDI,const FSceneView* View) override;
 		
@@ -1043,29 +1042,48 @@ void FInstancedStaticMeshSceneProxy::DrawDynamicElements(FPrimitiveDrawInterface
 #endif
 }
 
+inline void FInstancedStaticMeshSceneProxy::SetupInstancedMeshBatch(int32 LODIndex, FMeshBatch& OutMeshBatch) const
+{
+	const bool bInstanced = RHISupportsInstancing(GRHIShaderPlatform);
+	OutMeshBatch.VertexFactory = &InstancedRenderData.VertexFactories[LODIndex];
+	uint32 NumInstances = InstancedRenderData.InstanceBuffer.GetNumInstances();
+	auto* OutBatchElement0 = OutMeshBatch.Elements.GetTypedData();
+	OutBatchElement0->UserData = (void*)&UserData_AllInstances;
+	OutBatchElement0->UserIndex = 0;
+	if (bInstanced && NumInstances > 1)
+	{
+		OutBatchElement0->NumInstances = NumInstances;
+	}
+	else
+	{
+		OutMeshBatch.Elements.Reserve(NumInstances);
+		OutBatchElement0->NumInstances = 1;
+		for (uint32 Instance = 1; Instance < NumInstances; ++Instance)
+		{
+			auto* NewBatchElement = new(OutMeshBatch.Elements) FMeshBatchElement();
+			*NewBatchElement = *OutBatchElement0;
+			NewBatchElement->UserIndex = Instance;
+		}
+	}
+}
+
+
 bool FInstancedStaticMeshSceneProxy::GetShadowMeshElement(int32 LODIndex, uint8 InDepthPriorityGroup, FMeshBatch& OutMeshElement) const
 {
 	if (LODIndex < InstancedRenderData.VertexFactories.Num() && FStaticMeshSceneProxy::GetShadowMeshElement(LODIndex, InDepthPriorityGroup, OutMeshElement))
 	{
-		const bool bInstanced = RHISupportsInstancing(GRHIShaderPlatform);
-		OutMeshElement.Elements[0].NumInstances = bInstanced ? InstancedRenderData.InstanceBuffer.GetNumInstances() : 1;
-		OutMeshElement.Elements[0].UserData = (void*)&UserData_AllInstances;
-		OutMeshElement.VertexFactory = &InstancedRenderData.VertexFactories[LODIndex];
+		SetupInstancedMeshBatch(LODIndex, OutMeshElement);
 		return true;
 	}
 	return false;
 }
 
-
 /** Sets up a FMeshBatch for a specific LOD and element. */
 bool FInstancedStaticMeshSceneProxy::GetMeshElement(int32 LODIndex,int32 ElementIndex,uint8 InDepthPriorityGroup,FMeshBatch& OutMeshElement, const bool bUseSelectedMaterial, const bool bUseHoveredMaterial) const
 {
 	if (LODIndex < InstancedRenderData.VertexFactories.Num() && FStaticMeshSceneProxy::GetMeshElement(LODIndex, ElementIndex, InDepthPriorityGroup, OutMeshElement, bUseSelectedMaterial, bUseHoveredMaterial))
 	{
-		const bool bInstanced = RHISupportsInstancing(GRHIShaderPlatform);
-		OutMeshElement.Elements[0].NumInstances = bInstanced ? InstancedRenderData.InstanceBuffer.GetNumInstances() : 1;
-		OutMeshElement.Elements[0].UserData = (void*)&UserData_AllInstances;
-		OutMeshElement.VertexFactory = &InstancedRenderData.VertexFactories[LODIndex];
+		SetupInstancedMeshBatch(LODIndex, OutMeshElement);
 		return true;
 	}
 	return false;
@@ -1076,10 +1094,7 @@ bool FInstancedStaticMeshSceneProxy::GetWireframeMeshElement(int32 LODIndex, con
 {
 	if (LODIndex < InstancedRenderData.VertexFactories.Num() && FStaticMeshSceneProxy::GetWireframeMeshElement(LODIndex, WireframeRenderProxy, InDepthPriorityGroup, OutMeshElement))
 	{
-		const bool bInstanced = RHISupportsInstancing(GRHIShaderPlatform);
-		OutMeshElement.Elements[0].NumInstances = bInstanced ? InstancedRenderData.InstanceBuffer.GetNumInstances() : 1;
-		OutMeshElement.Elements[0].UserData = (void*)&UserData_AllInstances;
-		OutMeshElement.VertexFactory = &InstancedRenderData.VertexFactories[LODIndex];
+		SetupInstancedMeshBatch(LODIndex, OutMeshElement);
 		return true;
 	}
 	return false;
@@ -1658,3 +1673,61 @@ void UInstancedStaticMeshComponent::SelectInstance(bool bInSelected, int32 InIns
 	}
 #endif
 }
+
+void FInstancedStaticMeshVertexFactoryShaderParameters::SetMesh( FRHICommandList& RHICmdList, FShader* VertexShader,const class FVertexFactory* VertexFactory,const class FSceneView& View,const struct FMeshBatchElement& BatchElement,uint32 DataFlags ) const 
+{
+	SCOPED_DRAW_EVENT(STATIC_INSTANCED, DEC_SCENE_ITEMS);
+	FRHIVertexShader* VS = VertexShader->GetVertexShader();
+	if( InstancingFadeOutParamsParameter.IsBound() )
+	{
+		FVector4 InstancingFadeOutParams(0.f,0.f,1.f,1.f);
+
+		const FInstancingUserData* InstancingUserData = (FInstancingUserData*)BatchElement.UserData;
+		if (InstancingUserData)
+		{
+			InstancingFadeOutParams.X = InstancingUserData->StartCullDistance;
+			if( InstancingUserData->EndCullDistance > 0 )
+			{
+				if( InstancingUserData->EndCullDistance > InstancingUserData->StartCullDistance )
+				{
+					InstancingFadeOutParams.Y = 1.f / (float)(InstancingUserData->EndCullDistance - InstancingUserData->StartCullDistance);
+				}
+				else
+				{
+					InstancingFadeOutParams.Y = 1.f;
+				}
+			}
+			else
+			{
+				InstancingFadeOutParams.Y = 0.f;
+			}
+
+			InstancingFadeOutParams.Z = InstancingUserData->bRenderSelected ? 1.f : 0.f;
+			InstancingFadeOutParams.W = InstancingUserData->bRenderUnselected ? 1.f : 0.f;
+		}
+		SetShaderValue(RHICmdList, VS, InstancingFadeOutParamsParameter, InstancingFadeOutParams );
+	}
+
+	const bool bInstanced = RHISupportsInstancing(GRHIShaderPlatform);
+	if (!bInstanced)
+	{
+		if (CPUInstanceShadowMapBias.IsBound())
+		{
+			const auto* InstancingData = (FInstancingUserData*)BatchElement.UserData;
+			auto* InstanceStream = ((FInstancingUserData::FInstanceStream*)InstancingData->RenderData->InstanceBuffer.GetRawData()) + BatchElement.UserIndex;
+			SetShaderValue(RHICmdList, VS, CPUInstanceShadowMapBias, InstanceStream->InstanceShadowmapUVBias);
+			SetShaderValueArray(RHICmdList, VS, CPUInstanceTransform, InstanceStream->InstanceTransform, 3);
+			SetShaderValueArray(RHICmdList, VS, CPUInstanceInverseTransform, InstanceStream->InstanceInverseTransform, 3);
+		}
+	}
+
+	// upload SpeedTree buffer, if available
+	if (View.Family != NULL && View.Family->Scene != NULL)
+	{
+		FUniformBufferRHIParamRef SpeedTreeUniformBuffer = View.Family->Scene->GetSpeedTreeUniformBuffer(VertexFactory);
+		if (SpeedTreeUniformBuffer != NULL)
+		{
+			SetUniformBufferParameter(RHICmdList, VS, VertexShader->GetUniformBufferParameter<FSpeedTreeUniformParameters>(), SpeedTreeUniformBuffer);
+		}
+	}
+}
