UE4 - Fix for grey screen

diff --git a/Engine/Shaders/HeightFogCommon.usf b/Engine/Shaders/HeightFogCommon.usf
--- a/Engine/Shaders/HeightFogCommon.usf
+++ b/Engine/Shaders/HeightFogCommon.usf
@@ -39,8 +39,10 @@ half4 GetExponentialHeightFog(float3 WorldPosition, float3 InCameraPosition)
 
 	if (InscatteringLightDirection.w > 0)
 	{
-		// Setup a cosine lobe around the light direction to approximate inscattering from the directional light off of the ambient haze
-		half3 DirectionalLightInscattering = DirectionalInscatteringColor.xyz * pow(saturate(dot(CameraToReceiverNormalized, InscatteringLightDirection.xyz)), DirectionalInscatteringColor.w);
+		// Setup a cosine lobe around the light direction to approximate inscattering from the directional light off of the ambient haze;
+		// clamp to avoid issues on some platforms (Could've used ClampedPow(), but we want the dot to be 0..1)
+		half3 DirectionalLightInscattering = DirectionalInscatteringColor.xyz * pow(saturate(dot(CameraToReceiverNormalized, InscatteringLightDirection.xyz)), clamp(DirectionalInscatteringColor.w, 0.000001f, 1.0));
+
 		// Calculate the line integral of the eye ray through the haze, using a special starting distance to limit the inscattering to the distance
 		float DirExponentialHeightLineIntegral = ExponentialHeightLineIntegralShared * max(CameraToReceiverLength - DirectionalInscatteringStartDistance, 0.0f);
 		// Calculate the amount of light that made it through the fog using the transmission equation
