Fixed intensity of translucency lighting.

diff --git a/Engine/Shaders/BasePassPixelShader.usf b/Engine/Shaders/BasePassPixelShader.usf
--- a/Engine/Shaders/BasePassPixelShader.usf
+++ b/Engine/Shaders/BasePassPixelShader.usf
@@ -137,7 +137,7 @@ float3 GetTranslucencyLighting(FMaterialPixelParameters MaterialParameters, floa
 		TranslucentLighting.R.V.x = Vector0.r;
 		TranslucentLighting.G.V.x = Vector0.g;
 		TranslucentLighting.B.V.x = Vector0.b;
-		float3 NormalizedAmbientColor = Vector0.rgb / max(Vector0.r, max(Vector0.g, Vector0.b));
+		float3 NormalizedAmbientColor = Vector0.rgb / Luminance( Vector0.rgb );
 
 		// Scale the monocrome directional coefficients with the normalzed ambient color as an approximation to the uncompressed values
 		TranslucentLighting.R.V.yzw = Vector1.rgb * NormalizedAmbientColor.r;
@@ -146,8 +146,7 @@ float3 GetTranslucencyLighting(FMaterialPixelParameters MaterialParameters, floa
 			
 		// Compute diffuse lighting which takes the normal into account
 		FTwoBandSHVector DiffuseTransferSH = CalcDiffuseTransferSH(MaterialParameters.WorldNormal, 1);
-		//@todo - why is x2 needed to match opaque and non-directional?
-		float4 VolumeLighting = float4(max(half3(0,0,0), DotSH(TranslucentLighting, DiffuseTransferSH)) * 2, Vector0.a);
+		float4 VolumeLighting = float4(max(half3(0,0,0), DotSH(TranslucentLighting, DiffuseTransferSH)), Vector0.a);
 		InterpolatedLighting = DiffuseColor * VolumeLighting.rgb;
 
 	#elif TRANSLUCENCY_LIGHTING_VOLUMETRIC_NONDIRECTIONAL
@@ -158,10 +157,10 @@ float3 GetTranslucencyLighting(FMaterialPixelParameters MaterialParameters, floa
 
 		// Lerp between cascades
 		// Normal is not taken into account with non directional lighting, and only the ambient term of the SH coefficients are needed
-		float4 VolumeLighting = lerp(OuterLighting, InnerLighting, FinalLerpFactor);
+		float4 TranslucentLighting = lerp(OuterLighting, InnerLighting, FinalLerpFactor);
 
-		// Undo the SH ambient coefficient scaling
-		VolumeLighting.rgb /= SHAmbientFunction();
+		float DiffuseTransfer = PI * 0.282095f;
+		float4 VolumeLighting = float4( TranslucentLighting.rgb * DiffuseTransfer, TranslucentLighting.a );
 
 		// Determine lighting if no self shadowing is applied
 		InterpolatedLighting = DiffuseColor * VolumeLighting.rgb;
@@ -241,8 +240,7 @@ float3 GetTranslucencyLighting(FMaterialPixelParameters MaterialParameters, floa
 		#elif TRANSLUCENCY_LIGHTING_VOLUMETRIC_NONDIRECTIONAL
 
 			float3 PointIndirectLighting = float3(IndirectLightingSHCoefficients[0].x, IndirectLightingSHCoefficients[1].x, IndirectLightingSHCoefficients[2].x);
-			//@todo - why is /2 needed to match directional?
-			float3 DiffuseGI = PointIndirectLighting / (SHAmbientFunction() * 2);
+			float3 DiffuseGI = PointIndirectLighting * DiffuseTransfer;
 			IndirectIrradiance += Luminance(DiffuseGI);
 
 			InterpolatedLighting += DiffuseColor * DiffuseGI * MaterialAO;
diff --git a/Engine/Shaders/TranslucentLightInjectionShaders.usf b/Engine/Shaders/TranslucentLightInjectionShaders.usf
--- a/Engine/Shaders/TranslucentLightInjectionShaders.usf
+++ b/Engine/Shaders/TranslucentLightInjectionShaders.usf
@@ -66,7 +66,7 @@ void SimpleLightInjectMainPS(
 	float3 WorldPositionForLighting = WorldPosition + 1 * GetBoxPushout(NormalizedLightVector, .5f * VoxelSize);
 
 	float Attenuation = CalcSimpleLightAttenuation(WorldPositionForLighting);
-	float3 Lighting = SimpleLightColorAndExponent.rgb / PI * Attenuation;
+	float3 Lighting = SimpleLightColorAndExponent.rgb * Attenuation;
 
 	FTwoBandSHVectorRGB SHLighting = MulSH(SHBasisFunction(NormalizedLightVector), Lighting);
 	OutColor0 = float4(SHLighting.R.V.x, SHLighting.G.V.x, SHLighting.B.V.x, 0);
@@ -305,7 +305,7 @@ void InjectMainPS(
 		// Apply light function after edge fading, so that a light function which is modifying the light's brightness still works
 		ShadowFactor *= GetLightFunctionShadowFactor(WorldPositionForLighting);
 
-		float3 Lighting = DeferredLightUniforms.LightColor / PI * Attenuation * ShadowFactor;
+		float3 Lighting = DeferredLightUniforms.LightColor * Attenuation * ShadowFactor;
 
 		FTwoBandSHVectorRGB SHLighting = MulSH(SHBasisFunction(NormalizedLightVector), Lighting);
 
diff --git a/Engine/Shaders/TranslucentLightingShaders.usf b/Engine/Shaders/TranslucentLightingShaders.usf
--- a/Engine/Shaders/TranslucentLightingShaders.usf
+++ b/Engine/Shaders/TranslucentLightingShaders.usf
@@ -101,7 +101,7 @@ void CompositeGIMainPS(
 	float3 WorldPosition = float3(View.TranslucencyLightingVolumeMin[VolumeCascadeIndex].xy + Input.Vertex.UV / View.TranslucencyLightingVolumeInvSize[VolumeCascadeIndex].xy - .5f * View.TranslucencyLightingVolumeInvSize[VolumeCascadeIndex].w, ZPosition);
 	float3 CacheUV = WorldPosition * IndirectlightingCacheScale + IndirectlightingCacheAdd;
 
-	float4 Vector0 = Texture3DSample(IndirectLightingCacheTexture, IndirectLightingCacheTextureSampler, CacheUV) * SHAmbientFunction() / PI;
+	float4 Vector0 = Texture3DSample(IndirectLightingCacheTexture, IndirectLightingCacheTextureSampler, CacheUV);
 	float4 Vector1 = Texture3DSample(IndirectLightingCacheTexture1, IndirectLightingCacheTexture1Sampler, CacheUV);
 	float4 Vector2 = Texture3DSample(IndirectLightingCacheTexture2, IndirectLightingCacheTexture2Sampler, CacheUV);
 
