UE4 - Fix for opengl (ambiguous conversion)

diff --git a/Engine/Shaders/MaterialTemplate.usf b/Engine/Shaders/MaterialTemplate.usf
--- a/Engine/Shaders/MaterialTemplate.usf
+++ b/Engine/Shaders/MaterialTemplate.usf
@@ -289,7 +289,7 @@ MaterialFloat3 TransformLocalVectorToTangent(FMaterialPixelParameters Parameters
 MaterialFloat3 TransformLocalVectorToWorld(FMaterialVertexParameters Parameters,MaterialFloat3 InLocalVector)
 {
 	#if USE_INSTANCING || PARTICLE_MESH_FACTORY
-		return mul(InLocalVector, (half3x3)Parameters.InstanceLocalToWorld);
+		return mul(InLocalVector, (MaterialFloat3x3)Parameters.InstanceLocalToWorld);
 	#else
 		return mul(InLocalVector, GetLocalToWorld3x3());
 	#endif
