Nondirectional uses SH1

diff --git a/Engine/Shaders/BasePassPixelShader.usf b/Engine/Shaders/BasePassPixelShader.usf
--- a/Engine/Shaders/BasePassPixelShader.usf
+++ b/Engine/Shaders/BasePassPixelShader.usf
@@ -156,11 +156,16 @@ float3 GetTranslucencyLighting(FMaterialPixelParameters MaterialParameters, floa
 		float4 OuterLighting = Texture3DSampleLevel(TranslucencyLightingVolumeAmbientOuter, TranslucencyLightingVolumeAmbientOuterSampler, OuterVolumeUVs, 0);
 
 		// Lerp between cascades
+		float4 Vector0 = lerp(OuterLighting, InnerLighting, FinalLerpFactor);
+		
 		// Normal is not taken into account with non directional lighting, and only the ambient term of the SH coefficients are needed
-		float4 TranslucentLighting = lerp(OuterLighting, InnerLighting, FinalLerpFactor);
+		FOneBandSHVectorRGB TranslucentLighting;
+		TranslucentLighting.R.V.x = Vector0.r;
+		TranslucentLighting.G.V.x = Vector0.g;
+		TranslucentLighting.B.V.x = Vector0.b;
 
-		float DiffuseTransfer = PI * 0.282095f;
-		float4 VolumeLighting = float4( TranslucentLighting.rgb * DiffuseTransfer, TranslucentLighting.a );
+		FOneBandSHVector DiffuseTransferSH = CalcDiffuseTransferSH1(MaterialParameters.WorldNormal, 1);
+		float4 VolumeLighting = float4( DotSH1(TranslucentLighting, DiffuseTransferSH), Vector0.a );
 
 		// Determine lighting if no self shadowing is applied
 		InterpolatedLighting = DiffuseColor * VolumeLighting.rgb;
@@ -239,8 +244,12 @@ float3 GetTranslucencyLighting(FMaterialPixelParameters MaterialParameters, floa
 
 		#elif TRANSLUCENCY_LIGHTING_VOLUMETRIC_NONDIRECTIONAL
 
-			float3 PointIndirectLighting = float3(IndirectLightingSHCoefficients[0].x, IndirectLightingSHCoefficients[1].x, IndirectLightingSHCoefficients[2].x);
-			float3 DiffuseGI = PointIndirectLighting * DiffuseTransfer;
+			FOneBandSHVectorRGB PointIndirectLighting;
+			PointIndirectLighting.R.V = IndirectLightingSHCoefficients[0].x;
+			PointIndirectLighting.G.V = IndirectLightingSHCoefficients[1].x;
+			PointIndirectLighting.B.V = IndirectLightingSHCoefficients[2].x;
+
+			float3 DiffuseGI = DotSH1(PointIndirectLighting, DiffuseTransferSH);
 			IndirectIrradiance += Luminance(DiffuseGI);
 
 			InterpolatedLighting += DiffuseColor * DiffuseGI * MaterialAO;
diff --git a/Engine/Shaders/SHCommon.usf b/Engine/Shaders/SHCommon.usf
--- a/Engine/Shaders/SHCommon.usf
+++ b/Engine/Shaders/SHCommon.usf
@@ -4,6 +4,20 @@
 	SHCommon.usf: Spherical Harmonic functionality
 =============================================================================*/
 
+/** The SH coefficients for the projection of a function that maps directions to scalar values. */
+struct FOneBandSHVector
+{
+	half V;
+};
+
+/** The SH coefficients for the projection of a function that maps directions to RGB values. */
+struct FOneBandSHVectorRGB
+{
+	FOneBandSHVector R;
+	FOneBandSHVector G;
+	FOneBandSHVector B;
+};
+
 /** The SH coefficients for the projection of a function that maps directions to scalar values. */
 struct FTwoBandSHVector
 {
@@ -81,6 +95,21 @@ half3 DotSH(FTwoBandSHVectorRGB A,FTwoBandSHVector B)
 	return Result;
 }
 
+half DotSH1(FOneBandSHVector A,FOneBandSHVector B)
+{
+	half Result = A.V * B.V;
+	return Result;
+}
+
+half3 DotSH1(FOneBandSHVectorRGB A,FOneBandSHVector B)
+{
+	half3 Result = 0;
+	Result.r = DotSH1(A.R,B);
+	Result.g = DotSH1(A.G,B);
+	Result.b = DotSH1(A.B,B);
+	return Result;
+}
+
 half DotSH3(FThreeBandSHVector A,FThreeBandSHVector B)
 {
 	half Result = dot(A.V0, B.V0);
@@ -99,6 +128,14 @@ half3 DotSH3(FThreeBandSHVectorRGB A,FThreeBandSHVector B)
 }
 
 /** Computes the value of the SH basis functions for a given direction. */
+FOneBandSHVector SHBasisFunction1(half3 InputVector)
+{
+	FOneBandSHVector Result;
+	// These are derived from simplifying SHBasisFunction in C++
+	Result.V = 0.282095f; 
+	return Result;
+}
+
 FTwoBandSHVector SHBasisFunction(half3 InputVector)
 {
 	FTwoBandSHVector Result;
@@ -139,6 +176,20 @@ half SHAmbientFunction()
  * Computes the diffuse transfer function for a surface with the given normal and DiffusePower,
  * and projects it into the SH basis.
  */
+FOneBandSHVector CalcDiffuseTransferSH1(half3 Normal,half Exponent)
+{
+	FOneBandSHVector Result = SHBasisFunction1(Normal);
+
+	// These formula are scaling factors for each SH band that convolve a SH with the circularly symmetric function
+	// max(0,cos(theta))^Exponent
+	half L0 =					2 * PI / (1 + 1 * Exponent							);
+
+	// Multiply the coefficients in each band with the appropriate band scaling factor.
+	Result.V *= L0;
+
+	return Result;
+}
+
 FTwoBandSHVector CalcDiffuseTransferSH(half3 Normal,half Exponent)
 {
 	FTwoBandSHVector Result = SHBasisFunction(Normal);
