Enable Hit proxy with ES2 emulation shader fix.

diff --git a/Engine/Shaders/HitProxyVertexShader.usf b/Engine/Shaders/HitProxyVertexShader.usf
--- a/Engine/Shaders/HitProxyVertexShader.usf
+++ b/Engine/Shaders/HitProxyVertexShader.usf
@@ -13,7 +13,7 @@ struct FHitProxyVSToPS
 {
 	FVertexFactoryInterpolantsVSToPS FactoryInterpolants;
 #if USE_INSTANCING	
-	float4 InstanceHitProxyId : TEXCOORD5;
+	float4 InstanceHitProxyId : HIT_PROXY_ID;
 #endif
 	float4 PixelPosition : TEXCOORD6;
 	float4 Position : SV_POSITION;
@@ -24,7 +24,7 @@ struct FHitProxyVSToPS
 	{
 		FVertexFactoryInterpolantsVSToDS FactoryInterpolants;
 	    #if USE_INSTANCING	
-		    float4 InstanceHitProxyId : TEXCOORD5;
+		    float4 InstanceHitProxyId : HIT_PROXY_ID;
 	    #endif
 		float4 Position : VS_To_DS_Position;
 		OPTIONAL_VertexID_VS_To_DS
