Fix failure to link hit-proxy shaders on Mac OpenGL, which is much more pedantic than other platforms, by properly updating pixel shader semantics to match recent vertex shader semantic changes.

diff --git a/Engine/Shaders/HitProxyPixelShader.usf b/Engine/Shaders/HitProxyPixelShader.usf
--- a/Engine/Shaders/HitProxyPixelShader.usf
+++ b/Engine/Shaders/HitProxyPixelShader.usf
@@ -13,7 +13,7 @@ float4 HitProxyId;
 void Main(
 	FVertexFactoryInterpolantsVSToPS Interpolants,
 #if USE_INSTANCING	
-	float4 InstanceHitProxyId : TEXCOORD5,
+	float4 InstanceHitProxyId : HIT_PROXY_ID,
 #endif
 	float4 PixelPosition	: TEXCOORD6,
 	OPTIONAL_IsFrontFace,
