Fixed spot light overlap display

diff --git a/Engine/Shaders/StationaryLightOverlapShaders.usf b/Engine/Shaders/StationaryLightOverlapShaders.usf
--- a/Engine/Shaders/StationaryLightOverlapShaders.usf
+++ b/Engine/Shaders/StationaryLightOverlapShaders.usf
@@ -51,15 +51,15 @@ void OverlapRadialPixelMain(
 		float4 HomogeneousWorldPosition = mul(float4(InScreenPosition.xy / InScreenPosition.w * SceneDepth, SceneDepth, 1), View.ScreenToWorld);
 		float3 WorldPosition = HomogeneousWorldPosition.xyz / HomogeneousWorldPosition.w;
 
-		float3 L = DeferredLightUniforms.NormalizedLightDirection;
+		float3 ToLight = DeferredLightUniforms.NormalizedLightDirection;
 
 		#if RADIAL_ATTENUATION
-			L = DeferredLightUniforms.LightPosition - WorldPosition;
+			ToLight = DeferredLightUniforms.LightPosition - WorldPosition;
 		#endif
 
 		float SpotFalloff = 1;
 		#if RADIAL_ATTENUATION
-			SpotFalloff = SpotAttenuation(L.xyz, -DeferredLightUniforms.NormalizedLightDirection, DeferredLightUniforms.SpotAngles);
+			SpotFalloff = SpotAttenuation(normalize(ToLight), -DeferredLightUniforms.NormalizedLightDirection, DeferredLightUniforms.SpotAngles);
 		#endif
 
 		float RadialAttenuationRatio = length(WorldPosition - DeferredLightUniforms.LightPosition) * DeferredLightUniforms.LightInvRadius;
