Optimization to the faster temporal AA (removing velocity weighting). Improved temporal AA by adjusting anti-flicker logic, now not light biased.

diff --git a/Engine/Shaders/PostProcessTemporalAA.usf b/Engine/Shaders/PostProcessTemporalAA.usf
--- a/Engine/Shaders/PostProcessTemporalAA.usf
+++ b/Engine/Shaders/PostProcessTemporalAA.usf
@@ -242,6 +242,7 @@ void MainTemporalAAPS( float4 UVAndScreenPos : TEXCOORD0, float3 InExposureScale
 		#define AA_LOWPASS 1
 		#define AA_ONE_DYNAMIC_SAMPLE 1
 		#define AA_DEBUG 0
+		#define AA_VELOCITY_WEIGHTING 1
 	#else
 		#define AA_NEW 0
 		#define AA_BORDER 1
@@ -267,6 +268,7 @@ void MainFastTemporalAAPS( float4 UVAndScreenPos : TEXCOORD0, float3 InExposureS
 		#define AA_LOWPASS 0
 		#define AA_ONE_DYNAMIC_SAMPLE 1
 		#define AA_DEBUG 0
+		#define AA_VELOCITY_WEIGHTING 0
 	#else
 		#define AA_NEW 0
 		#define AA_BORDER 1
diff --git a/Engine/Shaders/PostProcessTemporalCommon.usf b/Engine/Shaders/PostProcessTemporalCommon.usf
--- a/Engine/Shaders/PostProcessTemporalCommon.usf
+++ b/Engine/Shaders/PostProcessTemporalCommon.usf
@@ -120,6 +120,12 @@ which can be configured via defines for various temporal AA passes.
 	#define AA_NEW 0
 #endif
 
+// Use velocity weighting.
+#ifndef AA_VELOCITY_WEIGHTING
+	#define AA_VELOCITY_WEIGHTING 0
+#endif
+
+
 
 
 #if AA_NEW
@@ -401,6 +407,8 @@ which can be configured via defines for various temporal AA passes.
 		
 	// COMPUTE BLEND AMOUNT 
     // --------------------
+	// Replace history with minimum difference of history from min or max neighborhood.
+	LumaHistory = min(abs(LumaMin-LumaHistory), abs(LumaMax-LumaHistory));
 	float HistoryAmount = (1.0/8.0) + HistoryBlur * (1.0/8.0);
 	float HistoryFactor = LumaHistory * HistoryAmount * (1.0 + HistoryBlur * HistoryAmount * 4.0);
 	float BlendFinal = saturate(HistoryFactor * rcp(LumaHistory + LumaContrast));
@@ -409,7 +417,7 @@ which can be configured via defines for various temporal AA passes.
 		BlendFinal = 1.0/4.0;
 	#endif
 
-   	#if (AA_ALPHA == 1)
+   	#if (AA_ALPHA == 1) && (AA_VELOCITY_WEIGHTING == 1)
 		// Velocity weighting.
 		// Decay prior value.
 		float PriorVelocity = OutColor.a;
@@ -976,4 +984,5 @@ which can be configured via defines for various temporal AA passes.
 #undef AA_FORCE_ALPHA_CLAMP
 #undef AA_ONE_DYNAMIC_SAMPLE
 #undef AA_GREEN_AS_LUMA
-#undef AA_NEW
\ No newline at end of file
+#undef AA_NEW
+#undef AA_VELOCITY_WEIGHTING
