Fixed double bright lighting

diff --git a/Engine/Shaders/DeferredLightingCommon.usf b/Engine/Shaders/DeferredLightingCommon.usf
--- a/Engine/Shaders/DeferredLightingCommon.usf
+++ b/Engine/Shaders/DeferredLightingCommon.usf
@@ -608,8 +608,7 @@ float4 GetDynamicLighting(float3 WorldPosition, float3 CameraVector, float2 InUV
 				// wrong screen space subsurface scattering but potentially faster
 				float3 SurfaceLighting = SurfaceShading(ScreenSpaceData.GBuffer, LobeRoughness, LobeEnergy, L, V, N, float2(1, 1));
 				LightAccumulator_Add(LightAccumulator, SurfaceLighting, 0, LightColor * (NoL * SurfaceAttenuation));
-#endif			
-				LightAccumulator_Add(LightAccumulator, SurfaceLightingDiff, SurfaceLightingSpec, LightColor * (NoL * SurfaceAttenuation));
+#endif
 			}
 
 			// accumulate subsurface
diff --git a/Engine/Shaders/TranslucentLightInjectionShaders.usf b/Engine/Shaders/TranslucentLightInjectionShaders.usf
--- a/Engine/Shaders/TranslucentLightInjectionShaders.usf
+++ b/Engine/Shaders/TranslucentLightInjectionShaders.usf
@@ -66,7 +66,7 @@ void SimpleLightInjectMainPS(
 	float3 WorldPositionForLighting = WorldPosition + 1 * GetBoxPushout(NormalizedLightVector, .5f * VoxelSize);
 
 	float Attenuation = CalcSimpleLightAttenuation(WorldPositionForLighting);
-	float3 Lighting = SimpleLightColorAndExponent.rgb * Attenuation;
+	float3 Lighting = SimpleLightColorAndExponent.rgb / PI * Attenuation;
 
 	FTwoBandSHVectorRGB SHLighting = MulSH(SHBasisFunction(NormalizedLightVector), Lighting);
 	OutColor0 = float4(SHLighting.R.V.x, SHLighting.G.V.x, SHLighting.B.V.x, 0);
@@ -305,7 +305,7 @@ void InjectMainPS(
 		// Apply light function after edge fading, so that a light function which is modifying the light's brightness still works
 		ShadowFactor *= GetLightFunctionShadowFactor(WorldPositionForLighting);
 
-		float3 Lighting = DeferredLightUniforms.LightColor * Attenuation * ShadowFactor;
+		float3 Lighting = DeferredLightUniforms.LightColor / PI * Attenuation * ShadowFactor;
 
 		FTwoBandSHVectorRGB SHLighting = MulSH(SHBasisFunction(NormalizedLightVector), Lighting);
 
