UE4 - Fix for subuv's precision

diff --git a/Engine/Shaders/MaterialTemplate.usf b/Engine/Shaders/MaterialTemplate.usf
--- a/Engine/Shaders/MaterialTemplate.usf
+++ b/Engine/Shaders/MaterialTemplate.usf
@@ -43,9 +43,9 @@ struct FMaterialParticleParameters
 
 #if USE_PARTICLE_SUBUVS
 	/** SubUV texture coordinates*/
-	half2 SubUVCoords[2];
+	MaterialFloat2 SubUVCoords[2];
 	/** SubUV interpolation value*/
-	half SubUVLerp;
+	MaterialFloat SubUVLerp;
 #endif
 
 	/** The size of the particle. */
