[INTEGRATE] 4.4 branch back to main.

diff --git a/Engine/Shaders/HitProxyPixelShader.usf b/Engine/Shaders/HitProxyPixelShader.usf
--- a/Engine/Shaders/HitProxyPixelShader.usf
+++ b/Engine/Shaders/HitProxyPixelShader.usf
@@ -13,7 +13,7 @@ float4 HitProxyId;
 void Main(
 	FVertexFactoryInterpolantsVSToPS Interpolants,
 #if USE_INSTANCING	
-	float4 InstanceHitProxyId : HIT_PROXY_ID,
+	float4 InstanceHitProxyId : TEXCOORD5,
 #endif
 	float4 PixelPosition	: TEXCOORD6,
 	OPTIONAL_IsFrontFace,
