Fixed IndirectLightingCache show flag disabling direct lighting on dynamic objects

diff --git a/Engine/Shaders/LightmapCommon.usf b/Engine/Shaders/LightmapCommon.usf
--- a/Engine/Shaders/LightmapCommon.usf
+++ b/Engine/Shaders/LightmapCommon.usf
@@ -110,16 +110,16 @@ half4 GetPrecomputedShadowMasks(FVertexFactoryInterpolantsVSToPS Interpolants)
 		half4 ShadowFactors = saturate( DistanceField * DistanceFieldParameters.x + DistanceFieldParameters.y );
 		return StaticShadowMapMasks * ShadowFactors * ShadowFactors;
 
-	#elif CACHED_VOLUME_INDIRECT_LIGHTING || CACHED_POINT_INDIRECT_LIGHTING
+	#elif HQ_TEXTURE_LIGHTMAP || LQ_TEXTURE_LIGHTMAP
 
-		// Mark as unshadowed for dynamic objects, they will receive dynamic shadowing from the environment
-		return 1;
+		// Mark as shadowed for lightmapped objects with no shadowmap
+		// This is necessary because objects inside a light's influence that were determined to be completely shadowed won't be rendered with STATICLIGHTING_TEXTUREMASK==1
+		return 0;
 
 	#else
 
-		// Mark as shadowed by default
-		// This is necessary because objects inside a light's influence that were determined to be completely shadowed won't be rendered with STATICLIGHTING_TEXTUREMASK==1
-		return 0;
+		// Mark as unshadowed by default
+		return 1;
 
 	#endif
 }
@@ -134,16 +134,16 @@ half GetPrimaryPrecomputedShadowMask(FVertexFactoryInterpolantsVSToPS Interpolan
 		half ShadowFactor = saturate( DistanceField * DistanceFieldParameters.x + DistanceFieldParameters.y );
 		return ShadowFactor * ShadowFactor;
 
-	#elif CACHED_VOLUME_INDIRECT_LIGHTING || CACHED_POINT_INDIRECT_LIGHTING
+	#elif HQ_TEXTURE_LIGHTMAP || LQ_TEXTURE_LIGHTMAP
 
-		// Mark as unshadowed for dynamic objects, they will receive dynamic shadowing from the environment
-		return 1;
+		// Mark as shadowed for lightmapped objects with no shadowmap
+		// This is necessary because objects inside a light's influence that were determined to be completely shadowed won't be rendered with STATICLIGHTING_TEXTUREMASK==1
+		return 0;
 
 	#else
 
-		// Mark as shadowed by default
-		// This is necessary because objects inside a light's influence that were determined to be completely shadowed won't be rendered with STATICLIGHTING_TEXTUREMASK==1
-		return 0;
+		// Mark as unshadowed by default
+		return 1;
 
 	#endif
 }
\ No newline at end of file
