Fixed specular intensity on mobile

diff --git a/Engine/Shaders/BasePassForForwardShadingPixelShader.usf b/Engine/Shaders/BasePassForForwardShadingPixelShader.usf
--- a/Engine/Shaders/BasePassForForwardShadingPixelShader.usf
+++ b/Engine/Shaders/BasePassForForwardShadingPixelShader.usf
@@ -29,8 +29,8 @@ half PhongApprox( half Roughness, half RoL )
 	// Spherical Gaussian approximation: pow( x, n ) ~= exp( (n + 0.775) * (x - 1) )
 	// Phong: n = 0.5 / a2 - 0.5
 	// 0.5 / ln(2), 0.275 / ln(2)
-	half2 c = half2( 0.72134752, 0.25 ) * rcp_a2 + half2( 0.39674113, 0.75 );	// 1 mad
-	return c.y * exp2( c.x * RoL - c.x );	// 2 mad, 1 exp2, 1 mul
+	half c = 0.72134752 * rcp_a2 + 0.39674113;	// 1 mad
+	return rcp_a2 * exp2( c * RoL - c );		// 2 mad, 1 exp2, 1 mul
 	// Total 7 instr
 }
 
