Make Landscape Splines work on OpenGL ES2
Tested on Nexus 5
Vertex shader has not been optimised for mobile, only made to work

diff --git a/Engine/Shaders/LocalVertexFactory.usf b/Engine/Shaders/LocalVertexFactory.usf
--- a/Engine/Shaders/LocalVertexFactory.usf
+++ b/Engine/Shaders/LocalVertexFactory.usf
@@ -41,7 +41,11 @@ float2 SplineEndScale;
 float2 SplineEndOffset;
 
 float3 SplineUpDir;
+#if ES2_PROFILE
+float SmoothInterpRollScale;
+#else
 bool SmoothInterpRollScale;
+#endif
 
 float SplineMeshMinZ;
 float SplineMeshScaleZ;
@@ -322,36 +326,36 @@ FMaterialVertexParameters GetMaterialVertexParameters( FVertexFactoryInput Input
 	}
 
 	/** Calculate rotation matrix that defines frame along spline, given the Z of a vertex. */
-	float3x3 CalcSliceRot(float ZPos)
+	half3x3 CalcSliceRot(float ZPos)
 	{
 		// Find how far 'along' mesh we are
-		float Alpha = ZPos * SplineMeshScaleZ - SplineMeshMinZ;
+		half Alpha = ZPos * SplineMeshScaleZ - SplineMeshMinZ;
 
 		// Apply hermite interp to Alpha if desired
-		float HermiteAlpha = SmoothInterpRollScale ? smoothstep(0.0, 1.0, Alpha) : Alpha;
+		half HermiteAlpha = SmoothInterpRollScale ? smoothstep(0.0, 1.0, Alpha) : Alpha;
 
 		// Then find the point and direction of the spline at this point along
-		float3 SplineDir = SplineEvalDir( SplineStartPos, SplineStartTangent, SplineEndPos, SplineEndTangent, Alpha );
+		half3 SplineDir = SplineEvalDir( SplineStartPos, SplineStartTangent, SplineEndPos, SplineEndTangent, Alpha );
 
 		// Find base frenet frame
-		float3 BaseXVec = normalize( cross(SplineUpDir, SplineDir) );
-		float3 BaseYVec = normalize( cross(SplineDir, BaseXVec) );
+		half3 BaseXVec = normalize( cross(SplineUpDir, (half3)SplineDir) );
+		half3 BaseYVec = normalize( cross((half3)SplineDir, BaseXVec) );
 
 		// Apply roll to frame around spline
-		float UseRoll = lerp(SplineStartRoll, SplineEndRoll, HermiteAlpha);
-		float SinAng, CosAng;
+		half UseRoll = lerp((half)SplineStartRoll, (half)SplineEndRoll, HermiteAlpha);
+		half SinAng, CosAng;
 		sincos(UseRoll, SinAng, CosAng);
-		float3 XVec = (CosAng * BaseXVec) - (SinAng * BaseYVec);
-		float3 YVec = (CosAng * BaseYVec) + (SinAng * BaseXVec);
+		half3 XVec = (CosAng * BaseXVec) - (SinAng * BaseYVec);
+		half3 YVec = (CosAng * BaseYVec) + (SinAng * BaseXVec);
 
 		// Find scale at this point along spline
-		float2 UseScale = lerp(SplineStartScale, SplineEndScale, HermiteAlpha);
+		half2 UseScale = lerp(SplineStartScale, SplineEndScale, HermiteAlpha);
 
 		XVec *= sign(UseScale.x);
 		YVec *= sign(UseScale.y);
 
 		// Build rotation transform
-		float3x3 SliceTransform = mul(transpose(float3x3(SplineMeshDir, SplineMeshX, SplineMeshY)), float3x3(SplineDir, XVec, YVec));
+		half3x3 SliceTransform = mul(transpose(half3x3(SplineMeshDir, SplineMeshX, SplineMeshY)), half3x3(SplineDir, XVec, YVec));
 
 		return SliceTransform;
 	}
@@ -386,7 +390,7 @@ half3x3 CalcTangentToLocal(FVertexFactoryInput Input)
 
 #if USE_SPLINEDEFORM
 	// Make slice rotation matrix, and use that to transform tangents
-	float3x3 SliceRot = CalcSliceRot(dot(Input.Position.xyz, SplineMeshDir));
+	half3x3 SliceRot = CalcSliceRot(dot(Input.Position.xyz, SplineMeshDir));
 
 	half3 TangentX = mul(TangentBias(Input.TangentX), SliceRot);
 	TangentZ.xyz = mul(TangentZ.xyz, SliceRot);
diff --git a/Engine/Source/Runtime/Engine/Private/SplineMeshComponent.cpp b/Engine/Source/Runtime/Engine/Private/SplineMeshComponent.cpp
--- a/Engine/Source/Runtime/Engine/Private/SplineMeshComponent.cpp
+++ b/Engine/Source/Runtime/Engine/Private/SplineMeshComponent.cpp
@@ -412,7 +412,7 @@ FPrimitiveSceneProxy* USplineMeshComponent::CreateSceneProxy()
 		StaticMesh &&
 		StaticMesh->HasValidRenderData();
 
-	if (bMeshIsValid && (GetScene()->GetFeatureLevel() >= ERHIFeatureLevel::SM3))
+	if (bMeshIsValid)
 	{
 		return ::new FSplineMeshSceneProxy(this);
 	}
diff --git a/Engine/Source/Runtime/Engine/Private/SplineMeshSceneProxy.h b/Engine/Source/Runtime/Engine/Private/SplineMeshSceneProxy.h
--- a/Engine/Source/Runtime/Engine/Private/SplineMeshSceneProxy.h
+++ b/Engine/Source/Runtime/Engine/Private/SplineMeshSceneProxy.h
@@ -26,8 +26,7 @@ public:
 	static bool ShouldCache(EShaderPlatform Platform, const class FMaterial* Material, const class FShaderType* ShaderType)
 	{
 		return (Material->IsUsedWithSplineMeshes() || Material->IsSpecialEngineMaterial())
-			&& FLocalVertexFactory::ShouldCache(Platform, Material, ShaderType)
-			&& IsFeatureLevelSupported(Platform, ERHIFeatureLevel::SM3);
+			&& FLocalVertexFactory::ShouldCache(Platform, Material, ShaderType);
 	}
 
 	/** Modify compile environment to enable spline deformation */
