Fixed SSR with clear coat. Optimized reflection compute again. Broke diffuse from captures.

diff --git a/Engine/Shaders/PostProcessAmbient.usf b/Engine/Shaders/PostProcessAmbient.usf
--- a/Engine/Shaders/PostProcessAmbient.usf
+++ b/Engine/Shaders/PostProcessAmbient.usf
@@ -372,7 +372,24 @@ float2 IntegrateBRDF( uint2 Random, float Roughness, float NoV )
 
 		if( NoL > 0 )
 		{
-			float Vis = Vis_Schlick( Roughness, NoV, NoL );
+			//float Vis = Vis_Schlick( Roughness, NoV, NoL );
+
+			float a = Square( Roughness );
+			float a2 = a*a;
+
+#if 1
+			float Vis_SchlickV = NoL * ( NoV * ( 1 - a ) + a );
+			float Vis_SchlickL = NoV * ( NoL * ( 1 - a ) + a );
+			float Vis = 0.5 * rcp( Vis_SchlickV + Vis_SchlickL );
+#elif 0
+			float Vis_SmithV = NoV + sqrt( NoV * (NoV - NoV * a2) + a2 );
+			float Vis_SmithL = NoL + sqrt( NoL * (NoL - NoL * a2) + a2 );
+			float Vis = rcp( Vis_SmithV * Vis_SmithL );
+#else
+			float Vis_SmithV = NoL * sqrt( NoV * (NoV - NoV * a2) + a2 );
+			float Vis_SmithL = NoV * sqrt( NoL * (NoL - NoL * a2) + a2 );
+			float Vis = 0.5 * rcp( Vis_SmithV + Vis_SmithL );
+#endif
 
 			// Incident light = NoL
 			// pdf = D * NoH / (4 * VoH)
@@ -442,11 +459,12 @@ void MainPS(in float4 UVAndScreenPos : TEXCOORD0, out float4 OutColor : SV_Targe
 		{
 			float ClearCoat = GBuffer.CustomData.x;
 			float2 AB = PreIntegratedGF.SampleLevel( PreIntegratedGFSampler, float2( NoV, GBuffer.Roughness ), 0 ).rg;
-			SpecularContribution += SampleColor * GBuffer.SpecularColor * AB.x + AB.y * (1 - ClearCoat);
+			SpecularContribution += SampleColor * ( GBuffer.SpecularColor * AB.x + AB.y * (1 - ClearCoat) );
 		}
 		else
 		{
 			SpecularContribution += SampleColor * EnvBRDF( GBuffer.SpecularColor, GBuffer.Roughness, NoV );
+			//SpecularContribution += ApproximateSpecularIBL( Random, GBuffer.SpecularColor, GBuffer.Roughness, GBuffer.WorldNormal, -ScreenVector );
 		}
 	}
 
diff --git a/Engine/Shaders/ReflectionEnvironmentComputeShaders.usf b/Engine/Shaders/ReflectionEnvironmentComputeShaders.usf
--- a/Engine/Shaders/ReflectionEnvironmentComputeShaders.usf
+++ b/Engine/Shaders/ReflectionEnvironmentComputeShaders.usf
@@ -24,6 +24,7 @@ SamplerState ReflectionEnvironmentColorSampler;
 #define ATOMIC_REDUCTION 0
 #endif
 
+#define AABB_INTERSECT		1
 #define VISUALIZE_OVERLAP	0
 
 uint NumCaptures;
@@ -161,17 +162,51 @@ void ComputeTileMinMax(uint ThreadIndex, float SceneDepth, out float MinTileZ, o
 #endif
 }
 
+bool SphereVsBox( float3 SphereCenter, float SphereRadius, float3 BoxCenter, float3 BoxExtent )
+{
+	float3 ClosestOnBox = max( 0, abs( BoxCenter - SphereCenter ) - BoxExtent );
+	return dot( ClosestOnBox, ClosestOnBox ) < SphereRadius * SphereRadius;
+}
+
+
 // Culls reflection captures in the scene with the current tile
 // Outputs are stored in shared memory
 void DoTileCulling(uint3 GroupId, uint ThreadIndex, float MinTileZ, float MaxTileZ, float MinTileZ2, float MaxTileZ2)
 {
+#if AABB_INTERSECT
+	float3 TileBoxCenter;
+	float3 TileBoxExtent;
+	// can be optmized
+
+	// left top front
+	float2 ScreenUV0 = float2((GroupId.xy + int2(0, 0))* float2(THREADGROUP_SIZEX, THREADGROUP_SIZEY) + .5f) / (ViewDimensions.zw - ViewDimensions.xy);
+	float3 ScreenPos0 = float3(float2(2.0f, -2.0f) * ScreenUV0 + float2(-1.0f, 1.0f), ConvertToDeviceZ(MinTileZ));
+
+	// right bottom back
+	float2 ScreenUV1 = float2((GroupId.xy + int2(1, 1)) * float2(THREADGROUP_SIZEX, THREADGROUP_SIZEY) - .5f) / (ViewDimensions.zw - ViewDimensions.xy);
+	float3 ScreenPos1 = float3(float2(2.0f, -2.0f) * ScreenUV1 + float2(-1.0f, 1.0f), ConvertToDeviceZ(MaxTileZ));
+
+	// back rect
+	float4 ViewPos0 = mul(float4(ScreenPos0.x, ScreenPos0.y, ScreenPos1.z, 1), View.ClipToView);	ViewPos0.xyz /= ViewPos0.w;
+	float4 ViewPos1 = mul(float4(ScreenPos0.x, ScreenPos1.y, ScreenPos1.z, 1), View.ClipToView);	ViewPos1.xyz /= ViewPos1.w;
+	float4 ViewPos2 = mul(float4(ScreenPos1.x, ScreenPos0.y, ScreenPos1.z, 1), View.ClipToView);	ViewPos2.xyz /= ViewPos2.w;
+	float4 ViewPos3 = mul(float4(ScreenPos1.x, ScreenPos1.y, ScreenPos1.z, 1), View.ClipToView);	ViewPos3.xyz /= ViewPos3.w;
+	// front point
+	float4 ViewPos4 = mul(float4(ScreenPos0.xy, ScreenPos0.z, 1), View.ClipToView);	ViewPos4.xyz /= ViewPos4.w;
+
+	float3 TileBoxMin = min(ViewPos4.xyz, min(ViewPos0.xyz, ViewPos3.xyz));
+	float3 TileBoxMax = max(ViewPos4.xyz, max(ViewPos0.xyz, ViewPos3.xyz));
+
+	TileBoxCenter = (TileBoxMax + TileBoxMin) * 0.5f;
+	TileBoxExtent = (TileBoxMax - TileBoxMin) * 0.5f;
+#else
 	// Setup tile frustum planes
 	float2 TileScale = float2(ViewDimensions.zw - ViewDimensions.xy) * rcp(2 * float2(THREADGROUP_SIZEX, THREADGROUP_SIZEY));
 	float2 TileBias = TileScale - GroupId.xy;
-
-    float4 C1 = float4(View.ViewToClip._11 * TileScale.x,	0.0f,								View.ViewToClip._31 * TileScale.x + TileBias.x,	0.0f);
-    float4 C2 = float4(0.0f,								-View.ViewToClip._22 * TileScale.y, View.ViewToClip._32 * TileScale.y + TileBias.y,	0.0f);
-    float4 C4 = float4(0.0f,								0.0f,								1.0f,											0.0f);
+	
+	float4 C1 = float4(View.ViewToClip._11 * TileScale.x,	0.0f,								View.ViewToClip._31 * TileScale.x + TileBias.x,	0.0f);
+	float4 C2 = float4(0.0f,								-View.ViewToClip._22 * TileScale.y, View.ViewToClip._32 * TileScale.y + TileBias.y,	0.0f);
+	float4 C4 = float4(0.0f,								0.0f,								1.0f,											0.0f);
 
 	// TODO transform to world space
 #if ATOMIC_REDUCTION
@@ -200,7 +235,8 @@ void DoTileCulling(uint3 GroupId, uint ThreadIndex, float MinTileZ, float MaxTil
 	{
 		frustumPlanes[i] *= rcp(length(frustumPlanes[i].xyz));
 	}
-
+#endif
+	
 	if (ThreadIndex == 0) 
 	{
 		TileNumReflectionCaptures = 0;
@@ -217,9 +253,19 @@ void DoTileCulling(uint3 GroupId, uint ThreadIndex, float MinTileZ, float MaxTil
 
 		float3 BoundsViewPosition = mul(float4(CapturePositionAndRadius.xyz + View.PreViewTranslation.xyz, 1), View.TranslatedWorldToView).xyz;
 
+#if AABB_INTERSECT
+		// Add this capture to the list of indices if it intersects
+		BRANCH
+		if( SphereVsBox( BoundsViewPosition, CapturePositionAndRadius.w, TileBoxCenter, TileBoxExtent ) )
+		{
+			uint ListIndex;
+			InterlockedAdd(TileNumReflectionCaptures, 1U, ListIndex);
+			TileReflectionCaptureIndices[ListIndex] = CaptureIndex;
+		}
+#else
 		// Cull the light against the tile's frustum planes
 		// Note: this has some false positives, a light that is intersecting three different axis frustum planes yet not intersecting the volume of the tile will be treated as intersecting
-		bool bInTile = true;  
+		bool bInTile = true;
 
 		// Test against the screen x and y oriented planes first
 		UNROLL
@@ -275,6 +321,7 @@ void DoTileCulling(uint3 GroupId, uint ThreadIndex, float MinTileZ, float MaxTil
 				TileReflectionCaptureIndices[ListIndex] = CaptureIndex;
 			}
 		}
+#endif
 	}
 
 	GroupMemoryBarrierWithGroupSync();
@@ -316,30 +363,66 @@ void DoTileCulling(uint3 GroupId, uint ThreadIndex, float MinTileZ, float MaxTil
 	GroupMemoryBarrierWithGroupSync();
 }
 
-struct FLobe
+float CountOverlap( float3 WorldPosition )
 {
-	float4	Color;
-	float3	Direction;
-	float	Roughness;
-	float	RayInfluance;
-};
+	float Overlap = 0;
+	float Opacity = 1;
 
-void GatherRadiance( inout FLobe Lobes[2], const uint NumLobes, float3 WorldPosition, float3 RayDirection )
-{
-	float LobeMip[2];
+	uint NumCapturesAffectingTile = TileNumReflectionCaptures;
 
-	UNROLL
-	for( uint Lobe = 0; Lobe < NumLobes; Lobe++ )
+	// Accumulate reflections from captures affecting this tile, applying largest captures first so that the smallest ones display on top
+	LOOP
+	for (uint TileCaptureIndex = 0; TileCaptureIndex < NumCapturesAffectingTile; TileCaptureIndex++) 
 	{
-		LobeMip[ Lobe ] = ComputeReflectionCaptureMipFromRoughness( Lobes[ Lobe ].Roughness );
+		#if SORT_CAPTURES
+			uint CaptureIndex = SortedTileReflectionCaptureIndices[TileCaptureIndex];
+		#else
+			uint CaptureIndex = TileReflectionCaptureIndices[TileCaptureIndex];
+		#endif
+			
+		float4 CapturePositionAndRadius = ReflectionCapture.PositionAndRadius[CaptureIndex];
+		float3 CaptureVector = WorldPosition - CapturePositionAndRadius.xyz;
+		float CaptureVectorLength = length(CaptureVector);
+
+		BRANCH
+		if (CaptureVectorLength < CapturePositionAndRadius.w)
+		{
+			float NormalizedDistanceToCapture = saturate(CaptureVectorLength / CapturePositionAndRadius.w);
+
+			// Fade out based on distance to capture
+			float x = saturate( 2.5 * NormalizedDistanceToCapture - 1.5 );
+			float DistanceAlpha = 1 - x*x*(3 - 2*x);
+
+			Overlap += 1;
+			Opacity *= 1 - DistanceAlpha;
+
+			BRANCH
+			if( Opacity < 0.001 )
+			{
+				break;
+			}
+		}
 	}
 
+	return Overlap;
+}
+
+void GatherRadiance( inout float4 Color, float3 WorldPosition, float3 RayDirection, float Roughness )
+{
+	float Mip = ComputeReflectionCaptureMipFromRoughness( Roughness );
+
 	uint NumCapturesAffectingTile = TileNumReflectionCaptures;
 
 	// Accumulate reflections from captures affecting this tile, applying largest captures first so that the smallest ones display on top
 	LOOP
 	for (uint TileCaptureIndex = 0; TileCaptureIndex < NumCapturesAffectingTile; TileCaptureIndex++) 
 	{
+		BRANCH
+		if( Color.a < 0.001 )
+		{
+			break;
+		}
+
 		#if SORT_CAPTURES
 			uint CaptureIndex = SortedTileReflectionCaptureIndices[TileCaptureIndex];
 		#else
@@ -349,7 +432,7 @@ void GatherRadiance( inout FLobe Lobes[2], const uint NumLobes, float3 WorldPosi
 		float4 CapturePositionAndRadius = ReflectionCapture.PositionAndRadius[CaptureIndex];
 		float4 CaptureProperties = ReflectionCapture.CaptureProperties[CaptureIndex];
 		float3 CaptureVector = WorldPosition - CapturePositionAndRadius.xyz;
-		float CaptureVectorLength = sqrt(dot(CaptureVector, CaptureVector));
+		float CaptureVectorLength = length(CaptureVector);
 
 		BRANCH
 		if (CaptureVectorLength < CapturePositionAndRadius.w)
@@ -362,53 +445,6 @@ void GatherRadiance( inout FLobe Lobes[2], const uint NumLobes, float3 WorldPosi
 
 			#define PROJECT_ONTO_SHAPE 1
 			#if PROJECT_ONTO_SHAPE
-					
-				#define SUPPORT_PLANE 0
-				#if SUPPORT_PLANE
-					BRANCH
-					// Plane
-					if (CaptureProperties.b > 1)
-					{
-						float4 ImagePlane = float4(ReflectionCapture.BoxTransform[CaptureIndex][0][0], ReflectionCapture.BoxTransform[CaptureIndex][1][0], ReflectionCapture.BoxTransform[CaptureIndex][2][0], ReflectionCapture.BoxTransform[CaptureIndex][3][0]);
-
-						float VectorDotPlaneNormal = dot(ImagePlane.xyz, RayDirection);
-						// VectorDotPlaneNormal < 0 means the ray hit the front face
-						BRANCH
-						if (VectorDotPlaneNormal < 0)
-						{
-							float PlaneDistance = dot(ImagePlane.xyz, WorldPosition) - ImagePlane.w;
-							// Time along the ray defined by WorldPosition + IntersectionTime * RayDirection that the intersection took place
-							float IntersectionTime = -PlaneDistance / VectorDotPlaneNormal;
-
-							BRANCH
-							// Skip intersections behind the pixel being shaded
-							if (IntersectionTime > 0)
-							{
-								// Calculate the world space intersection position
-								float3 IntersectPosition = WorldPosition + IntersectionTime * ReflectionVector;
-									
-								float2 ReflectionUVs;
-								float4 CurrentReflectionXAxis = float4(ReflectionCapture.BoxTransform[CaptureIndex][0][1], ReflectionCapture.BoxTransform[CaptureIndex][1][1], ReflectionCapture.BoxTransform[CaptureIndex][2][1], ReflectionCapture.BoxTransform[CaptureIndex][3][1]);
-								float3 CurrentImageReflectionOrigin = CapturePositionAndRadius.xyz;
-								float XLength = length(CurrentReflectionXAxis.xyz);
-								float3 NormalizedXAxis = CurrentReflectionXAxis.xyz / XLength;
-								// Calculate the quad UVs by projecting the vector from the intersection to the quad origin onto each quad axis
-								ReflectionUVs.x = dot(NormalizedXAxis, IntersectPosition - CurrentImageReflectionOrigin.xyz);
-								float3 ReflectionYAxis = cross(ImagePlane.xyz, NormalizedXAxis) * CurrentReflectionXAxis.w;
-								ReflectionUVs.y = dot(ReflectionYAxis, IntersectPosition - CurrentImageReflectionOrigin.xyz);
-								ReflectionUVs = .5f * ReflectionUVs / float2(XLength, XLength * CurrentReflectionXAxis.w) + .5f;
-
-								ProjectedCaptureVector = IntersectPosition - CapturePositionAndRadius.xyz;
-
-								if (ReflectionUVs.x > 0 && ReflectionUVs.x < 1 && ReflectionUVs.y > 0 && ReflectionUVs.y < 1)
-								{
-									CompositedLighting.rgb += 1;
-								}
-							}
-						}
-					}
-					else 
-				#endif
 				// Box
 				BRANCH if (CaptureProperties.b > 0)
 				{
@@ -444,131 +480,64 @@ void GatherRadiance( inout FLobe Lobes[2], const uint NumLobes, float3 WorldPosi
 				// Sphere
 				else
 				{
-					float ProjectionSphereRadius = CapturePositionAndRadius.w * 1.2f;
+					float ProjectionSphereRadius = CapturePositionAndRadius.w;// * 1.2f;
 					float SphereRadiusSquared = ProjectionSphereRadius * ProjectionSphereRadius;
 
-					float3 ReceiverToSphereCenter = WorldPosition - CapturePositionAndRadius.xyz;
-					float ReceiverToSphereCenterSq = dot(ReceiverToSphereCenter, ReceiverToSphereCenter);
+					float3 LocalPosition = WorldPosition - CapturePositionAndRadius.xyz;
+					float LocalPositionSqr = dot(LocalPosition, LocalPosition);
 
 					// Find the intersection between the ray along the reflection vector and the capture's sphere
 					float3 QuadraticCoef;
 					QuadraticCoef.x = 1;
-					QuadraticCoef.y = 2 * dot(RayDirection, ReceiverToSphereCenter);
-					QuadraticCoef.z = ReceiverToSphereCenterSq - SphereRadiusSquared;
+					QuadraticCoef.y = dot(RayDirection, LocalPosition);
+					QuadraticCoef.z = LocalPositionSqr - SphereRadiusSquared;
+
+					float Determinant = QuadraticCoef.y * QuadraticCoef.y - QuadraticCoef.z;
 
-					float Determinant = QuadraticCoef.y * QuadraticCoef.y - 4 * QuadraticCoef.z;
-							 
 					// Only continue if the ray intersects the sphere
-					if (Determinant >= 0)
+					//if (Determinant >= 0)
 					{
-						float FarIntersection = (sqrt(Determinant) - QuadraticCoef.y) * 0.5;
+						float FarIntersection = sqrt(Determinant) - QuadraticCoef.y;
 
-						float3 IntersectPosition = WorldPosition + FarIntersection * RayDirection;
-						ProjectedCaptureVector = IntersectPosition - CapturePositionAndRadius.xyz;
+						ProjectedCaptureVector = LocalPosition + FarIntersection * RayDirection;
 						// Note: some compilers don't handle smoothstep min > max (this was 1, .6)
-						DistanceAlpha = 1.0 - smoothstep(.6, 1, NormalizedDistanceToCapture);
+						//DistanceAlpha = 1.0 - smoothstep(.6, 1, NormalizedDistanceToCapture);
+
+						float x = saturate( 2.5 * NormalizedDistanceToCapture - 1.5 );
+						DistanceAlpha = 1 - x*x*(3 - 2*x);
 					}
 				}
-
 			#else 
 				DistanceAlpha = 1.0;
 			#endif
 
 			float CaptureArrayIndex = CaptureProperties.g;
-			float Opacity = 0;
 
-			UNROLL
-			for( uint Lobe = 0; Lobe < NumLobes; Lobe++ )
 			{
-				float3 SampleDir = lerp( Lobes[ Lobe ].Direction, ProjectedCaptureVector, Lobes[ Lobe ].RayInfluance );
-				float4 Sample = ReflectionEnvironmentColorTexture.SampleLevel( ReflectionEnvironmentColorSampler, float4( SampleDir, CaptureArrayIndex ), LobeMip[ Lobe ] );
+				float4 Sample = ReflectionEnvironmentColorTexture.SampleLevel( ReflectionEnvironmentColorSampler, float4( ProjectedCaptureVector, CaptureArrayIndex ), Mip );
 
 				Sample.rgb *= CaptureProperties.r;
 				Sample *= DistanceAlpha;
 
 				// Under operator (back to front)
-				Lobes[ Lobe ].Color.rgb += Sample.rgb * Lobes[ Lobe ].Color.a;
-				Lobes[ Lobe ].Color.a *= 1 - Sample.a;
-
-				Opacity += Lobes[ Lobe ].Color.a;
-			}
-
-			BRANCH
-			if( Opacity < 0.001 )
-			{
-				break;
+				Color.rgb += Sample.rgb * Color.a;
+				Color.a *= 1 - Sample.a;
 			}
 		}
 	}
 
 #define APPLY_SKY_LIGHT 1
 #if APPLY_SKY_LIGHT
-	float Opacity = 0;
-
-	UNROLL
-	for( uint Lobe = 0; Lobe < NumLobes; Lobe++ )
-	{
-		Opacity += Lobes[ Lobe ].Color.a;
-	}
-
 	BRANCH
-	if( SkyLightParameters.y > 0 && Opacity >= 0.001 )
+	if( SkyLightParameters.y > 0 && Color.a >= 0.001 )
 	{
-		UNROLL
-		for( uint Lobe = 0; Lobe < NumLobes; Lobe++ )
-		{
-			// TODO use LobeMip
-			Lobes[ Lobe ].Color.rgb += Lobes[ Lobe ].Color.a * GetSkyLightReflection( Lobes[ Lobe ].Direction, Lobes[ Lobe ].Roughness, USE_LIGHTMAPS );
-		}
+		// TODO use Mip
+		Color.rgb += Color.a * GetSkyLightReflection( RayDirection, Roughness, USE_LIGHTMAPS );
 	}
 #endif
 }
 
 
-float CountOverlap( float3 WorldPosition )
-{
-	float Overlap = 0;
-	float Opacity = 1;
-
-	uint NumCapturesAffectingTile = TileNumReflectionCaptures;
-
-	// Accumulate reflections from captures affecting this tile, applying largest captures first so that the smallest ones display on top
-	LOOP
-	for (uint TileCaptureIndex = 0; TileCaptureIndex < NumCapturesAffectingTile; TileCaptureIndex++) 
-	{
-		#if SORT_CAPTURES
-			uint CaptureIndex = SortedTileReflectionCaptureIndices[TileCaptureIndex];
-		#else
-			uint CaptureIndex = TileReflectionCaptureIndices[TileCaptureIndex];
-		#endif
-			
-		float4 CapturePositionAndRadius = ReflectionCapture.PositionAndRadius[CaptureIndex];
-		float3 CaptureVector = WorldPosition - CapturePositionAndRadius.xyz;
-		float CaptureVectorLength = sqrt(dot(CaptureVector, CaptureVector));
-
-		BRANCH
-		if (CaptureVectorLength < CapturePositionAndRadius.w)
-		{
-			float NormalizedDistanceToCapture = saturate(CaptureVectorLength / CapturePositionAndRadius.w);
-
-			// Fade out based on distance to capture
-			float DistanceAlpha = 1.0 - smoothstep(.6, 1, NormalizedDistanceToCapture);
-
-			Overlap += 1;
-			Opacity *= 1 - DistanceAlpha;
-
-			BRANCH
-			if( Opacity < 0.001 )
-			{
-				break;
-			}
-		}
-	}
-
-	return Overlap;
-}
-
-
 Texture2D ScreenSpaceReflections;
 Texture2D InSceneColor;
 
@@ -600,103 +569,63 @@ void ReflectionEnvironmentTiledDeferredMain(
 	FScreenSpaceData ScreenSpaceData = GetScreenSpaceDataUint(PixelPos);
 	FGBufferData GBuffer = ScreenSpaceData.GBuffer;
 	
+	float4 Color = float4(0, 0, 0, 1);
+	
 	float4 HomogeneousWorldPosition = mul(float4(ScreenPosition * SceneDepth, SceneDepth, 1), View.ScreenToWorld);
 	float3 WorldPosition = HomogeneousWorldPosition.xyz / HomogeneousWorldPosition.w;
 	float3 CameraToPixel = normalize(WorldPosition - View.ViewOrigin.xyz);
 	float3 ReflectionVector = reflect(CameraToPixel, GBuffer.WorldNormal);
-	
-	// Save GPRs by using R instead of V
-	float NoV = saturate( dot( GBuffer.WorldNormal, ReflectionVector ) );
-
-	FLobe Lobes[2];
-
-	UNROLL
-	for( uint Lobe = 0; Lobe < 2; Lobe++ )
-	{
-		Lobes[ Lobe ].Color			= float4(0, 0, 0, 1);
-		Lobes[ Lobe ].Direction		= ReflectionVector;
-		Lobes[ Lobe ].RayInfluance	= 1;
-		Lobes[ Lobe ].Roughness		= GBuffer.Roughness;
-	}
 
 #if VISUALIZE_OVERLAP
 	float Overlap = CountOverlap( WorldPosition );
-#else
-
-#if 1//APPLY_SSR
-	float4 SSR = ScreenSpaceReflections.Load( int3(PixelPos, 0) );
-	Lobes[0].Color.rgb = SSR.rgb;
-	Lobes[0].Color.a = 1 - SSR.a;
 #endif
 
-#if USE_LIGHTMAPS
-	// We have high frequency directional data but low frequency spatial data in the envmap.
-	// We have high frequency spatial data but low frequency directional data in the lightmap.
-	// So, we combine the two for the best of both. This is done by removing the low spatial frequencies from the envmap and replacing them with the lightmap data.
-	// This is only done with luma so as to not get odd color shifting.
-	// Note: make sure this matches the lightmap mixing done for translucency (BasePassPixelShader.usf)
-	Lobes[0].Color.a *= GBuffer.IndirectIrradiance;
-	Lobes[1].Color.a *= GBuffer.IndirectIrradiance;
-#else
-	// Diffuse lobe
-	Lobes[1].Direction		= GBuffer.WorldNormal;
-	Lobes[1].RayInfluance	= 0;
-	Lobes[1].Roughness		= 1;
-#endif
-	
-	float AO = ScreenSpaceData.AmbientOcclusion;
-	float SpecularOcclusion = saturate( Square( NoV + AO ) - 1 + AO );
-
-	Lobes[0].Color.a *= SpecularOcclusion;
-	Lobes[1].Color.a *= AO;
-
-	BRANCH	//[forcecase]
-	switch( GBuffer.ShadingModelID )
+	BRANCH
+	if( GBuffer.ShadingModelID != SHADINGMODELID_UNLIT )
 	{
-		case SHADINGMODELID_UNLIT:
-			break;
-		case SHADINGMODELID_DEFAULT_LIT:
-		case SHADINGMODELID_SUBSURFACE:
-		case SHADINGMODELID_PREINTEGRATED_SKIN:
-		case SHADINGMODELID_SUBSURFACE_PROFILE:
-		{
-			BRANCH
-			if( GBuffer.ShadingModelID == SHADINGMODELID_SUBSURFACE )
-			{
-				GBuffer.DiffuseColor += ExtractSubsurfaceColor(GBuffer);
-			}
+		// Save GPRs by using R instead of V
+		float NoV = saturate( dot( GBuffer.WorldNormal, ReflectionVector ) );
 
-		#if USE_LIGHTMAPS
-			GBuffer.DiffuseColor = 0;
-		#endif
-
-			GBuffer.SpecularColor = EnvBRDF( GBuffer.SpecularColor, GBuffer.Roughness, NoV );
-
-			GatherRadiance( Lobes, USE_LIGHTMAPS ? 1 : 2, WorldPosition, ReflectionVector );
+	#if 1//APPLY_SSR
+		float4 SSR = ScreenSpaceReflections.Load( int3(PixelPos, 0) );
+		Color.rgb = SSR.rgb;
+		Color.a = 1 - SSR.a;
+	#endif
 
-			Lobes[0].Color.rgb *= GBuffer.SpecularColor;
-			Lobes[1].Color.rgb *= GBuffer.DiffuseColor;
-			break;
+		if( GBuffer.ShadingModelID == SHADINGMODELID_CLEAR_COAT )
+		{
+			const float ClearCoat = GBuffer.CustomData.x;
+			//Color = float4(0,0,0,1);
+			Color = lerp( Color, float4(0,0,0,1), ClearCoat );
 		}
-		case SHADINGMODELID_CLEAR_COAT:
-#if 1
+	
+		float AO = ScreenSpaceData.AmbientOcclusion;
+		//float SpecularOcclusion = saturate( Square( NoV + AO ) - 1 + AO );
+		float SpecularOcclusion = saturate( pow( NoV + AO, GBuffer.Roughness ) - 1 + AO );
+
+		Color.a *= SpecularOcclusion;
+
+	#if USE_LIGHTMAPS
+		// We have high frequency directional data but low frequency spatial data in the envmap.
+		// We have high frequency spatial data but low frequency directional data in the lightmap.
+		// So, we combine the two for the best of both. This is done by removing the low spatial frequencies from the envmap and replacing them with the lightmap data.
+		// This is only done with luma so as to not get odd color shifting.
+		// Note: make sure this matches the lightmap mixing done for translucency (BasePassPixelShader.usf)
+		Color.a *= GBuffer.IndirectIrradiance;
+	#endif
+	
+		GatherRadiance( Color, WorldPosition, ReflectionVector, GBuffer.Roughness );
+
+		BRANCH
+		if( GBuffer.ShadingModelID == SHADINGMODELID_CLEAR_COAT )
 		{
 			const float ClearCoat			= GBuffer.CustomData.x;
 			const float ClearCoatRoughness	= GBuffer.CustomData.y;
 
-		#if USE_LIGHTMAPS
-			Lobes[0].Roughness = ClearCoatRoughness;
-		#else
-			Lobes[0].Roughness = lerp( GBuffer.Roughness, ClearCoatRoughness, ClearCoat );
-			Lobes[1].Roughness = lerp( 1, GBuffer.Roughness, GBuffer.Metallic * ClearCoat );
-			Lobes[1].Direction = lerp( GBuffer.WorldNormal, ReflectionVector, ClearCoat );
-			Lobes[1].RayInfluance = ClearCoat;
-		#endif
-
 			// TODO EnvBRDF should have a mask param
 			float2 AB = PreIntegratedGF.SampleLevel( PreIntegratedGFSampler, float2( NoV, GBuffer.Roughness ), 0 ).rg;
-			GBuffer.SpecularColor = GBuffer.SpecularColor * AB.x + AB.y * saturate( 50 * GBuffer.SpecularColor.g ) * (1 - ClearCoat);
-			
+			Color.rgb *= GBuffer.SpecularColor * AB.x + AB.y * saturate( 50 * GBuffer.SpecularColor.g ) * (1 - ClearCoat);
+		
 			// F_Schlick
 			float F0 = 0.04;
 			float Fc = pow( 1 - NoV, 5 );
@@ -704,29 +633,30 @@ void ReflectionEnvironmentTiledDeferredMain(
 			F *= ClearCoat;
 			
 			float LayerAttenuation = (1 - F);
-
+		
+			Color.rgb *= LayerAttenuation;
+			Color.a = F;
+		
+			Color.rgb += SSR.rgb * F;
+			Color.a *= 1 - SSR.a;
+			Color.a *= SpecularOcclusion;
+		
 		#if USE_LIGHTMAPS
-			Lobes[0].Color.a *= F;
-			Lobes[1].Color.a *= LayerAttenuation;
-
-			float3 Lobe0Reflectance = 1;
-			float3 Lobe1Reflectance = GBuffer.SpecularColor;
-		#else
-			float3 Lobe0Reflectance = lerp( GBuffer.SpecularColor, F, ClearCoat );
-			float3 Lobe1Reflectance = ( GBuffer.DiffuseColor + GBuffer.SpecularColor * ClearCoat ) * LayerAttenuation;
+			// We have high frequency directional data but low frequency spatial data in the envmap.
+			// We have high frequency spatial data but low frequency directional data in the lightmap.
+			// So, we combine the two for the best of both. This is done by removing the low spatial frequencies from the envmap and replacing them with the lightmap data.
+			// This is only done with luma so as to not get odd color shifting.
+			// Note: make sure this matches the lightmap mixing done for translucency (BasePassPixelShader.usf)
+			Color.a *= GBuffer.IndirectIrradiance;
 		#endif
 
-			GatherRadiance( Lobes, 2, WorldPosition, ReflectionVector );
-			
-			Lobes[0].Color.rgb *= Lobe0Reflectance;
-			Lobes[1].Color.rgb *= Lobe1Reflectance;
-			break;
+			GatherRadiance( Color, WorldPosition, ReflectionVector, ClearCoatRoughness );
+		}
+		else
+		{
+			Color.rgb *= EnvBRDF( GBuffer.SpecularColor, GBuffer.Roughness, NoV );
 		}
-#endif
-		default:
-			break;
 	}
-#endif
 
 	// Only write to the buffer for threads inside the view
 	BRANCH
@@ -738,8 +668,7 @@ void ReflectionEnvironmentTiledDeferredMain(
 		//OutColor.rgb = 0.1 * TileNumReflectionCaptures;
 		OutColor.rgb = 0.1 * Overlap;
 #else
-		OutColor.rgb += Lobes[0].Color.rgb;
-		OutColor.rgb += Lobes[1].Color.rgb;
+		OutColor.rgb = Color.rgb;
 #endif
 
 		// Transform NaNs to black, transform negative colors to black.
diff --git a/Engine/Shaders/ScreenSpaceReflections.usf b/Engine/Shaders/ScreenSpaceReflections.usf
--- a/Engine/Shaders/ScreenSpaceReflections.usf
+++ b/Engine/Shaders/ScreenSpaceReflections.usf
@@ -297,13 +297,25 @@ void ScreenSpaceReflectionsPS(in float4 UVAndScreenPos : TEXCOORD0, out float4 O
 	OutColor = 0;
 
 	FScreenSpaceData ScreenSpaceData = GetScreenSpaceData(UV);
+	FGBufferData GBuffer = ScreenSpaceData.GBuffer;
+	
+	float Roughness = GBuffer.Roughness;
+
+	BRANCH
+	if( GBuffer.ShadingModelID == SHADINGMODELID_CLEAR_COAT )
+	{
+		const float ClearCoat			= GBuffer.CustomData.x;
+		const float ClearCoatRoughness	= GBuffer.CustomData.y;
+
+		Roughness = lerp( Roughness, ClearCoatRoughness, ClearCoat );
+	}
 
 	// mask SSR to reduce noise and for better performance, roughness of 0 should have SSR, at MaxRoughness we fade to 0
-	float RoughnessFade = saturate(ScreenSpaceData.GBuffer.Roughness * SSRParams.y + 2);
+	float RoughnessFade = saturate(Roughness * SSRParams.y + 2);
 
 #if 1
 	// Early out
-	BRANCH if( RoughnessFade == 0 || ScreenSpaceData.GBuffer.ShadingModelID == 0 )
+	BRANCH if( RoughnessFade == 0 || GBuffer.ShadingModelID == 0 )
 	{
 		return;
 	}
@@ -318,9 +330,7 @@ void ScreenSpaceReflectionsPS(in float4 UVAndScreenPos : TEXCOORD0, out float4 O
 	return;
 #endif
 
-	float Roughness = ScreenSpaceData.GBuffer.Roughness;
-
-	float3 N = ScreenSpaceData.GBuffer.WorldNormal;
+	float3 N = GBuffer.WorldNormal;
 	float SceneDepth = CalcSceneDepth(UV);
 
 	float3 PositionTranslatedWorld = mul( float3( ScreenPos * SceneDepth, SceneDepth ), (float3x3)View.ScreenToTranslatedWorld );
diff --git a/Engine/Source/Runtime/Engine/Public/SceneManagement.h b/Engine/Source/Runtime/Engine/Public/SceneManagement.h
--- a/Engine/Source/Runtime/Engine/Public/SceneManagement.h
+++ b/Engine/Source/Runtime/Engine/Public/SceneManagement.h
@@ -1506,6 +1506,7 @@ BEGIN_UNIFORM_BUFFER_STRUCT(FViewUniformShaderParameters,ENGINE_API)
 	DECLARE_UNIFORM_BUFFER_STRUCT_MEMBER(FMatrix,TranslatedWorldToView)
 	DECLARE_UNIFORM_BUFFER_STRUCT_MEMBER(FMatrix,ViewToTranslatedWorld)
 	DECLARE_UNIFORM_BUFFER_STRUCT_MEMBER(FMatrix,ViewToClip)
+	DECLARE_UNIFORM_BUFFER_STRUCT_MEMBER(FMatrix,ClipToView)
 	DECLARE_UNIFORM_BUFFER_STRUCT_MEMBER(FMatrix,ClipToTranslatedWorld)
 	DECLARE_UNIFORM_BUFFER_STRUCT_MEMBER(FMatrix,ScreenToWorld)
 	DECLARE_UNIFORM_BUFFER_STRUCT_MEMBER(FMatrix,ScreenToTranslatedWorld)
diff --git a/Engine/Source/Runtime/Renderer/Private/SceneRendering.cpp b/Engine/Source/Runtime/Renderer/Private/SceneRendering.cpp
--- a/Engine/Source/Runtime/Renderer/Private/SceneRendering.cpp
+++ b/Engine/Source/Runtime/Renderer/Private/SceneRendering.cpp
@@ -230,6 +230,7 @@ TUniformBufferRef<FViewUniformShaderParameters> FViewInfo::CreateUniformBuffer(
 	ViewUniformShaderParameters.TranslatedWorldToView = EffectiveTranslatedViewMatrix;
 	ViewUniformShaderParameters.ViewToTranslatedWorld = InvViewMatrix * FTranslationMatrix(ViewMatrices.PreViewTranslation);
 	ViewUniformShaderParameters.ViewToClip = ViewMatrices.ProjMatrix;
+	ViewUniformShaderParameters.ClipToView = ViewMatrices.GetInvProjMatrix();
 	ViewUniformShaderParameters.ClipToTranslatedWorld = ViewMatrices.InvTranslatedViewProjectionMatrix;
 	ViewUniformShaderParameters.ViewForward = EffectiveTranslatedViewMatrix.GetColumn(2);
 	ViewUniformShaderParameters.ViewUp = EffectiveTranslatedViewMatrix.GetColumn(1);
