* improved VisualizeSSS

diff --git a/Engine/Shaders/PostProcessSubsurface.usf b/Engine/Shaders/PostProcessSubsurface.usf
--- a/Engine/Shaders/PostProcessSubsurface.usf
+++ b/Engine/Shaders/PostProcessSubsurface.usf
@@ -138,11 +138,10 @@ void SetupPS(in float4 UVAndScreenPos : TEXCOORD0, out float4 OutColor : SV_Targ
 
 			uint SubsurfaceProfileInt = ExtractSubsurfaceProfileInt(ScreenSpaceData.GBuffer);
 
-//			BRANCH if (ScreenSpaceData.GBuffer.ShadingModelID == SHADINGMODELID_SUBSURFACE_PROFILE)
-			{
-
-				OutColor = float4(0.5f, 0.5f, 0.5f, 0);
+			OutColor = float4(0.5f, 0.5f, 0.5f, 0);
 
+			BRANCH if (ScreenSpaceData.GBuffer.ShadingModelID == SHADINGMODELID_SUBSURFACE_PROFILE)
+			{
 				if (SubsurfaceProfileInt == 0)
 				{
 					// default (no Profile)
