Fixed -opengl4 and removed ancient workaround in RadixSort shaders.

diff --git a/Engine/Shaders/RadixSortShaders.usf b/Engine/Shaders/RadixSortShaders.usf
--- a/Engine/Shaders/RadixSortShaders.usf
+++ b/Engine/Shaders/RadixSortShaders.usf
@@ -32,15 +32,6 @@
 /** Just use a memory barrier, not a full group sync for parallel prefix scans. */
 #define PPS_BARRIER GroupMemoryBarrier
 
-/** By default, enable the constant buffer workaround until the driver bug has been fixed. */
-#ifndef CONSTANT_BUFFER_WORKAROUND
-#if !GL4_PROFILE
-#define CONSTANT_BUFFER_WORKAROUND 1
-#else
-#define CONSTANT_BUFFER_WORKAROUND 0
-#endif
-#endif
-
 /** The number of digits per radix. */
 #define DIGIT_COUNT (1 << RADIX_BITS)
 /** Bitmask to retrieve the digit of a key. */
@@ -55,17 +46,6 @@
 	Common functionality.
 ------------------------------------------------------------------------------*/
 
-#if CONSTANT_BUFFER_WORKAROUND
-struct FRadixSortParameterBuffer
-{
-	uint RadixShift;
-	uint TilesPerGroup;
-	uint ExtraTileCount;
-	uint ExtraKeyCount;
-	uint GroupCount;
-};
-#endif // #if CONSTANT_BUFFER_WORKAROUND
-
 /** Specifies which shaders use the raking technique for computing parallel prefix sums. */
 #define USES_RAKING (RADIX_SORT_UPSWEEP || RADIX_SORT_SPINE || RADIX_SORT_DOWNSWEEP)
 
@@ -199,12 +179,6 @@ void RadixSort_ClearOffsets(
 
 #if RADIX_SORT_UPSWEEP
 
-#if CONSTANT_BUFFER_WORKAROUND
-	Buffer<uint> RadixSortParameterBuffer;
-#else // #if CONSTANT_BUFFER_WORKAROUND
-	#define RadixSortStruct RadixSort
-#endif // #if CONSTANT_BUFFER_WORKAROUND
-
 /** Input keys to be sorted. */
 Buffer<uint> InKeys;
 /** Output buffer for offsets. */
@@ -228,15 +202,6 @@ void RadixSort_Upsweep(
 	uint TileCountForGroup;
 	uint ExtraKeysForGroup;
 
-#if CONSTANT_BUFFER_WORKAROUND
-	FRadixSortParameterBuffer RadixSortStruct;
-	RadixSortStruct.RadixShift = RadixSortParameterBuffer[0];
-	RadixSortStruct.TilesPerGroup = RadixSortParameterBuffer[1];
-	RadixSortStruct.ExtraTileCount = RadixSortParameterBuffer[2];
-	RadixSortStruct.ExtraKeyCount = RadixSortParameterBuffer[3];
-	RadixSortStruct.GroupCount = RadixSortParameterBuffer[4];
-#endif // #if CONSTANT_BUFFER_WORKAROUND
-
 	// Determine group and thread IDs.
 	const uint ThreadId = GroupThreadId.x;
 	const uint GroupId = GroupIdXYZ.x;
@@ -248,15 +213,15 @@ void RadixSort_Upsweep(
 	const uint RakingIndex = GetRakingIndex( ThreadId );
 
 #if 0
-	if ( RadixSortStruct.ExtraKeyCount == 16 && RadixSortStruct.TilesPerGroup != 0 )
+	if ( RadixSortUB.ExtraKeyCount == 16 && RadixSortUB.TilesPerGroup != 0 )
 	{
 		if ( ThreadId == 0 )
 		{
-			OutOffsets[DIGIT_COUNT*1] = RadixSortStruct.RadixShift;
-			OutOffsets[DIGIT_COUNT*2] = RadixSortStruct.TilesPerGroup;
-			OutOffsets[DIGIT_COUNT*3] = RadixSortStruct.ExtraTileCount;
-			OutOffsets[DIGIT_COUNT*4] = RadixSortStruct.ExtraKeyCount;
-			OutOffsets[DIGIT_COUNT*5] = RadixSortStruct.GroupCount;
+			OutOffsets[DIGIT_COUNT*1] = RadixSortUB.RadixShift;
+			OutOffsets[DIGIT_COUNT*2] = RadixSortUB.TilesPerGroup;
+			OutOffsets[DIGIT_COUNT*3] = RadixSortUB.ExtraTileCount;
+			OutOffsets[DIGIT_COUNT*4] = RadixSortUB.ExtraKeyCount;
+			OutOffsets[DIGIT_COUNT*5] = RadixSortUB.GroupCount;
 		}
 	}
 #endif
@@ -279,21 +244,21 @@ void RadixSort_Upsweep(
 	}
 
 	// The number of tiles to process in this group.
-	if ( GroupId < RadixSortStruct.ExtraTileCount )
+	if ( GroupId < RadixSortUB.ExtraTileCount )
 	{
-		FirstTileIndex = GroupId * (RadixSortStruct.TilesPerGroup + 1);
-		TileCountForGroup = RadixSortStruct.TilesPerGroup + 1;
+		FirstTileIndex = GroupId * (RadixSortUB.TilesPerGroup + 1);
+		TileCountForGroup = RadixSortUB.TilesPerGroup + 1;
 	}
 	else
 	{
-		FirstTileIndex = GroupId * RadixSortStruct.TilesPerGroup + RadixSortStruct.ExtraTileCount;
-		TileCountForGroup = RadixSortStruct.TilesPerGroup;
+		FirstTileIndex = GroupId * RadixSortUB.TilesPerGroup + RadixSortUB.ExtraTileCount;
+		TileCountForGroup = RadixSortUB.TilesPerGroup;
 	}
 
 	// The last group has to process any keys after the last tile.
-	if ( GroupId == (RadixSortStruct.GroupCount - 1) )
+	if ( GroupId == (RadixSortUB.GroupCount - 1) )
 	{
-		ExtraKeysForGroup = RadixSortStruct.ExtraKeyCount;
+		ExtraKeysForGroup = RadixSortUB.ExtraKeyCount;
 	}
 	else
 	{
@@ -311,7 +276,7 @@ void RadixSort_Upsweep(
 	while ( GroupKeyBegin < GroupKeyEnd )
 	{
 		const uint Key = InKeys[GroupKeyBegin + ThreadId];
-		const uint Digit = (Key >> RadixSortStruct.RadixShift) & DIGIT_MASK;
+		const uint Digit = (Key >> RadixSortUB.RadixShift) & DIGIT_MASK;
 		const uint BankIndex = Digit * BANKS_PER_DIGIT + BankOffset;
 		LocalCounters[BankIndex * PADDED_BANK_SIZE + CounterOffset] += 1;
 		//LocalCounters[BankOffset * PADDED_BANK_SIZE + CounterOffset] += 1;
@@ -327,7 +292,7 @@ void RadixSort_Upsweep(
 		if ( GroupKeyBegin + ThreadId < GroupKeyEnd )
 		{
 			const uint Key = InKeys[GroupKeyBegin + ThreadId];
-			const uint Digit = (Key >> RadixSortStruct.RadixShift) & DIGIT_MASK;
+			const uint Digit = (Key >> RadixSortUB.RadixShift) & DIGIT_MASK;
 			const uint BankIndex = Digit * BANKS_PER_DIGIT + BankOffset;
 			LocalCounters[BankIndex * PADDED_BANK_SIZE + CounterOffset] += 1;
 			//LocalCounters[BankOffset * PADDED_BANK_SIZE + CounterOffset] += 100;
@@ -544,14 +509,6 @@ void RadixSort_Spine(
 
 #if RADIX_SORT_DOWNSWEEP
 
-#if CONSTANT_BUFFER_WORKAROUND
-	static FRadixSortParameterBuffer RadixSortBuffer;
-	#define RadixSortStruct RadixSortBuffer
-	Buffer<uint> RadixSortParameterBuffer;
-#else // #if CONSTANT_BUFFER_WORKAROUND
-	#define RadixSortStruct RadixSort
-#endif // #if CONSTANT_BUFFER_WORKAROUND
-
 /** The amount of padding to add for scanning digits. */
 #define DIGIT_SCAN_PADDING (DIGIT_COUNT >> 1)
 
@@ -597,11 +554,7 @@ void GetDigits(
 	[unroll]
 	for ( KeyIndex = 0; KeyIndex < KEYS_PER_LOOP; ++KeyIndex )
 	{
-#if CONSTANT_BUFFER_WORKAROUND
-		Digits[KeyIndex] = (Keys[KeyIndex] >> RadixSortStruct.RadixShift) & DIGIT_MASK;
-#else
-		Digits[KeyIndex] = (Keys[KeyIndex] >> RadixSort.RadixShift) & DIGIT_MASK;
-#endif
+		Digits[KeyIndex] = (Keys[KeyIndex] >> RadixSortUB.RadixShift) & DIGIT_MASK;
 	}
 }
 
@@ -775,14 +728,6 @@ void RadixSort_Downsweep(
 	const uint ThreadId = GroupThreadId.x;
 	const uint GroupId = GroupIdXYZ.x;
 
-#if CONSTANT_BUFFER_WORKAROUND
-	RadixSortStruct.RadixShift = RadixSortParameterBuffer[0];
-	RadixSortStruct.TilesPerGroup = RadixSortParameterBuffer[1];
-	RadixSortStruct.ExtraTileCount = RadixSortParameterBuffer[2];
-	RadixSortStruct.ExtraKeyCount = RadixSortParameterBuffer[3];
-	RadixSortStruct.GroupCount = RadixSortParameterBuffer[4];
-#endif // #if CONSTANT_BUFFER_WORKAROUND
-
 	const uint BankOffset = ThreadId / WORDS_PER_BANK;
 	const uint CounterOffset = ThreadId & (WORDS_PER_BANK - 1);
 
@@ -813,15 +758,15 @@ void RadixSort_Downsweep(
 	}
 
 	// The number of tiles to process in this group.
-	if ( GroupId < RadixSortStruct.ExtraTileCount )
+	if ( GroupId < RadixSortUB.ExtraTileCount )
 	{
-		FirstTileIndex = GroupId * (RadixSortStruct.TilesPerGroup + 1);
-		TileCountForGroup = RadixSortStruct.TilesPerGroup + 1;
+		FirstTileIndex = GroupId * (RadixSortUB.TilesPerGroup + 1);
+		TileCountForGroup = RadixSortUB.TilesPerGroup + 1;
 	}
 	else
 	{
-		FirstTileIndex = GroupId * RadixSortStruct.TilesPerGroup + RadixSortStruct.ExtraTileCount;
-		TileCountForGroup = RadixSortStruct.TilesPerGroup;
+		FirstTileIndex = GroupId * RadixSortUB.TilesPerGroup + RadixSortUB.ExtraTileCount;
+		TileCountForGroup = RadixSortUB.TilesPerGroup;
 	}
 
 	// Key range for this group.
@@ -905,13 +850,13 @@ void RadixSort_Downsweep(
 	}
 
 	// The last group has to process any keys after the last tile.
-	if ( GroupId == (RadixSortStruct.GroupCount - 1) && RadixSortStruct.ExtraKeyCount > 0 )
+	if ( GroupId == (RadixSortUB.GroupCount - 1) && RadixSortUB.ExtraKeyCount > 0 )
 	{
 		// Load keys.
 		[unroll]
 		for ( KeyIndex = 0; KeyIndex < KEYS_PER_LOOP; ++KeyIndex )
 		{
-			if ( ThreadId * KEYS_PER_LOOP + KeyIndex < RadixSortStruct.ExtraKeyCount )
+			if ( ThreadId * KEYS_PER_LOOP + KeyIndex < RadixSortUB.ExtraKeyCount )
 			{
 				Keys[KeyIndex] = InKeys[GroupKeyBegin + ThreadId * KEYS_PER_LOOP + KeyIndex];
 			}
@@ -941,7 +886,7 @@ void RadixSort_Downsweep(
 		[unroll]
 		for ( KeyIndex = 0; KeyIndex < KEYS_PER_LOOP; ++KeyIndex )
 		{
-			if ( ThreadId * KEYS_PER_LOOP + KeyIndex < RadixSortStruct.ExtraKeyCount )
+			if ( ThreadId * KEYS_PER_LOOP + KeyIndex < RadixSortUB.ExtraKeyCount )
 			{
 				const uint Digit = Digits[KeyIndex];
 				const uint GlobalScatterIndex = LocalOffsets[Digit] + ThreadId * KEYS_PER_LOOP + KeyIndex - LocalTotals[Digit + DIGIT_SCAN_PADDING - 1];
@@ -953,7 +898,7 @@ void RadixSort_Downsweep(
 		[unroll]
 		for ( KeyIndex = 0; KeyIndex < KEYS_PER_LOOP; ++KeyIndex )
 		{
-			if ( ThreadId * KEYS_PER_LOOP + KeyIndex < RadixSortStruct.ExtraKeyCount )
+			if ( ThreadId * KEYS_PER_LOOP + KeyIndex < RadixSortUB.ExtraKeyCount )
 			{
 				Keys[KeyIndex] = InValues[GroupKeyBegin + ThreadId * KEYS_PER_LOOP + KeyIndex];
 			}
@@ -980,7 +925,7 @@ void RadixSort_Downsweep(
 		[unroll]
 		for ( KeyIndex = 0; KeyIndex < KEYS_PER_LOOP; ++KeyIndex )
 		{
-			if ( ThreadId * KEYS_PER_LOOP + KeyIndex < RadixSortStruct.ExtraKeyCount )
+			if ( ThreadId * KEYS_PER_LOOP + KeyIndex < RadixSortUB.ExtraKeyCount )
 			{
 				const uint Digit = Digits[KeyIndex];
 				const uint GlobalScatterIndex = LocalOffsets[Digit] + ThreadId * KEYS_PER_LOOP + KeyIndex - LocalTotals[Digit + DIGIT_SCAN_PADDING - 1];
diff --git a/Engine/Source/Runtime/Engine/Private/GPUSort.cpp b/Engine/Source/Runtime/Engine/Private/GPUSort.cpp
--- a/Engine/Source/Runtime/Engine/Private/GPUSort.cpp
+++ b/Engine/Source/Runtime/Engine/Private/GPUSort.cpp
@@ -55,7 +55,7 @@ BEGIN_UNIFORM_BUFFER_STRUCT( FRadixSortParameters, )
 	DECLARE_UNIFORM_BUFFER_STRUCT_MEMBER( uint32, GroupCount )
 END_UNIFORM_BUFFER_STRUCT( FRadixSortParameters )
 
-IMPLEMENT_UNIFORM_BUFFER_STRUCT(FRadixSortParameters,TEXT("RadixSort"));
+IMPLEMENT_UNIFORM_BUFFER_STRUCT(FRadixSortParameters,TEXT("RadixSortUB"));
 
 typedef TUniformBufferRef<FRadixSortParameters> FRadixSortUniformBufferRef;
 
