Fixed OpenGL shader compile error in PostProcessSubsurface (TTP #344724)

diff --git a/Engine/Shaders/PostProcessSubsurface.usf b/Engine/Shaders/PostProcessSubsurface.usf
--- a/Engine/Shaders/PostProcessSubsurface.usf
+++ b/Engine/Shaders/PostProcessSubsurface.usf
@@ -136,7 +136,7 @@ void SetupPS(in float4 UVAndScreenPos : TEXCOORD0, out float4 OutColor : SV_Targ
 		{
 			FScreenSpaceData ScreenSpaceData = GetScreenSpaceData(View.ViewRectMin.xy * View.ViewSizeAndSceneTexelSize.zw + IDAreaLocalUV * (View.ViewSizeAndSceneTexelSize.xy * View.ViewSizeAndSceneTexelSize.zw));
 
-			uint SubsurfaceProfileInt = ExtractSubsurfaceProfileInt(ScreenSpaceData.GBuffer);
+			int SubsurfaceProfileInt = ExtractSubsurfaceProfileInt(ScreenSpaceData.GBuffer);
 
 			OutColor = float4(0.5f, 0.5f, 0.5f, 0);
 
@@ -180,7 +180,7 @@ void SetupPS(in float4 UVAndScreenPos : TEXCOORD0, out float4 OutColor : SV_Targ
 					OutColor = float4(1, 1, 1, 0);
 				}
 
-				uint2 LeftTop = (PixelPos / 8) * 8;
+				int2 LeftTop = (PixelPos / 8) * 8;
 				PrintCharacter(PixelPos, OutColor.rgb, float3(1, 1, 1), LeftTop, SubsurfaceProfileInt);
 			}
 		}
