Increase sharpness in motion for Oculus.

diff --git a/Engine/Shaders/PostProcessTemporalCommon.usf b/Engine/Shaders/PostProcessTemporalCommon.usf
--- a/Engine/Shaders/PostProcessTemporalCommon.usf
+++ b/Engine/Shaders/PostProcessTemporalCommon.usf
@@ -410,7 +410,7 @@ which can be configured via defines for various temporal AA passes.
 	// Replace history with minimum difference of history from min or max neighborhood.
 	LumaHistory = min(abs(LumaMin-LumaHistory), abs(LumaMax-LumaHistory));
 	float HistoryAmount = (1.0/8.0) + HistoryBlur * (1.0/8.0);
-	float HistoryFactor = LumaHistory * HistoryAmount * (1.0 + HistoryBlur * HistoryAmount * 4.0);
+	float HistoryFactor = LumaHistory * HistoryAmount * (1.0 + HistoryBlur * HistoryAmount * 8.0);
 	float BlendFinal = saturate(HistoryFactor * rcp(LumaHistory + LumaContrast));
 	#if RESPONSIVE
 		// Responsive forces 1/4 of new frame.
