UE4 - Integrate
- Addes 3_1 define on shaders
- Fixes GL packed uniforms array size incorrect assert
- Disallows running the editor with -featureleveles2
- r.MobileHDR is now read-only
- Moved some checks from SM3 to SM4

diff --git a/Engine/Shaders/BRDF.usf b/Engine/Shaders/BRDF.usf
--- a/Engine/Shaders/BRDF.usf
+++ b/Engine/Shaders/BRDF.usf
@@ -260,7 +260,7 @@ half3 EnvBRDFApprox( half3 SpecularColor, half Roughness, half NoV )
 	half a004 = min( r.x * r.x, exp2( -9.28 * NoV ) ) * r.x + r.y;
 	half2 AB = half2( -1.04, 1.04 ) * a004 + r.zw;
 
-#if !ES2_PROFILE
+#if !(ES2_PROFILE || ES3_1_PROFILE)
 	// Anything less than 2% is physically impossible and is instead considered to be shadowing
 	// In ES2 this is skipped for performance as the impact can be small
 	// Note: this is needed for the 'specular' show flag to work, since it uses a SpecularColor of 0
diff --git a/Engine/Shaders/Common.usf b/Engine/Shaders/Common.usf
--- a/Engine/Shaders/Common.usf
+++ b/Engine/Shaders/Common.usf
@@ -81,8 +81,9 @@
 #define FEATURE_LEVEL_ES2	1
 #define FEATURE_LEVEL_ES3_1	2
 #define FEATURE_LEVEL_SM3	3
-#define FEATURE_LEVEL_SM4	4	 
+#define FEATURE_LEVEL_SM4	4
 #define FEATURE_LEVEL_SM5	5
+#define FEATURE_LEVEL_MAX	6
 
 // A8 textures when sampled either have their component in R or A. D3D uses A and that is the default.
 #define A8_SAMPLE_MASK .a
@@ -119,8 +120,8 @@
 #elif METAL_PROFILE
 	#define FEATURE_LEVEL FEATURE_LEVEL_ES2
 	#define FCOLOR_COMPONENT_SWIZZLE .bgra  
-#elif ES2_PROFILE
-
+#elif ES2_PROFILE || ES3_1_PROFILE
+	//@todo ES3_1 GL
 	#define FEATURE_LEVEL FEATURE_LEVEL_ES2
 
 	#if COMPILER_GLSL_ES2 
@@ -138,7 +139,7 @@
 
 	#error Add your platform here
 
-	#define FEATURE_LEVEL FEATURE_LEVEL_SM3
+	#define FEATURE_LEVEL FEATURE_LEVEL_MAX
 
 #endif
 
@@ -591,24 +592,24 @@ float ConvertToDeviceZ(float SceneDepth)
 /** Returns clip space W, which is world space distance along the View Z axis. */
 float CalcSceneDepth( float2 ScreenUV )
 {
-	#if FEATURE_LEVEL > FEATURE_LEVEL_ES2
-		return ConvertFromDeviceZ(Texture2DSampleLevel(SceneDepthTexture, SceneDepthTextureSampler, ScreenUV, 0).r);
+#if FEATURE_LEVEL > FEATURE_LEVEL_ES3_1
+	return ConvertFromDeviceZ(Texture2DSampleLevel(SceneDepthTexture, SceneDepthTextureSampler, ScreenUV, 0).r);
+#else
+	#if COMPILER_GLSL_ES2
+		#if IOS
+			// Only call FramebufferFetch when actually compiling for IOS ES2.
+			return FramebufferFetchES2().w;
+		#elif WEBGL
+			return Texture2DSampleLevel(SceneAlphaCopyTexture, SceneAlphaCopyTextureSampler, ScreenUV, 0).r;
+		#else 
+			return 65000.0f; 
+		#endif 
+	#elif METAL_PROFILE
+			return FramebufferFetchES2().w;
 	#else
-		#if COMPILER_GLSL_ES2
-			#if IOS
-				// Only call FramebufferFetch when actually compiling for IOS ES2.
-				return FramebufferFetchES2().w;
-			#elif WEBGL
-				return Texture2DSampleLevel(SceneAlphaCopyTexture, SceneAlphaCopyTextureSampler, ScreenUV, 0).r;
-			#else 
-				return 65000.0f; 
-			#endif 
-		#elif METAL_PROFILE
-				return FramebufferFetchES2().w;
-		#else
-			return ConvertFromDeviceZ(Texture2DSampleLevel(SceneDepthTexture, SceneDepthTextureSampler, ScreenUV, 0).r);
-		#endif
+		return ConvertFromDeviceZ(Texture2DSampleLevel(SceneDepthTexture, SceneDepthTextureSampler, ScreenUV, 0).r);
 	#endif
+#endif
 }
 
 #if FEATURE_LEVEL >= FEATURE_LEVEL_SM4
diff --git a/Engine/Shaders/Definitions.usf b/Engine/Shaders/Definitions.usf
--- a/Engine/Shaders/Definitions.usf
+++ b/Engine/Shaders/Definitions.usf
@@ -88,6 +88,10 @@
 #define ES2_PROFILE										0
 #endif
 
+#ifndef ES3_1_PROFILE
+#define ES3_1_PROFILE									0
+#endif
+
 #ifndef METAL_PROFILE
 #define METAL_PROFILE									0
 #endif
diff --git a/Engine/Shaders/FilterPixelShader.usf b/Engine/Shaders/FilterPixelShader.usf
--- a/Engine/Shaders/FilterPixelShader.usf
+++ b/Engine/Shaders/FilterPixelShader.usf
@@ -50,7 +50,7 @@ void Combine(inout float4 Dest, float4 Src, float4 Weight)
 #endif
 }
 
-#if ES2_PROFILE
+#if (ES2_PROFILE || ES3_1_PROFILE)
 
 void Main(
 	float2 InUV : TEXCOORD0,
diff --git a/Engine/Shaders/FilterVertexShader.usf b/Engine/Shaders/FilterVertexShader.usf
--- a/Engine/Shaders/FilterVertexShader.usf
+++ b/Engine/Shaders/FilterVertexShader.usf
@@ -8,7 +8,7 @@
 
 float4 SampleOffsets[(NUM_SAMPLES + 1) / 2];
 
-#if ES2_PROFILE
+#if (ES2_PROFILE || ES3_1_PROFILE)
 
 void Main(
 	float4 InPosition : ATTRIBUTE0,
diff --git a/Engine/Shaders/LandscapeVertexFactory.usf b/Engine/Shaders/LandscapeVertexFactory.usf
--- a/Engine/Shaders/LandscapeVertexFactory.usf
+++ b/Engine/Shaders/LandscapeVertexFactory.usf
@@ -37,7 +37,7 @@ float4 NeighborSectionLod[4];
 
 float4x4 LocalToWorldNoScaling;
 
-#if FEATURE_LEVEL > FEATURE_LEVEL_ES2
+#if FEATURE_LEVEL >= FEATURE_LEVEL_SM4
 Texture2D HeightmapTexture;
 SamplerState HeightmapTextureSampler;
 #endif
@@ -54,7 +54,7 @@ SamplerState XYOffsetmapTextureSampler;
 
 struct FVertexFactoryInput
 {
-#if FEATURE_LEVEL > FEATURE_LEVEL_ES2
+#if FEATURE_LEVEL >= FEATURE_LEVEL_SM4
 	float4 Position: ATTRIBUTE0;
 #else
 	float4 PackedPosition: ATTRIBUTE0; 
@@ -65,7 +65,7 @@ struct FVertexFactoryInput
 struct FVertexFactoryInterpolantsVSToPS
 {
 	float2	LayerTexCoord			: TEXCOORD0;
-#if ES2_PROFILE
+#if (ES2_PROFILE || ES3_1_PROFILE)
 	float2	WeightMapTexCoord		: TEXCOORD1;
 	float4	TransformedTexCoord0	: TEXCOORD2;
 	float4  TransformedTexCoord1	: TEXCOORD3;
@@ -74,7 +74,7 @@ struct FVertexFactoryInterpolantsVSToPS
 	float4	TransformedTexCoords	: TEXCOORD2;
 #endif
 #if NEEDS_LIGHTMAP_COORDINATE
-#if ES2_PROFILE
+#if (ES2_PROFILE || ES3_1_PROFILE)
 	float2	LightMapCoordinate[2]	: TEXCOORD4;
 	float2	ShadowMapCoordinate		: TEXCOORD6;
 #else
@@ -87,7 +87,7 @@ struct FLandscapeTexCoords
 {
 	float2	LayerTexCoord;
 	float2	WeightMapTexCoord;
-#if !ES2_PROFILE
+#if !(ES2_PROFILE || ES3_1_PROFILE)
 	float2	HeightMapTexCoord;
 #endif
 #if NEEDS_LIGHTMAP_COORDINATE
@@ -110,7 +110,7 @@ struct FVertexFactoryInterpolantsVSToDS
 #if USING_TESSELLATION
 float2 VertexFactoryGetTextureCoordinateDS( FVertexFactoryInterpolantsVSToDS Interpolants )
 {
-#if ES2_PROFILE
+#if (ES2_PROFILE || ES3_1_PROFILE)
 	return Interpolants.InterpolantsVSToPS.WeightMapTexCoord.xy;
 #else
 	return Interpolants.InterpolantsVSToPS.WeightHeightMapTexCoord.zw;
@@ -146,7 +146,7 @@ FLandscapeTexCoords GetLandscapeTexCoords(FVertexFactoryInput Input, FVertexFact
 	Result.LayerTexCoord = Intermediates.LocalPosition.xy + LandscapeParameters.SubsectionSizeVertsLayerUVPan.zw + Intermediates.InputPosition.zw * LandscapeParameters.SubsectionOffsetParams.ww;
 
 	Result.WeightMapTexCoord	= Intermediates.LocalPosition.xy * LandscapeParameters.WeightmapUVScaleBias.xy + LandscapeParameters.WeightmapUVScaleBias.zw + Intermediates.InputPosition.zw * LandscapeParameters.SubsectionOffsetParams.zz;
-#if !ES2_PROFILE
+#if !(ES2_PROFILE || ES3_1_PROFILE)
 	Result.HeightMapTexCoord	= Intermediates.LocalPosition.xy * LandscapeParameters.HeightmapUVScaleBias.xy + LandscapeParameters.HeightmapUVScaleBias.zw + 0.5*LandscapeParameters.HeightmapUVScaleBias.xy + Intermediates.InputPosition.zw * LandscapeParameters.SubsectionOffsetParams.xy;
 #endif
 
@@ -172,7 +172,7 @@ float3x3 VertexFactoryGetPerPixelTangentBasis(FVertexFactoryInterpolantsVSToPS I
 {
 	float3x3 Result;
 #if PIXELSHADER
-#if FEATURE_LEVEL > FEATURE_LEVEL_ES2
+#if FEATURE_LEVEL >= FEATURE_LEVEL_SM4
 	float4 SampleValue = Texture2DSample(NormalmapTexture, NormalmapTextureSampler, Interpolants.WeightHeightMapTexCoord.zw);
 #else
 	float4 SampleValue = Texture2DSample(NormalmapTexture, NormalmapTextureSampler, Interpolants.WeightMapTexCoord);
@@ -180,7 +180,7 @@ float3x3 VertexFactoryGetPerPixelTangentBasis(FVertexFactoryInterpolantsVSToPS I
 #if WEBGL
 	float2 SampleNormal = float2(SampleValue.r, SampleValue.a) * float2(2.0,2.0) - float2(1.0,1.0);
 #else
-    float2 SampleNormal = float2(SampleValue.b, SampleValue.a) * float2(2.0,2.0) - float2(1.0,1.0);
+	float2 SampleNormal = float2(SampleValue.b, SampleValue.a) * float2(2.0,2.0) - float2(1.0,1.0);
 #endif
 	float3 WorldNormal = float3( SampleNormal, sqrt(max(1.0-dot(SampleNormal,SampleNormal),0.0)) );
 
@@ -197,26 +197,26 @@ FMaterialTessellationParameters GetMaterialTessellationParameters(FVertexFactory
 #if NUM_MATERIAL_TEXCOORDS
 	Result.TexCoords[0] = Interpolants.InterpolantsVSToPS.LayerTexCoord.xy;
 #if NUM_MATERIAL_TEXCOORDS > 1
-	#if ES2_PROFILE
+	#if (ES2_PROFILE || ES3_1_PROFILE)
 		Result.TexCoords[1] = Interpolants.InterpolantsVSToPS.TransformedTexCoord0.xy;
 	#else
 	Result.TexCoords[1] = Interpolants.InterpolantsVSToPS.TransformedTexCoords.xy;
 	#endif
 #if NUM_MATERIAL_TEXCOORDS > 2
-	#if ES2_PROFILE
+	#if (ES2_PROFILE || ES3_1_PROFILE)
 		Result.TexCoords[2] = Interpolants.InterpolantsVSToPS.TransformedTexCoord1.xy;
 	#else
 		Result.TexCoords[2] = Interpolants.InterpolantsVSToPS.TransformedTexCoords.zw;
 	#endif
 #if NUM_MATERIAL_TEXCOORDS > 3
-	#if ES2_PROFILE
+	#if (ES2_PROFILE || ES3_1_PROFILE)
 		Result.TexCoords[3] = Interpolants.InterpolantsVSToPS.WeightMapTexCoord;
 	#else
 	Result.TexCoords[3] = Interpolants.InterpolantsVSToPS.WeightHeightMapTexCoord.xy;
 	#endif
 #if NUM_MATERIAL_TEXCOORDS > 4
 	#if NEEDS_LIGHTMAP_COORDINATE
-		#if ES2_PROFILE
+		#if (ES2_PROFILE || ES3_1_PROFILE)
 			Result.TexCoords[4] = Interpolants.InterpolantsVSToPS.LightMapCoordinate[0].xy;
 		#else
 		Result.TexCoords[4] = Interpolants.InterpolantsVSToPS.LightMapCoordinate.xy;
@@ -225,7 +225,7 @@ FMaterialTessellationParameters GetMaterialTessellationParameters(FVertexFactory
 		Result.TexCoords[4] = float2(0,0);
 	#endif
 #if NUM_MATERIAL_TEXCOORDS > 5
-	#if ES2_PROFILE
+	#if (ES2_PROFILE || ES3_1_PROFILE)
 		Result.TexCoords[5] = float2(0, 0);
 	#else
 		Result.TexCoords[5] = Interpolants.InterpolantsVSToPS.WeightHeightMapTexCoord.zw;
@@ -263,7 +263,7 @@ FMaterialPixelParameters GetMaterialPixelParameters(FVertexFactoryInterpolantsVS
 	FMaterialPixelParameters Result = MakeInitializedMaterialPixelParameters();
 
 #if NEEDS_LIGHTMAP_COORDINATE
-	#if ES2_PROFILE
+	#if (ES2_PROFILE || ES3_1_PROFILE)
 		// Not supported in pixel shader
 		float2 LightmapUVs = float2(0, 0);
 	#else
@@ -276,19 +276,19 @@ FMaterialPixelParameters GetMaterialPixelParameters(FVertexFactoryInterpolantsVS
 #if NUM_MATERIAL_TEXCOORDS
 	Result.TexCoords[0] = Interpolants.LayerTexCoord.xy;				// XY layer
 #if NUM_MATERIAL_TEXCOORDS > 1
-	#if ES2_PROFILE
+	#if (ES2_PROFILE || ES3_1_PROFILE)
 		Result.TexCoords[1] = Interpolants.TransformedTexCoord0.xy; // VS calcualted TexCoord 1, default is XZ layer
 	#else
 		Result.TexCoords[1] = Interpolants.TransformedTexCoords.xy; // VS calcualted TexCoord 1, default is XZ layer
-	#endif // ES2_PROFILE
+	#endif // (ES2_PROFILE || ES3_1_PROFILE)
 #if NUM_MATERIAL_TEXCOORDS > 2
-	#if ES2_PROFILE
+	#if (ES2_PROFILE || ES3_1_PROFILE)
 		Result.TexCoords[2] = Interpolants.TransformedTexCoord1.xy; // VS calcualted TexCoord 2, default is YZ layer
 	#else
 		Result.TexCoords[2] = Interpolants.TransformedTexCoords.zw;	// VS calcualted TexCoord 2, default is YZ layer
 	#endif
 #if NUM_MATERIAL_TEXCOORDS > 3
-	#if ES2_PROFILE
+	#if (ES2_PROFILE || ES3_1_PROFILE)
 		Result.TexCoords[3] = Interpolants.WeightMapTexCoord;		// Weightmap
 	#else
 		Result.TexCoords[3] = Interpolants.WeightHeightMapTexCoord.xy;		// Weightmap
@@ -300,7 +300,7 @@ FMaterialPixelParameters GetMaterialPixelParameters(FVertexFactoryInterpolantsVS
 		Result.TexCoords[4] = float2(0,0);
 	#endif	//#if NEEDS_LIGHTMAP_COORDINATE
 #if NUM_MATERIAL_TEXCOORDS > 5
-	#if ES2_PROFILE
+	#if (ES2_PROFILE || ES3_1_PROFILE)
 		Result.TexCoords[5] = float2(0,0); // Height map UV not supported
 	#else
 	Result.TexCoords[5] = Interpolants.WeightHeightMapTexCoord.zw;	// Height map
@@ -329,7 +329,7 @@ FMaterialPixelParameters GetMaterialPixelParameters(FVertexFactoryInterpolantsVS
 	Result.LightmapUVs = LightmapUVs;
 #endif
 
-#if ES2_PROFILE && PIXELSHADER
+#if (ES2_PROFILE || ES3_1_PROFILE) && PIXELSHADER
 
 #if WEBGL
 	Result.LayerWeights.xy = Texture2DSample(NormalmapTexture, NormalmapTextureSampler, Interpolants.WeightMapTexCoord).bg;
@@ -357,7 +357,7 @@ FMaterialVertexParameters GetMaterialVertexParameters( FVertexFactoryInput Input
 
 #if NUM_MATERIAL_TEXCOORDS_VERTEX
 	Result.TexCoords[0] = LandscapeTexCoords.LayerTexCoord.xy;
-#if ES2_PROFILE
+#if (ES2_PROFILE || ES3_1_PROFILE)
 	Result.TexCoordOffset = LodBias.zw;
 #endif
 #if NUM_MATERIAL_TEXCOORDS_VERTEX > 1
@@ -373,7 +373,7 @@ FMaterialVertexParameters GetMaterialVertexParameters( FVertexFactoryInput Input
 		Result.TexCoords[4] = float2(0,0);
 	#endif
 #if NUM_MATERIAL_TEXCOORDS_VERTEX > 5
-	#if ES2_PROFILE
+	#if (ES2_PROFILE || ES3_1_PROFILE)
 		Result.TexCoords[5] = float2(0,0);
 	#else
 		Result.TexCoords[5] = LandscapeTexCoords.HeightMapTexCoord;
@@ -398,7 +398,7 @@ FMaterialVertexParameters GetMaterialVertexParameters( FVertexFactoryInput Input
 #if NEEDS_LIGHTMAP_COORDINATE
 void GetLightMapCoordinates(FVertexFactoryInterpolantsVSToPS Interpolants, out float2 LightmapUV0, out float2 LightmapUV1)
 {
-#if ES2_PROFILE
+#if (ES2_PROFILE || ES3_1_PROFILE)
 	LightmapUV0 = Interpolants.LightMapCoordinate[0].xy;
 	LightmapUV1 = Interpolants.LightMapCoordinate[1].xy;
 #else
@@ -409,7 +409,7 @@ void GetLightMapCoordinates(FVertexFactoryInterpolantsVSToPS Interpolants, out f
 
 float2 GetShadowMapCoordinate(FVertexFactoryInterpolantsVSToPS Interpolants)
 {
-#if ES2_PROFILE
+#if (ES2_PROFILE || ES3_1_PROFILE)
 	return Interpolants.ShadowMapCoordinate.xy;
 #else
 	return Interpolants.LightMapCoordinate.zw;
@@ -421,7 +421,7 @@ FVertexFactoryIntermediates GetVertexFactoryIntermediates(FVertexFactoryInput In
 {
 	FVertexFactoryIntermediates Intermediates;
 
-#if FEATURE_LEVEL > FEATURE_LEVEL_ES2
+#if FEATURE_LEVEL >= FEATURE_LEVEL_SM4
 	Intermediates.InputPosition = Input.Position;
 #else
 	Intermediates.InputPosition = Input.PackedPosition * 255.f;
@@ -512,7 +512,7 @@ FVertexFactoryIntermediates GetVertexFactoryIntermediates(FVertexFactoryInput In
 	float2 CoordTranslate = float2( LandscapeParameters.SubsectionSizeVertsLayerUVPan.x * pow(2, -LodValue) - 1, max(LandscapeParameters.SubsectionSizeVertsLayerUVPan.x * pow(2, -LodValue-1), 2) - 1 ) * LandscapeParameters.SubsectionSizeVertsLayerUVPan.y;
 
 	// Position for current LOD in base LOD units
-#if FEATURE_LEVEL > FEATURE_LEVEL_ES2
+#if FEATURE_LEVEL >= FEATURE_LEVEL_SM4
 	float2 InputPositionLODBase = Intermediates.InputPosition.xy * LodValues.y;
 #else
 	float2 InputPositionLODBase = Intermediates.InputPosition.xy;
@@ -527,7 +527,7 @@ FVertexFactoryIntermediates GetVertexFactoryIntermediates(FVertexFactoryInput In
 	float2 InputPositionNextLOD = NextLODCoordsInt / CoordTranslate.y;
 
 	// Get the height and normal XY for current and next LOD out of the textures
-#if FEATURE_LEVEL > FEATURE_LEVEL_ES2
+#if FEATURE_LEVEL >= FEATURE_LEVEL_SM4
 	float2 SampleCoords = InputPositionLODAdjusted * LandscapeParameters.HeightmapUVScaleBias.xy + LandscapeParameters.HeightmapUVScaleBias.zw + 0.5*LandscapeParameters.HeightmapUVScaleBias.xy + Intermediates.InputPosition.zw * LandscapeParameters.SubsectionOffsetParams.xy;
 	float4 SampleValue = Texture2DSampleLevel(HeightmapTexture, HeightmapTextureSampler, SampleCoords, LodValue-LodBias.z);
 	float Height = ((SampleValue.r * 255.0 * 256.0 + SampleValue.g * 255.0) - 32768.0) * TERRAIN_ZSCALE;
@@ -632,7 +632,7 @@ FVertexFactoryInterpolantsVSToPS VertexFactoryGetInterpolantsVSToPS(FVertexFacto
 	Interpolants = (FVertexFactoryInterpolantsVSToPS)0;
 	FLandscapeTexCoords LandscapeTexCoords = GetLandscapeTexCoords(Input, Intermediates);
 
-#if ES2_PROFILE
+#if (ES2_PROFILE || ES3_1_PROFILE)
 	Interpolants.LayerTexCoord	= LandscapeTexCoords.LayerTexCoord;
 	Interpolants.WeightMapTexCoord	= LandscapeTexCoords.WeightMapTexCoord; 
 	#if NUM_MATERIAL_TEXCOORDS
@@ -667,10 +667,10 @@ FVertexFactoryInterpolantsVSToPS VertexFactoryGetInterpolantsVSToPS(FVertexFacto
 		#endif // 1
 
 	#endif // NUM_MATERIAL_TEXCOORDS
-#endif // ES2_PROFILE
+#endif // (ES2_PROFILE || ES3_1_PROFILE)
 
 #if NEEDS_LIGHTMAP_COORDINATE
-	#if ES2_PROFILE
+	#if (ES2_PROFILE || ES3_1_PROFILE)
 		Interpolants.LightMapCoordinate[0].xy = LandscapeTexCoords.LightMapCoordinate * LightMapCoordinateScaleBias.xy + LightMapCoordinateScaleBias.zw;
 		Interpolants.LightMapCoordinate[0].y *= 0.5;
 		Interpolants.LightMapCoordinate[1].xy = Interpolants.LightMapCoordinate[0].xy;
@@ -722,12 +722,12 @@ FVertexFactoryInterpolantsVSToDS VertexFactoryInterpolate(FVertexFactoryInterpol
 {
 	FVertexFactoryInterpolantsVSToDS O;
 	
-#if NEEDS_LIGHTMAP_COORDINATE && !ES2_PROFILE
+#if NEEDS_LIGHTMAP_COORDINATE && !(ES2_PROFILE || ES3_1_PROFILE)
 	TESSELLATION_INTERPOLATE_MEMBER(InterpolantsVSToPS.LightMapCoordinate);
 #endif
 
 	TESSELLATION_INTERPOLATE_MEMBER(InterpolantsVSToPS.LayerTexCoord);
-#if ES2_PROFILE
+#if (ES2_PROFILE || ES3_1_PROFILE)
 	TESSELLATION_INTERPOLATE_MEMBER(InterpolantsVSToPS.WeightMapTexCoord);
 	TESSELLATION_INTERPOLATE_MEMBER(InterpolantsVSToPS.TransformedTexCoord0);
 	TESSELLATION_INTERPOLATE_MEMBER(InterpolantsVSToPS.TransformedTexCoord1);
diff --git a/Engine/Shaders/LocalVertexFactory.usf b/Engine/Shaders/LocalVertexFactory.usf
--- a/Engine/Shaders/LocalVertexFactory.usf
+++ b/Engine/Shaders/LocalVertexFactory.usf
@@ -204,7 +204,7 @@ FMaterialPixelParameters GetMaterialPixelParameters(FVertexFactoryInterpolantsVS
 
 #if LIGHTMAP_UV_ACCESS
 #if NEEDS_LIGHTMAP_COORDINATE
-	#if ES2_PROFILE
+	#if (ES2_PROFILE || ES3_1_PROFILE)
 		// Not supported in pixel shader
 		Result.LightmapUVs = float2(0, 0);
 	#else
@@ -507,7 +507,7 @@ FVertexFactoryInterpolantsVSToPS VertexFactoryGetInterpolantsVSToPS(FVertexFacto
 	#if USE_INSTANCING
 	//float2 InstanceLightMapBias = float2(Input.InstanceInverseTransform1.w, Input.InstanceInverseTransform2.w);
 	//Interpolants.LightMapCoordinate.xy = Input.LightMapCoordinate * LightMapCoordinateScaleBias.xy + InstanceLightMapBias;
-		#if ES2_PROFILE
+	#if (ES2_PROFILE || ES3_1_PROFILE)
 	SetLightMapCoordinate(Interpolants, (float4)0, (float2)0);
 		#else
 	Interpolants.LightMapCoordinate = 0;
@@ -517,7 +517,7 @@ FVertexFactoryInterpolantsVSToPS VertexFactoryGetInterpolantsVSToPS(FVertexFacto
 	float4 LightMapCoordinate = 0;
 	float2 ShadowMapCoordinate = 0;
 
-		#if ES2_PROFILE
+	#if (ES2_PROFILE || ES3_1_PROFILE)
 		LightMapCoordinate.xy = Input.LightMapCoordinate * LightMapCoordinateScaleBias.xy + LightMapCoordinateScaleBias.zw;
 		LightMapCoordinate.y *= 0.5;
 		LightMapCoordinate.zw = LightMapCoordinate.xy;
diff --git a/Engine/Shaders/LocalVertexFactoryCommon.usf b/Engine/Shaders/LocalVertexFactoryCommon.usf
--- a/Engine/Shaders/LocalVertexFactoryCommon.usf
+++ b/Engine/Shaders/LocalVertexFactoryCommon.usf
@@ -6,7 +6,7 @@
 
 // FVertexFactoryInterpolantsVSToPS is split into a PC and ES2 version, since the ES2 version has more restrictions on packing
 // Then accessors are provided for each property which try to encapsulate the packing differences
-#if FEATURE_LEVEL >= FEATURE_LEVEL_SM3
+#if FEATURE_LEVEL >= FEATURE_LEVEL_SM4
 
 struct FVertexFactoryInterpolantsVSToPS
 {
@@ -105,7 +105,7 @@ void SetTangents(inout FVertexFactoryInterpolantsVSToPS Interpolants, float3 InT
 }
 
 
-#else // #if FEATURE_LEVEL >= FEATURE_LEVEL_SM3
+#else // #if FEATURE_LEVEL >= FEATURE_LEVEL_SM4
 
 #define NUM_PACKED_UVS (NUM_MATERIAL_TEXCOORDS > 1 ? 3 : 1)
 
diff --git a/Engine/Shaders/MaterialTemplate.usf b/Engine/Shaders/MaterialTemplate.usf
--- a/Engine/Shaders/MaterialTemplate.usf
+++ b/Engine/Shaders/MaterialTemplate.usf
@@ -138,7 +138,7 @@ struct FMaterialPixelParameters
 	/** Per-particle properties. Only valid for particle vertex factories. */
 	FMaterialParticleParameters Particle;
 
-#if ES2_PROFILE
+#if (ES2_PROFILE || ES3_1_PROFILE)
 	float4 LayerWeights;
 #endif
 };
@@ -146,7 +146,7 @@ struct FMaterialPixelParameters
 // @todo compat hack
 FMaterialPixelParameters MakeInitializedMaterialPixelParameters()
 {
-#if SM5_PROFILE || SM4_PROFILE || PS4_PROFILE || ES2_PROFILE
+#if SM5_PROFILE || SM4_PROFILE || PS4_PROFILE || (ES2_PROFILE || ES3_1_PROFILE)
 	return (FMaterialPixelParameters)0;
 #else
 	// @todo compat hack - put back to the 0 typecast (2 lines up)
@@ -238,7 +238,7 @@ struct FMaterialVertexParameters
 	half4 VertexColor;
 #if NUM_MATERIAL_TEXCOORDS_VERTEX
 	float2 TexCoords[NUM_MATERIAL_TEXCOORDS_VERTEX];
-#if ES2_PROFILE
+#if (ES2_PROFILE || ES3_1_PROFILE)
 	float2 TexCoordOffset; // Offset for UV localization for large UV values
 #endif
 #endif
@@ -541,11 +541,11 @@ MaterialFloat2 GetParticleMacroUV(FMaterialPixelParameters Parameters)
 	return (Parameters.ScreenPosition.xy / Parameters.ScreenPosition.w - Parameters.Particle.MacroUV.xy) * Parameters.Particle.MacroUV.zw + MaterialFloat2(.5, .5);
 }
 
-#define ES2_EMULATION (!COMPILER_GLSL_ES2 && COMPILE_SHADERS_FOR_DEVELOPMENT)
+#define ES2_EMULATION (!COMPILER_GLSL_ES2 && COMPILE_SHADERS_FOR_DEVELOPMENT && !METAL_PROFILE)
 
 MaterialFloat4 ProcessMaterialColorTextureLookup(MaterialFloat4 TextureValue)
 {
-#if ES2_PROFILE
+#if (ES2_PROFILE || ES3_1_PROFILE)
 	#if ES2_EMULATION
 	if( View.ES2PreviewMode > 0.5f)
 	{
@@ -567,7 +567,7 @@ MaterialFloat4 ProcessMaterialLinearColorTextureLookup(MaterialFloat4 TextureVal
 
 MaterialFloat ProcessMaterialGreyscaleTextureLookup(MaterialFloat TextureValue)
 {
-#if ES2_PROFILE
+#if (ES2_PROFILE || ES3_1_PROFILE)
 	#if ES2_EMULATION
 	if( View.ES2PreviewMode > 0.5f )
 	{
@@ -983,7 +983,7 @@ half3 GetMaterialNormal(FMaterialPixelParameters Parameters)
 	RetNormal = GetMaterialNormalRaw(Parameters);
 	
 	// todo: COMPILE_SHADERS_FOR_DEVELOPMENT is unfinished feature, using XBOXONE_PROFILE as workaround
-	#if COMPILE_SHADERS_FOR_DEVELOPMENT == 1 && !XBOXONE_PROFILE && !ES2_PROFILE
+	#if COMPILE_SHADERS_FOR_DEVELOPMENT == 1 && !XBOXONE_PROFILE && !(ES2_PROFILE || ES3_1_PROFILE)
 	{
 		// this feature is only needed for development/editor - we can compile it out for a shipping build (see r.CompileShadersForDevelopment cvar help)
 		half3 OverrideNormal = View.NormalOverrideParameter.xyz;
@@ -1041,7 +1041,7 @@ half GetMaterialRoughnessRaw(FMaterialPixelParameters Parameters)
 
 half GetMaterialRoughness(FMaterialPixelParameters Parameters)
 {
-#if ES2_PROFILE
+#if (ES2_PROFILE || ES3_1_PROFILE)
 	// The smallest normalized value that can be represented in IEEE 754 (FP16) is 2^-14 = 6.1e-5.
 	// The code will make the following computation involving roughness: 1.0 / Roughness^4.
 	// Therefore to prevent division by zero on devices that do not support denormals, Roughness^4
@@ -1306,7 +1306,7 @@ void CalcMaterialParameters(
 	// flip the normal for backfaces being rendered with a two-sided material
 	Parameters.TangentNormal *= Parameters.TwoSidedSign;
 
-#if FEATURE_LEVEL != FEATURE_LEVEL_ES2
+#if FEATURE_LEVEL >= FEATURE_LEVEL_SM4
 	// ES2 will rely on only the final normalize for performance
 	Parameters.TangentNormal = normalize(Parameters.TangentNormal);
 #endif
diff --git a/Engine/Shaders/MeshParticleVertexFactory.usf b/Engine/Shaders/MeshParticleVertexFactory.usf
--- a/Engine/Shaders/MeshParticleVertexFactory.usf
+++ b/Engine/Shaders/MeshParticleVertexFactory.usf
@@ -65,7 +65,7 @@ struct FVertexFactoryInterpolantsVSToPS
 	float4 SubUV1AndLerp : TEXCOORD2;
 #else
 	#if NUM_MATERIAL_TEXCOORDS
-			#if ES2_PROFILE
+			#if (ES2_PROFILE || ES3_1_PROFILE)
 				// Avoid dependent texture fetches, put all UVs in xy
 				float2	TexCoords[NUM_MATERIAL_TEXCOORDS] : TEXCOORD0;
 			#else
@@ -148,7 +148,7 @@ FMaterialPixelParameters GetMaterialPixelParameters(FVertexFactoryInterpolantsVS
 	Result.Particle.SubUVCoords[1] = Interpolants.SubUV1AndLerp.xy;
 	Result.Particle.SubUVLerp = Interpolants.SubUV1AndLerp.z;
 #elif NUM_MATERIAL_TEXCOORDS
-	#if ES2_PROFILE
+	#if (ES2_PROFILE || ES3_1_PROFILE)
 		UNROLL
 		for( int CoordinateIndex = 0; CoordinateIndex < NUM_MATERIAL_TEXCOORDS; CoordinateIndex++ )
 		{
@@ -382,7 +382,7 @@ FVertexFactoryInterpolantsVSToPS VertexFactoryGetInterpolantsVSToPS(FVertexFacto
 	FVertexFactoryInterpolantsVSToPS Interpolants;
 
 #if USE_PARTICLE_SUBUVS
-    Interpolants.SubUV0AndTexCoord0.xy = Intermediates.SubUVCoords.xy;
+	Interpolants.SubUV0AndTexCoord0.xy = Intermediates.SubUVCoords.xy;
 	Interpolants.SubUV1AndLerp.xy = Intermediates.SubUVCoords.zw;
 	Interpolants.SubUV1AndLerp.zw = Intermediates.SubUVLerp.xx;
 
@@ -394,7 +394,7 @@ FVertexFactoryInterpolantsVSToPS VertexFactoryGetInterpolantsVSToPS(FVertexFacto
 		Interpolants.SubUV0AndTexCoord0.zw = 0;
 	#endif
 #elif NUM_MATERIAL_TEXCOORDS
-	#if ES2_PROFILE
+	#if (ES2_PROFILE || ES3_1_PROFILE)
 		GetMaterialCustomizedUVs(VertexParameters, Interpolants.TexCoords);
 	#else
 		// Ensure the unused components of the last packed texture coordinate are initialized.
diff --git a/Engine/Shaders/ParticleSpriteVertexFactory.usf b/Engine/Shaders/ParticleSpriteVertexFactory.usf
--- a/Engine/Shaders/ParticleSpriteVertexFactory.usf
+++ b/Engine/Shaders/ParticleSpriteVertexFactory.usf
@@ -16,7 +16,7 @@
 #define USE_PARTICLE_SIZE (FEATURE_LEVEL >= FEATURE_LEVEL_SM5 && NEEDS_PARTICLE_SIZE)
 
 // We don't compute a tangent basis in ES2 to save shader instructions.
-#define NEEDS_TANGENT_BASIS (FEATURE_LEVEL > FEATURE_LEVEL_ES2)
+#define NEEDS_TANGENT_BASIS (FEATURE_LEVEL >= FEATURE_LEVEL_SM4)
 
 struct FVertexFactoryInput
 {
diff --git a/Engine/Shaders/ReflectionEnvironmentShared.usf b/Engine/Shaders/ReflectionEnvironmentShared.usf
--- a/Engine/Shaders/ReflectionEnvironmentShared.usf
+++ b/Engine/Shaders/ReflectionEnvironmentShared.usf
@@ -76,3 +76,21 @@ float3 GetSkySHDiffuse(float3 Normal)
 	// max to not get negative colors, maybe we can optimize to use saturate() and scale later or clamp this way: saturate(A/100)*100
 	return max(0, Intermediate0 + Intermediate1 + Intermediate2);
 }
+
+/**
+* Computes sky diffuse lighting from the SH irradiance map.
+* This has the SH basis evaluation and diffuse convolusion weights combined for minimal ALU's - see "Stupid Spherical Harmonics (SH) Tricks"
+* Only does the first 3 components for speed.
+*/
+float3 GetSkySHDiffuseSimple(float3 Normal)
+{
+	float4 NormalVector = float4(Normal, 1);
+
+	float3 Intermediate0;
+	Intermediate0.x = dot(View.SkyIrradianceEnvironmentMap[0], NormalVector);
+	Intermediate0.y = dot(View.SkyIrradianceEnvironmentMap[1], NormalVector);
+	Intermediate0.z = dot(View.SkyIrradianceEnvironmentMap[2], NormalVector);		
+
+	// max to not get negative colors, maybe we can optimize to use saturate() and scale later or clamp this way: saturate(A/100)*100
+	return max(0, Intermediate0);
+}
diff --git a/Engine/Shaders/SimpleElementPixelShader.usf b/Engine/Shaders/SimpleElementPixelShader.usf
--- a/Engine/Shaders/SimpleElementPixelShader.usf
+++ b/Engine/Shaders/SimpleElementPixelShader.usf
@@ -196,7 +196,7 @@ void GammaDistanceFieldMain(
 	float4 BaseColor = float4(Color.rgb, smoothstep(0.5-SmoothstepWidth, 0.5+SmoothstepWidth, BaseDist));
 	
 // @todo-mobile: bool doesn't work on ES2
-#if !ES2_PROFILE
+#if !(ES2_PROFILE || ES3_1_PROFILE)
 	if( EnableShadow )
 	{
 		// sample texture with shadow direction offset
@@ -229,7 +229,7 @@ void GammaDistanceFieldMain(
 		}
 	}
  #endif
- 	
+	
 	clip(BaseColor.a - ClipRef);
 	OutColor = BaseColor;
 	OutColor.a *= Color.a;
diff --git a/Engine/Shaders/VisualizeTexture.usf b/Engine/Shaders/VisualizeTexture.usf
--- a/Engine/Shaders/VisualizeTexture.usf
+++ b/Engine/Shaders/VisualizeTexture.usf
@@ -93,15 +93,15 @@ void VisualizeTexturePS(
 	int InputValueMapping = (int)VisualizeParam[2].r;
 	if(InputValueMapping == 1)
 	{
-#if FEATURE_LEVEL > FEATURE_LEVEL_ES2
+#if FEATURE_LEVEL >= FEATURE_LEVEL_SM4
 		// current z buffer is negated (1 is near, 0 is far)
 		float Depth = 1 - TexLookup.r;
 
 		// map depth to nice colors
 		float Scale = 64.0; // This 64 looks like a hardware limit.
-        Depth = (exp2(Depth * Scale) - 1.0) * (1.0 / (exp2(Scale) - 1.0));
+		Depth = (exp2(Depth * Scale) - 1.0) * (1.0 / (exp2(Scale) - 1.0));
+		Depth *= Depth;
 		Depth *= Depth;
-        Depth *= Depth;
 
 		float Value = Depth * Depth;
 		TexLookup = float4(pow(Value, 0.25f), pow(Value, 1), pow(Value, 4), 0);		
@@ -130,7 +130,7 @@ void VisualizeTexturePS(
 	}
 	
 	// blink blue if not finite
-#if FEATURE_LEVEL > FEATURE_LEVEL_ES2
+#if FEATURE_LEVEL >= FEATURE_LEVEL_SM4
 	if (isfinite(OutColor.rgb).x == false)
 	{
 		AlternateColor = float3(0, 0, 1);
diff --git a/Engine/Source/Developer/ShaderCompilerCommon/Private/ShaderCompilerCommon.cpp b/Engine/Source/Developer/ShaderCompilerCommon/Private/ShaderCompilerCommon.cpp
--- a/Engine/Source/Developer/ShaderCompilerCommon/Private/ShaderCompilerCommon.cpp
+++ b/Engine/Source/Developer/ShaderCompilerCommon/Private/ShaderCompilerCommon.cpp
@@ -9,6 +9,30 @@
 
 IMPLEMENT_MODULE(FDefaultModuleImpl, ShaderCompilerCommon);
 
+int16 GetNumUniformBuffersUsed(const FShaderResourceTable& InSRT)
+{
+	auto CountLambda = [&](const TArray<uint32>& In)
+					{
+						int16 LastIndex = -1;
+						for (int32 i = 0; i < In.Num(); ++i)
+						{
+							auto BufferIndex = FRHIResourceTableEntry::GetUniformBufferIndex(In[i]);
+							if (BufferIndex != FRHIResourceTableEntry::GetEndOfStreamToken())
+							{
+								LastIndex = FMath::Max(LastIndex, (int16)BufferIndex);
+							}
+						}
+
+						return LastIndex + 1;
+					};
+	int16 Num = CountLambda(InSRT.SamplerMap);
+	Num = FMath::Max(Num, (int16)CountLambda(InSRT.ShaderResourceViewMap));
+	Num = FMath::Max(Num, (int16)CountLambda(InSRT.TextureMap));
+	Num = FMath::Max(Num, (int16)CountLambda(InSRT.UnorderedAccessViewMap));
+	return Num;
+}
+
+
 void BuildResourceTableTokenStream(const TArray<uint32>& InResourceMap, int32 MaxBoundResourceTable, TArray<uint32>& OutTokenStream)
 {
 	// First we sort the resource map.
diff --git a/Engine/Source/Developer/ShaderCompilerCommon/Public/ShaderCompilerCommon.h b/Engine/Source/Developer/ShaderCompilerCommon/Public/ShaderCompilerCommon.h
--- a/Engine/Source/Developer/ShaderCompilerCommon/Public/ShaderCompilerCommon.h
+++ b/Engine/Source/Developer/ShaderCompilerCommon/Public/ShaderCompilerCommon.h
@@ -1,7 +1,7 @@
 // Copyright 1998-2014 Epic Games, Inc. All Rights Reserved.
 
 /*=============================================================================
-	ShaderCompilerCommon.cpp: Common functionality for platform Shader Compilers.
+	ShaderCompilerCommon.h: Common functionality for platform Shader Compilers.
 =============================================================================*/
 
 #pragma once
@@ -35,6 +35,10 @@ extern SHADERCOMPILERCOMMON_API void BuildResourceTableTokenStream(
 	TArray<uint32>& OutTokenStream
 	);
 
+// Finds the number of used uniform buffers in a resource map
+extern SHADERCOMPILERCOMMON_API int16 GetNumUniformBuffersUsed(const FShaderResourceTable& InSRT);
+
+
 // This function goes through the shader source and converts
 //		static struct
 //		{
@@ -61,6 +65,8 @@ extern SHADERCOMPILERCOMMON_API void BuildResourceTableTokenStream(
 //		}
 extern SHADERCOMPILERCOMMON_API bool RemoveUniformBuffersFromSource(FString& SourceCode);
 
+// Cross compiler support/common functionality
+
 extern SHADERCOMPILERCOMMON_API FString CreateCrossCompilerBatchFileContents(
 											const FString& ShaderFile,
 											const FString& OutputFile,
diff --git a/Engine/Source/Editor/MaterialEditor/Private/MaterialEditor.cpp b/Engine/Source/Editor/MaterialEditor/Private/MaterialEditor.cpp
--- a/Engine/Source/Editor/MaterialEditor/Private/MaterialEditor.cpp
+++ b/Engine/Source/Editor/MaterialEditor/Private/MaterialEditor.cpp
@@ -85,7 +85,7 @@ bool FMatExpressionPreview::ShouldCache(EShaderPlatform Platform, const FShaderT
 		// we only need the non-light-mapped, base pass, local vertex factory shaders for drawing an opaque Material Tile
 		// @todo: Added a FindShaderType by fname or something"
 
-		if (IsES2Platform(Platform))
+		if (IsMobilePlatform(Platform))
 		{
 			if (FCString::Stristr(ShaderType->GetName(), TEXT("BasePassForForwardShadingVSFNoLightMapPolicy")) ||
 				FCString::Stristr(ShaderType->GetName(), TEXT("BasePassForForwardShadingPSFNoLightMapPolicy")))
diff --git a/Engine/Source/Editor/TextureEditor/Private/Models/TextureEditorViewportClient.cpp b/Engine/Source/Editor/TextureEditor/Private/Models/TextureEditorViewportClient.cpp
--- a/Engine/Source/Editor/TextureEditor/Private/Models/TextureEditorViewportClient.cpp
+++ b/Engine/Source/Editor/TextureEditor/Private/Models/TextureEditorViewportClient.cpp
@@ -79,7 +79,7 @@ void FTextureEditorViewportClient::Draw(FViewport* Viewport, FCanvas* Canvas)
 
 		TRefCountPtr<FBatchedElementParameters> BatchedElementParameters;
 
-		if (GRHIFeatureLevel >= ERHIFeatureLevel::SM3)
+		if (GRHIFeatureLevel >= ERHIFeatureLevel::SM4)
 		{
 			if (TextureCube || RTTextureCube)
 			{
diff --git a/Engine/Source/Editor/UnrealEd/Private/MeshPaintRendering.cpp b/Engine/Source/Editor/UnrealEd/Private/MeshPaintRendering.cpp
--- a/Engine/Source/Editor/UnrealEd/Private/MeshPaintRendering.cpp
+++ b/Engine/Source/Editor/UnrealEd/Private/MeshPaintRendering.cpp
@@ -23,7 +23,7 @@ namespace MeshPaintRendering
 
 		static bool ShouldCache( EShaderPlatform Platform )
 		{
-			return IsFeatureLevelSupported(Platform, ERHIFeatureLevel::SM3);
+			return IsFeatureLevelSupported(Platform, ERHIFeatureLevel::SM4);
 		}
 
 		/** Default constructor. */
@@ -67,7 +67,7 @@ namespace MeshPaintRendering
 
 		static bool ShouldCache(EShaderPlatform Platform)
 		{
-			return IsFeatureLevelSupported(Platform, ERHIFeatureLevel::SM3);
+			return IsFeatureLevelSupported(Platform, ERHIFeatureLevel::SM4);
 		}
 
 		/** Default constructor. */
@@ -186,7 +186,7 @@ namespace MeshPaintRendering
 
 		static bool ShouldCache( EShaderPlatform Platform )
 		{
-			return IsFeatureLevelSupported(Platform, ERHIFeatureLevel::SM3);
+			return IsFeatureLevelSupported(Platform, ERHIFeatureLevel::SM4);
 		}
 
 		/** Default constructor. */
@@ -230,7 +230,7 @@ namespace MeshPaintRendering
 
 		static bool ShouldCache(EShaderPlatform Platform)
 		{
-			return IsFeatureLevelSupported(Platform, ERHIFeatureLevel::SM3);
+			return IsFeatureLevelSupported(Platform, ERHIFeatureLevel::SM4);
 		}
 
 		/** Default constructor. */
diff --git a/Engine/Source/Editor/UnrealEd/Private/Texture2DPreview.cpp b/Engine/Source/Editor/UnrealEd/Private/Texture2DPreview.cpp
--- a/Engine/Source/Editor/UnrealEd/Private/Texture2DPreview.cpp
+++ b/Engine/Source/Editor/UnrealEd/Private/Texture2DPreview.cpp
@@ -38,7 +38,7 @@ public:
 	/** Should the shader be cached? Always. */
 	static bool ShouldCache(EShaderPlatform Platform)
 	{
-		return IsFeatureLevelSupported(Platform, ERHIFeatureLevel::SM3);
+		return IsFeatureLevelSupported(Platform, ERHIFeatureLevel::SM4);
 	}
 
 	void SetParameters(FRHICommandList& RHICmdList, const FTexture* TextureValue, const FMatrix& ColorWeightsValue, float GammaValue, float MipLevel, bool bIsNormalMap)
diff --git a/Engine/Source/Editor/UnrealEd/Public/CubemapUnwrapUtils.h b/Engine/Source/Editor/UnrealEd/Public/CubemapUnwrapUtils.h
--- a/Engine/Source/Editor/UnrealEd/Public/CubemapUnwrapUtils.h
+++ b/Engine/Source/Editor/UnrealEd/Public/CubemapUnwrapUtils.h
@@ -141,7 +141,7 @@ public:
 
 	static bool ShouldCache(EShaderPlatform Platform)
 	{
-		return IsFeatureLevelSupported(Platform, ERHIFeatureLevel::SM3);
+		return IsFeatureLevelSupported(Platform, ERHIFeatureLevel::SM4);
 	}
 
 	FIESLightProfilePS() {}
diff --git a/Engine/Source/Runtime/Core/Private/HAL/ConsoleManager.cpp b/Engine/Source/Runtime/Core/Private/HAL/ConsoleManager.cpp
--- a/Engine/Source/Runtime/Core/Private/HAL/ConsoleManager.cpp
+++ b/Engine/Source/Runtime/Core/Private/HAL/ConsoleManager.cpp
@@ -1394,7 +1394,7 @@ static TAutoConsoleVariable<int32> CVarMobileHDR(
 	1,
 	TEXT("0: Mobile renders in LDR gamma space. (suggested for unlit games targeting low-end phones)\n")
 	TEXT("1: Mobile renders in HDR linear space. (default)"),
-	ECVF_RenderThreadSafe);
+	ECVF_RenderThreadSafe | ECVF_ReadOnly);
 
 static TAutoConsoleVariable<int32> CVarMobileHDR32bpp(
 	TEXT("r.MobileHDR32bpp"),
diff --git a/Engine/Source/Runtime/Engine/Private/Atmosphere/Atmosphere.h b/Engine/Source/Runtime/Engine/Private/Atmosphere/Atmosphere.h
--- a/Engine/Source/Runtime/Engine/Private/Atmosphere/Atmosphere.h
+++ b/Engine/Source/Runtime/Engine/Private/Atmosphere/Atmosphere.h
@@ -133,7 +133,7 @@ public:
 	 */
 	virtual void InitRHI() override
 	{
-		if (TextureData && GetFeatureLevel() >= ERHIFeatureLevel::SM3)
+		if (TextureData && GetFeatureLevel() >= ERHIFeatureLevel::SM4)
 		{
 			switch(TexType)
 			{
diff --git a/Engine/Source/Runtime/Engine/Private/LightComponent.cpp b/Engine/Source/Runtime/Engine/Private/LightComponent.cpp
--- a/Engine/Source/Runtime/Engine/Private/LightComponent.cpp
+++ b/Engine/Source/Runtime/Engine/Private/LightComponent.cpp
@@ -15,7 +15,7 @@
 
 void FStaticShadowDepthMap::InitRHI()
 {
-	if (ShadowMapSizeX > 0 && ShadowMapSizeY > 0 && GRHIFeatureLevel >= ERHIFeatureLevel::SM3)
+	if (ShadowMapSizeX > 0 && ShadowMapSizeY > 0 && GRHIFeatureLevel >= ERHIFeatureLevel::SM4)
 	{
 		FRHIResourceCreateInfo CreateInfo;
 		FTexture2DRHIRef Texture2DRHI = RHICreateTexture2D(ShadowMapSizeX, ShadowMapSizeY, PF_R16F, 1, 1, 0, CreateInfo);
diff --git a/Engine/Source/Runtime/Engine/Private/LightMap.cpp b/Engine/Source/Runtime/Engine/Private/LightMap.cpp
--- a/Engine/Source/Runtime/Engine/Private/LightMap.cpp
+++ b/Engine/Source/Runtime/Engine/Private/LightMap.cpp
@@ -70,9 +70,9 @@ FLightMap::FLightMap()
 	: bAllowHighQualityLightMaps(true)
 	, NumRefs(0) 
 {
-	bAllowHighQualityLightMaps = !IsES2Platform(GRHIShaderPlatform) && AllowHighQualityLightmaps();
+	bAllowHighQualityLightMaps = !IsMobilePlatform(GRHIShaderPlatform) && AllowHighQualityLightmaps();
 #if !PLATFORM_DESKTOP 
-	checkf(bAllowHighQualityLightMaps || IsES2Platform(GRHIShaderPlatform), TEXT("Low quality lightmaps are not currently supported on consoles. Make sure console variable r.HighQualityLightMaps is true for this platform"));
+	checkf(bAllowHighQualityLightMaps || IsMobilePlatform(GRHIShaderPlatform), TEXT("Low quality lightmaps are not currently supported on consoles. Make sure console variable r.HighQualityLightMaps is true for this platform"));
 #endif
 }
 
diff --git a/Engine/Source/Runtime/Engine/Private/Materials/HLSLMaterialTranslator.h b/Engine/Source/Runtime/Engine/Private/Materials/HLSLMaterialTranslator.h
--- a/Engine/Source/Runtime/Engine/Private/Materials/HLSLMaterialTranslator.h
+++ b/Engine/Source/Runtime/Engine/Private/Materials/HLSLMaterialTranslator.h
@@ -2047,7 +2047,7 @@ protected:
 
 	virtual int32 TextureCoordinate(uint32 CoordinateIndex, bool UnMirrorU, bool UnMirrorV) override
 	{
-		const uint32 MaxNumCoordinates = FeatureLevel == ERHIFeatureLevel::ES2 ? 3 : 8;
+		const uint32 MaxNumCoordinates = (FeatureLevel == ERHIFeatureLevel::ES2) ? 3 : 8;
 
 		if (CoordinateIndex >= MaxNumCoordinates)
 		{
@@ -2141,7 +2141,7 @@ protected:
 		{
 			// Mobile: Sampling of a particular level depends on an extension; iOS does have it by default but
 			// there's a driver as of 7.0.2 that will cause a GPU hang if used with an Aniso > 1 sampler, so show an error for now
-			if (ErrorUnlessFeatureLevelSupported(ERHIFeatureLevel::SM3) == INDEX_NONE)
+			if (ErrorUnlessFeatureLevelSupported(ERHIFeatureLevel::SM4) == INDEX_NONE)
 			{
 				Errorf(TEXT("Sampling for a specific mip-level is not supported for mobile"));
 				return INDEX_NONE;
@@ -2263,7 +2263,7 @@ protected:
 	// @param SceneTextureId of type ESceneTextureId e.g. PPI_SubsurfaceColor
 	virtual int32 SceneTextureLookup(int32 UV, uint32 InSceneTextureId) override
 	{
-		if (ErrorUnlessFeatureLevelSupported(ERHIFeatureLevel::SM3) == INDEX_NONE)
+		if (ErrorUnlessFeatureLevelSupported(ERHIFeatureLevel::SM4) == INDEX_NONE)
 		{
 			return INDEX_NONE;
 		}
@@ -2419,7 +2419,7 @@ protected:
 			return NonPixelShaderExpressionError();
 		}
 
-		if (ErrorUnlessFeatureLevelSupported(ERHIFeatureLevel::SM3) == INDEX_NONE)
+		if (ErrorUnlessFeatureLevelSupported(ERHIFeatureLevel::SM4) == INDEX_NONE)
 		{
 			return INDEX_NONE;
 		}
@@ -2554,7 +2554,7 @@ protected:
 		{
 			return INDEX_NONE;
 		}
-		else if (GetFeatureLevel() != ERHIFeatureLevel::ES2)
+		else if (GetFeatureLevel() >= ERHIFeatureLevel::SM4)
 		{
 			FString WeightmapName = FString::Printf(TEXT("Weightmap%d"),WeightmapIndex);
 			int32 TextureCodeIndex = TextureParameter(FName(*WeightmapName),  GEngine->WeightMapPlaceholderTexture);
@@ -3113,7 +3113,7 @@ protected:
 			return NonPixelShaderExpressionError();
 		}
 
-		if (ErrorUnlessFeatureLevelSupported(ERHIFeatureLevel::SM3) == INDEX_NONE)
+		if (ErrorUnlessFeatureLevelSupported(ERHIFeatureLevel::SM4) == INDEX_NONE)
 		{
 			return INDEX_NONE;
 		}
@@ -3205,7 +3205,7 @@ protected:
 
 	virtual int32 DDX( int32 X ) override
 	{
-		if (ErrorUnlessFeatureLevelSupported(ERHIFeatureLevel::SM3) == INDEX_NONE)
+		if (ErrorUnlessFeatureLevelSupported(ERHIFeatureLevel::SM4) == INDEX_NONE)
 		{
 			return INDEX_NONE;
 		}
@@ -3231,7 +3231,7 @@ protected:
 
 	virtual int32 DDY( int32 X ) override
 	{
-		if (ErrorUnlessFeatureLevelSupported(ERHIFeatureLevel::SM3) == INDEX_NONE)
+		if (ErrorUnlessFeatureLevelSupported(ERHIFeatureLevel::SM4) == INDEX_NONE)
 		{
 			return INDEX_NONE;
 		}
@@ -3256,7 +3256,7 @@ protected:
 
 	virtual int32 AntialiasedTextureMask(int32 Tex, int32 UV, float Threshold, uint8 Channel) override
 	{
-		if (ErrorUnlessFeatureLevelSupported(ERHIFeatureLevel::SM3) == INDEX_NONE)
+		if (ErrorUnlessFeatureLevelSupported(ERHIFeatureLevel::SM4) == INDEX_NONE)
 		{
 			return INDEX_NONE;
 		}
@@ -3298,7 +3298,7 @@ protected:
 
 	virtual int32 Noise(int32 Position, float Scale, int32 Quality, uint8 NoiseFunction, bool bTurbulence, int32 Levels, float OutputMin, float OutputMax, float LevelScale, int32 FilterWidth) override
 	{
-		if (ErrorUnlessFeatureLevelSupported(ERHIFeatureLevel::SM3) == INDEX_NONE)
+		if (ErrorUnlessFeatureLevelSupported(ERHIFeatureLevel::SM4) == INDEX_NONE)
 		{
 			return INDEX_NONE;
 		}
@@ -3346,7 +3346,7 @@ protected:
 
 	virtual int32 AtmosphericFogColor( int32 WorldPosition ) override
 	{
-		if (ErrorUnlessFeatureLevelSupported(ERHIFeatureLevel::SM3) == INDEX_NONE)
+		if (ErrorUnlessFeatureLevelSupported(ERHIFeatureLevel::SM4) == INDEX_NONE)
 		{
 			return INDEX_NONE;
 		}
@@ -3495,7 +3495,7 @@ protected:
 	*/
 	virtual int32 SpeedTree(ESpeedTreeGeometryType GeometryType, ESpeedTreeWindType WindType, ESpeedTreeLODType LODType, float BillboardThreshold) override 
 	{ 
-		if (ErrorUnlessFeatureLevelSupported(ERHIFeatureLevel::SM3) == INDEX_NONE)
+		if (ErrorUnlessFeatureLevelSupported(ERHIFeatureLevel::SM4) == INDEX_NONE)
 		{
 			return INDEX_NONE;
 		}
@@ -3520,7 +3520,7 @@ protected:
 	 */
 	virtual int32 TextureCoordinateOffset() override
 	{
-		if (FeatureLevel == ERHIFeatureLevel::ES2 && ShaderFrequency == SF_Vertex)
+		if (FeatureLevel < ERHIFeatureLevel::SM4 && ShaderFrequency == SF_Vertex)
 		{
 			return AddInlinedCodeChunk(MCT_Float2, TEXT("Parameters.TexCoordOffset"));
 		}
diff --git a/Engine/Source/Runtime/Engine/Private/Materials/MaterialShared.cpp b/Engine/Source/Runtime/Engine/Private/Materials/MaterialShared.cpp
--- a/Engine/Source/Runtime/Engine/Private/Materials/MaterialShared.cpp
+++ b/Engine/Source/Runtime/Engine/Private/Materials/MaterialShared.cpp
@@ -449,7 +449,7 @@ bool FMaterial::MaterialModifiesMeshPosition_GameThread() const
 bool FMaterial::MaterialMayModifyMeshPosition() const
 {
 	// Conservative estimate when called before material translation has occurred. 
-    // This function is only intended for use in deciding whether or not shader permutations are required.
+	// This function is only intended for use in deciding whether or not shader permutations are required.
 	return HasVertexPositionOffsetConnected() || HasMaterialAttributesConnected() || GetTessellationMode() != MTM_NoTessellation;
 }
 
@@ -909,7 +909,7 @@ void FMaterialResource::GetRepresentativeShaderTypesAndDescriptions(TArray<FStri
 	static auto* MobileHDR = IConsoleManager::Get().FindTConsoleVariableDataInt(TEXT("r.MobileHDR"));
 	bool bMobileHDR = MobileHDR && MobileHDR->GetValueOnAnyThread() == 1;
 
-	if (GetFeatureLevel() > ERHIFeatureLevel::ES2)
+	if (GetFeatureLevel() >= ERHIFeatureLevel::SM4)
 	{
 		if (GetShadingModel() == MSM_Unlit)
 		{
@@ -2285,11 +2285,11 @@ void FMaterialInstanceBasePropertyOverrides::UpdateHash(FSHA1& HashState) const
 		HashState.Update((const uint8*)&OpacityMaskClipValue, sizeof(OpacityMaskClipValue));
 	}
 
- 	if(bOverride_BlendMode)
+	if(bOverride_BlendMode)
 	{
 		const FString HashString = TEXT("bOverride_BlendMode");
- 		HashState.UpdateWithString(*HashString, HashString.Len());
- 		HashState.Update((const uint8*)&BlendMode, sizeof(BlendMode));
+		HashState.UpdateWithString(*HashString, HashString.Len());
+		HashState.Update((const uint8*)&BlendMode, sizeof(BlendMode));
 	}
 
 	if(bOverride_ShadingModel)
diff --git a/Engine/Source/Runtime/Engine/Private/MeshParticleVertexFactory.cpp b/Engine/Source/Runtime/Engine/Private/MeshParticleVertexFactory.cpp
--- a/Engine/Source/Runtime/Engine/Private/MeshParticleVertexFactory.cpp
+++ b/Engine/Source/Runtime/Engine/Private/MeshParticleVertexFactory.cpp
@@ -38,7 +38,7 @@ public:
 
 	virtual void SetMesh(FRHICommandList& RHICmdList, FShader* Shader,const FVertexFactory* VertexFactory,const FSceneView& View,const FMeshBatchElement& BatchElement,uint32 DataFlags) const override
 	{
-		const bool bInstanced = View.GetFeatureLevel() >= ERHIFeatureLevel::SM3;
+		const bool bInstanced = View.GetFeatureLevel() >= ERHIFeatureLevel::SM4;
 		FMeshParticleVertexFactory* MeshParticleVF = (FMeshParticleVertexFactory*)VertexFactory;
 		FVertexShaderRHIParamRef VertexShaderRHI = Shader->GetVertexShader();
 		SetUniformBufferParameter(RHICmdList, VertexShaderRHI, Shader->GetUniformBufferParameter<FMeshParticleUniformParameters>(), MeshParticleVF->GetUniformBuffer() );
@@ -75,7 +75,7 @@ void FMeshParticleVertexFactory::InitRHI()
 {
 	FVertexDeclarationElementList Elements;
 
-	const bool bInstanced = GetFeatureLevel() >= ERHIFeatureLevel::SM3;
+	const bool bInstanced = GetFeatureLevel() >= ERHIFeatureLevel::SM4;
 
 	if (Data.bInitialized)
 	{
diff --git a/Engine/Source/Runtime/Engine/Private/Particles/ParticleSystemRender.cpp b/Engine/Source/Runtime/Engine/Private/Particles/ParticleSystemRender.cpp
--- a/Engine/Source/Runtime/Engine/Private/Particles/ParticleSystemRender.cpp
+++ b/Engine/Source/Runtime/Engine/Private/Particles/ParticleSystemRender.cpp
@@ -540,7 +540,7 @@ FORCEINLINE FVector GetCameraOffsetFromPayload(
 	const FVector& InPosition,
 	const FVector& InCameraPosition )
 {
- 	checkSlow(InCameraPayloadOffset > 0);
+	checkSlow(InCameraPayloadOffset > 0);
 
 	FVector DirToCamera = InCameraPosition - InPosition;
 	float CheckSize = DirToCamera.SizeSquared();
@@ -1039,7 +1039,7 @@ bool FDynamicSpriteEmitterData::GetVertexAndIndexDataNonInstanced(void* VertexDa
 void FDynamicSpriteEmitterData::PreRenderView(FParticleSystemSceneProxy* Proxy, const FSceneViewFamily* ViewFamily, const uint32 VisibilityMap, int32 FrameNumber)
 {
 	const auto FeatureLevel = ViewFamily->Scene->GetFeatureLevel();
-	const bool bInstanced = FeatureLevel >= ERHIFeatureLevel::SM3;
+	const bool bInstanced = FeatureLevel >= ERHIFeatureLevel::SM4;
 
 	// Sort and generate particles for this view.
 	const FDynamicSpriteEmitterReplayDataBase* SourceData = GetSourceData();
@@ -1416,7 +1416,7 @@ int32 FDynamicSpriteEmitterData::RenderEmitter(FParticleSystemSceneProxy* Proxy,
 
 	const auto FeatureLevel = View->GetFeatureLevel();
 
-	const bool bInstanced = FeatureLevel >= ERHIFeatureLevel::SM3;
+	const bool bInstanced = FeatureLevel >= ERHIFeatureLevel::SM4;
 
 	if (bValid == false)
 	{
@@ -2002,7 +2002,7 @@ int32 FDynamicMeshEmitterData::RenderEmitter(FParticleSystemSceneProxy* Proxy, F
 		return 0;
 	}
 
-	const bool bInstanced = View->GetFeatureLevel() >= ERHIFeatureLevel::SM3;
+	const bool bInstanced = View->GetFeatureLevel() >= ERHIFeatureLevel::SM4;
 
 	int32 NumDraws = 0;
 	if (Source.EmitterRenderMode == ERM_Normal)
@@ -2077,7 +2077,7 @@ void FDynamicMeshEmitterData::PreRenderView(FParticleSystemSceneProxy* Proxy, co
 	}
 
 	const auto FeatureLevel = ViewFamily->Scene->GetFeatureLevel();
-	const bool bInstanced = FeatureLevel >= ERHIFeatureLevel::SM3;
+	const bool bInstanced = FeatureLevel >= ERHIFeatureLevel::SM4;
 
 	int32 ParticleCount = Source.ActiveParticleCount;
 	if ((Source.MaxDrawCount >= 0) && (ParticleCount > Source.MaxDrawCount))
@@ -7741,7 +7741,7 @@ void FParticleSystemSceneProxy::GetObjectPositionAndScale(const FSceneView& View
 	const FDynamicEmitterReplayDataBase& EmitterData = DynamicData->DynamicEmitterDataArray[Index]->GetSource();
 	if (EmitterData.bOverrideSystemMacroUV)
 	{
-	    MacroUVRadius = EmitterData.MacroUVRadius;
+		MacroUVRadius = EmitterData.MacroUVRadius;
 		MacroUVPosition = GetLocalToWorld().TransformVector(EmitterData.MacroUVPosition);
 	}
 
@@ -7753,7 +7753,6 @@ void FParticleSystemSceneProxy::GetObjectPositionAndScale(const FSceneView& View
 		const FVector4 UpPostProjectionPosition = View.ViewProjectionMatrix.TransformPosition(MacroUVPosition + MacroUVRadius * View.ViewMatrices.ViewMatrix.GetColumn(1));
 		//checkSlow(RightPostProjectionPosition.X - ObjectPostProjectionPositionWithW.X >= 0.0f && UpPostProjectionPosition.Y - ObjectPostProjectionPositionWithW.Y >= 0.0f);
 
-        
 		// Scales to transform the view space positions corresponding to SystemPositionForMacroUVs +- SystemRadiusForMacroUVs into [0, 1] in xy
 		// Scales to transform the screen space positions corresponding to SystemPositionForMacroUVs +- SystemRadiusForMacroUVs into [0, 1] in zw
 		ObjectMacroUVScales = FVector2D(
diff --git a/Engine/Source/Runtime/Engine/Private/Particles/ParticleVertexFactory.cpp b/Engine/Source/Runtime/Engine/Private/Particles/ParticleVertexFactory.cpp
--- a/Engine/Source/Runtime/Engine/Private/Particles/ParticleVertexFactory.cpp
+++ b/Engine/Source/Runtime/Engine/Private/Particles/ParticleVertexFactory.cpp
@@ -140,7 +140,7 @@ static TGlobalResource<FParticleSpriteVertexDeclaration> GParticleSpriteVertexDe
 
 inline TGlobalResource<FParticleSpriteVertexDeclaration>& GetParticleSpriteVertexDeclaration(ERHIFeatureLevel::Type InFeatureLevel)
 {
-	if (InFeatureLevel >= ERHIFeatureLevel::SM3)
+	if (InFeatureLevel >= ERHIFeatureLevel::SM4)
 	{
 		return GParticleSpriteVertexDeclarationInstanced;
 	}
@@ -177,7 +177,7 @@ void FParticleSpriteVertexFactory::InitRHI()
 
 void FParticleSpriteVertexFactory::InitStreams()
 {
-	const bool bInstanced = GetFeatureLevel() >= ERHIFeatureLevel::SM3;
+	const bool bInstanced = GetFeatureLevel() >= ERHIFeatureLevel::SM4;
 
 	check(Streams.Num() == 0);
 	if(bInstanced) 
diff --git a/Engine/Source/Runtime/Engine/Private/SkeletalRenderGPUSkin.cpp b/Engine/Source/Runtime/Engine/Private/SkeletalRenderGPUSkin.cpp
--- a/Engine/Source/Runtime/Engine/Private/SkeletalRenderGPUSkin.cpp
+++ b/Engine/Source/Runtime/Engine/Private/SkeletalRenderGPUSkin.cpp
@@ -240,7 +240,7 @@ void FSkeletalMeshObjectGPUSkin::UpdateDynamicData_RenderThread(FRHICommandListI
 
 			bool bClothFactory = (DynamicData->ClothSimulUpdateData.Num() > 0) && Chunk.HasApexClothData();
 
-			if (GRHIFeatureLevel < ERHIFeatureLevel::SM3)
+			if (GRHIFeatureLevel < ERHIFeatureLevel::SM4)
 			{
 				bClothFactory = false;
 			}
@@ -265,7 +265,7 @@ void FSkeletalMeshObjectGPUSkin::UpdateDynamicData_RenderThread(FRHICommandListI
 			//  sizeof(FMatrix) == 64
 			// CACHE_LINE_SIZE (128) / 64 = 2
 			const int32 PreFetchStride = 2; // FPlatformMisc::Prefetch stride
-    
+	
 			TArray<FMatrix>& ReferenceToLocalMatrices = DynamicData->ReferenceToLocal;
 			const int32 NumReferenceToLocal = ReferenceToLocalMatrices.Num();
 			for( int32 BoneIdx=0; BoneIdx < NumBones; BoneIdx++ )
@@ -274,7 +274,7 @@ void FSkeletalMeshObjectGPUSkin::UpdateDynamicData_RenderThread(FRHICommandListI
 				FPlatformMisc::Prefetch( ChunkMatrices.GetTypedData() + BoneIdx + PreFetchStride, CACHE_LINE_SIZE ); 
 				FPlatformMisc::Prefetch( ReferenceToLocalMatrices.GetTypedData() + BoneIdx + PreFetchStride );
 				FPlatformMisc::Prefetch( ReferenceToLocalMatrices.GetTypedData() + BoneIdx + PreFetchStride, CACHE_LINE_SIZE );
-    
+	
 				FBoneSkinning& BoneMat = ChunkMatrices[BoneIdx];
 				const FBoneIndexType RefToLocalIdx = Chunk.BoneMap[BoneIdx];
 				const FMatrix& RefToLocal = ReferenceToLocalMatrices[RefToLocalIdx];
@@ -397,8 +397,8 @@ void FSkeletalMeshObjectGPUSkin::FSkeletalMeshObjectLOD::UpdateMorphVertexBuffer
 				FMorphGPUSkinVertex& DestVertex = Buffer[iVertex];
 				float AccumulatedWeight = AccumulatedWeightArray[iVertex];
 
- 				if (AccumulatedWeight > MinVertexAnimBlendWeight)
- 				{
+				if (AccumulatedWeight > MinVertexAnimBlendWeight)
+				{
 					// when copy back, make sure to normalize by accumulated weight
 					// since delta diff of normal is (-2, 2), we divide by 2 to packed into packed normal
 					// when we unpack, we'll apply *2. 
@@ -655,7 +655,7 @@ static void CreateVertexFactoryMorph(TIndirectArray<VertexFactoryTypeBase>& Vert
 {
 	auto* VertexFactory = new VertexFactoryType(InBoneMatrices);
 	VertexFactories.Add(VertexFactory);
-					 	
+						
 	// Setup the update data for enqueue
 	TDynamicUpdateVertexFactoryData<VertexFactoryType> VertexUpdateData(VertexFactory, InVertexBuffers);
 
@@ -693,7 +693,7 @@ static void CreateVertexFactoryCloth(TArray<VertexFactoryTypeBase*>& VertexFacto
 {
 	VertexFactoryType* VertexFactory = new VertexFactoryType(InBoneMatrices);
 	VertexFactories.Add(VertexFactory);
-					 	
+						
 	// Setup the update data for enqueue
 	TDynamicUpdateVertexFactoryData<VertexFactoryType> VertexUpdateData(VertexFactory, InVertexBuffers);
 
@@ -836,7 +836,7 @@ void FSkeletalMeshObjectGPUSkin::FVertexFactoryData::InitAPEXClothVertexFactorie
 	ClothVertexFactories.Empty(Chunks.Num());
 	for( int32 FactoryIdx=0; FactoryIdx < Chunks.Num(); FactoryIdx++ )
 	{
-		if (Chunks[FactoryIdx].HasApexClothData() && GRHIFeatureLevel >= ERHIFeatureLevel::SM3)
+		if (Chunks[FactoryIdx].HasApexClothData() && GRHIFeatureLevel >= ERHIFeatureLevel::SM4)
 		{
 			if (VertexBuffers.VertexBufferGPUSkin->HasExtraBoneInfluences())
 			{
diff --git a/Engine/Source/Runtime/Engine/Private/SkyLightComponent.cpp b/Engine/Source/Runtime/Engine/Private/SkyLightComponent.cpp
--- a/Engine/Source/Runtime/Engine/Private/SkyLightComponent.cpp
+++ b/Engine/Source/Runtime/Engine/Private/SkyLightComponent.cpp
@@ -20,7 +20,7 @@
 
 void FSkyTextureCubeResource::InitRHI()
 {
-	if (GetFeatureLevel() >= ERHIFeatureLevel::SM3)
+	if (GetFeatureLevel() >= ERHIFeatureLevel::SM4)
 	{
 		FRHIResourceCreateInfo CreateInfo;
 		TextureCubeRHI = RHICreateTextureCube(Size, Format, NumMips, 0, CreateInfo);
diff --git a/Engine/Source/Runtime/Engine/Public/ClearReplacementShaders.h b/Engine/Source/Runtime/Engine/Public/ClearReplacementShaders.h
--- a/Engine/Source/Runtime/Engine/Public/ClearReplacementShaders.h
+++ b/Engine/Source/Runtime/Engine/Public/ClearReplacementShaders.h
@@ -25,7 +25,7 @@ public:
 
 	static bool ShouldCache(EShaderPlatform Platform)
 	{
-		return IsFeatureLevelSupported(Platform, ERHIFeatureLevel::SM3);
+		return IsFeatureLevelSupported(Platform, ERHIFeatureLevel::SM4);
 	}
 };
 
@@ -50,7 +50,7 @@ public:
 
 	static bool ShouldCache(EShaderPlatform Platform)
 	{
-		return IsFeatureLevelSupported(Platform, ERHIFeatureLevel::SM3);
+		return IsFeatureLevelSupported(Platform, ERHIFeatureLevel::SM4);
 	}
 
 	static void ModifyCompilationEnvironment(EShaderPlatform Platform, FShaderCompilerEnvironment& OutEnvironment)
diff --git a/Engine/Source/Runtime/Engine/Public/Landscape/LandscapeRender.h b/Engine/Source/Runtime/Engine/Public/Landscape/LandscapeRender.h
--- a/Engine/Source/Runtime/Engine/Public/Landscape/LandscapeRender.h
+++ b/Engine/Source/Runtime/Engine/Public/Landscape/LandscapeRender.h
@@ -161,7 +161,7 @@ public:
 	{
 		// only compile landscape materials for landscape vertex factory
 		// The special engine materials must be compiled for the landscape vertex factory because they are used with it for wireframe, etc.
-		return IsFeatureLevelSupported(Platform, ERHIFeatureLevel::SM3) &&
+		return IsFeatureLevelSupported(Platform, ERHIFeatureLevel::SM4) &&
 			(Material->IsUsedWithLandscape() || Material->IsSpecialEngineMaterial());
 	}
 
diff --git a/Engine/Source/Runtime/Engine/Public/MeshParticleVertexFactory.h b/Engine/Source/Runtime/Engine/Public/MeshParticleVertexFactory.h
--- a/Engine/Source/Runtime/Engine/Public/MeshParticleVertexFactory.h
+++ b/Engine/Source/Runtime/Engine/Public/MeshParticleVertexFactory.h
@@ -97,7 +97,7 @@ public:
 	{
 		FParticleVertexFactoryBase::ModifyCompilationEnvironment(Platform, Material, OutEnvironment);
 
-		const bool bInstanced = IsFeatureLevelSupported(Platform, ERHIFeatureLevel::SM3);
+		const bool bInstanced = IsFeatureLevelSupported(Platform, ERHIFeatureLevel::SM4);
 
 		// Set a define so we can tell in MaterialTemplate.usf when we are compiling a mesh particle vertex factory
 		OutEnvironment.SetDefine(TEXT("PARTICLE_MESH_FACTORY"),TEXT("1"));
diff --git a/Engine/Source/Runtime/Engine/Public/ParticleHelper.h b/Engine/Source/Runtime/Engine/Public/ParticleHelper.h
--- a/Engine/Source/Runtime/Engine/Public/ParticleHelper.h
+++ b/Engine/Source/Runtime/Engine/Public/ParticleHelper.h
@@ -851,7 +851,7 @@ struct FModuleLocationBoneSocketInstancePayload
 	void InitArrayProxies( int32 FixedArraySize )
 	{
 		// Calculate offsets into instance data buffer for the arrays and initialize the buffer proxies. The allocation 
-        // size for these arrays is calculated in RequiredBytesPerInstance.
+		// size for these arrays is calculated in RequiredBytesPerInstance.
 		const uint32 StructSize =  sizeof(FModuleLocationBoneSocketInstancePayload);
 		PrevFrameBoneSocketPositions = TPreallocatedArrayProxy<FVector>((uint8*)this + StructSize, FixedArraySize);
 
@@ -887,8 +887,8 @@ struct FModuleLocationVertSurfaceInstancePayload
 	void InitArrayProxies(int32 FixedArraySize)
 	{
 		// Calculate offsets into instance data buffer for the arrays and initialize the buffer proxies. The allocation 
-        // size for these arrays is calculated in RequiredBytesPerInstance.
- 		const uint32 StructSize =  sizeof(FModuleLocationVertSurfaceInstancePayload);
+		// size for these arrays is calculated in RequiredBytesPerInstance.
+		const uint32 StructSize =  sizeof(FModuleLocationVertSurfaceInstancePayload);
 		ValidAssociatedBoneIndices = TPreallocatedArrayProxy<int32>((uint8*)this + StructSize, FixedArraySize);
 
 		uint32 StructOffset = StructSize + (FixedArraySize*sizeof(int32));
@@ -1249,8 +1249,8 @@ struct FDynamicEmitterReplayDataBase
 		  Scale( FVector( 1.0f ) ),
 		  SortMode(0),	// Default to PSORTMODE_None
 		  bOverrideSystemMacroUV(0),
-	      MacroUVRadius(0.f),
-	      MacroUVPosition(0.f,0.f,0.f)
+		  MacroUVRadius(0.f),
+		  MacroUVPosition(0.f,0.f,0.f)
 	{
 	}
 
@@ -1272,8 +1272,8 @@ struct FDynamicEmitterReplayDataBase
 		Ar << Scale;
 		Ar << SortMode;
 		Ar << bOverrideSystemMacroUV;
-	    Ar << MacroUVRadius;
-	    Ar << MacroUVPosition;
+		Ar << MacroUVRadius;
+		Ar << MacroUVPosition;
 	}
 
 };
@@ -1334,7 +1334,7 @@ struct FDynamicEmitterDataBase
 	/**
 	 *	Called during InitViews for view processing on scene proxies before rendering them
 	 *  Only called for primitives that are visible and have bDynamicRelevance
- 	 *
+	 *
 	 *	@param	Proxy			The 'owner' particle system scene proxy
 	 *	@param	ViewFamily		The ViewFamily to pre-render for
 	 *	@param	VisibilityMap	A BitArray that indicates whether the primitive was visible in that view (index)
@@ -1621,7 +1621,7 @@ struct FDynamicSpriteEmitterData : public FDynamicSpriteEmitterDataBase
 	 */
 	virtual int32 GetDynamicVertexStride(ERHIFeatureLevel::Type InFeatureLevel) const override
 	{
-		const bool bInstanced = InFeatureLevel >= ERHIFeatureLevel::SM3;
+		const bool bInstanced = InFeatureLevel >= ERHIFeatureLevel::SM4;
 		return bInstanced ? sizeof(FParticleSpriteVertex) : sizeof(FParticleSpriteVertexNonInstanced);
 	}
 
@@ -1673,7 +1673,7 @@ struct FDynamicSpriteEmitterData : public FDynamicSpriteEmitterDataBase
 	/**
 	 *	Called during InitViews for view processing on scene proxies before rendering them
 	 *  Only called for primitives that are visible and have bDynamicRelevance
- 	 *
+	 *
 	 *	@param	Proxy			The 'owner' particle system scene proxy
 	 *	@param	ViewFamily		The ViewFamily to pre-render for
 	 *	@param	VisibilityMap	A BitArray that indicates whether the primitive was visible in that view (index)
@@ -1720,7 +1720,7 @@ struct FDynamicSpriteEmitterData : public FDynamicSpriteEmitterDataBase
 	}
 
 	/** The frame source data for this particle system.  This is everything needed to represent this
-	    this particle system frame.  It does not include any transient rendering thread data.  Also, for
+		this particle system frame.  It does not include any transient rendering thread data.  Also, for
 		non-simulating 'replay' particle systems, this data may have come straight from disk! */
 	FDynamicSpriteEmitterReplayData Source;
 
@@ -1849,7 +1849,7 @@ struct FDynamicMeshEmitterData : public FDynamicSpriteEmitterDataBase
 	/**
 	 *	Called during InitViews for view processing on scene proxies before rendering them
 	 *  Only called for primitives that are visible and have bDynamicRelevance
- 	 *
+	 *
 	 *	@param	Proxy			The 'owner' particle system scene proxy
 	 *	@param	ViewFamily		The ViewFamily to pre-render for
 	 *	@param	VisibilityMap	A BitArray that indicates whether the primitive was visible in that view (index)
@@ -1893,7 +1893,7 @@ struct FDynamicMeshEmitterData : public FDynamicSpriteEmitterDataBase
 	}
 
 	/** The frame source data for this particle system.  This is everything needed to represent this
-	    this particle system frame.  It does not include any transient rendering thread data.  Also, for
+		this particle system frame.  It does not include any transient rendering thread data.  Also, for
 		non-simulating 'replay' particle systems, this data may have come straight from disk! */
 	FDynamicMeshEmitterReplayData Source;
 
@@ -2128,7 +2128,7 @@ struct FDynamicBeam2EmitterData : public FDynamicSpriteEmitterDataBase
 	/**
 	 *	Called during InitViews for view processing on scene proxies before rendering them
 	 *  Only called for primitives that are visible and have bDynamicRelevance
- 	 *
+	 *
 	 *	@param	Proxy			The 'owner' particle system scene proxy
 	 *	@param	ViewFamily		The ViewFamily to pre-render for
 	 *	@param	VisibilityMap	A BitArray that indicates whether the primitive was visible in that view (index)
@@ -2179,7 +2179,7 @@ struct FDynamicBeam2EmitterData : public FDynamicSpriteEmitterDataBase
 	virtual void GetIndexAllocInfo(int32& OutNumIndices, int32& OutStride ) const;
 
 	/** The frame source data for this particle system.  This is everything needed to represent this
-	    this particle system frame.  It does not include any transient rendering thread data.  Also, for
+		this particle system frame.  It does not include any transient rendering thread data.  Also, for
 		non-simulating 'replay' particle systems, this data may have come straight from disk! */
 	FDynamicBeam2EmitterReplayData Source;
 
@@ -2267,7 +2267,7 @@ struct FDynamicTrailsEmitterData : public FDynamicSpriteEmitterDataBase
 		, bRenderTangents(false)
 		, bRenderTessellation(false)
 		, bTextureTileDistance(false)
-	    , DistanceTessellationStepSize(12.5f)
+		, DistanceTessellationStepSize(12.5f)
 		, TangentTessellationScalar(25.0f)
 		, TextureTileDistance(0.0f)
 	{
@@ -2304,7 +2304,7 @@ struct FDynamicTrailsEmitterData : public FDynamicSpriteEmitterDataBase
 	/**
 	 *	Called during InitViews for view processing on scene proxies before rendering them
 	 *  Only called for primitives that are visible and have bDynamicRelevance
- 	 *
+	 *
 	 *	@param	Proxy			The 'owner' particle system scene proxy
 	 *	@param	ViewFamily		The ViewFamily to pre-render for
 	 *	@param	VisibilityMap	A BitArray that indicates whether the primitive was visible in that view (index)
@@ -2524,7 +2524,7 @@ public:
 	/**
 	 *	Called during InitViews for view processing on scene proxies before rendering them
 	 *  Only called for primitives that are visible and have bDynamicRelevance
- 	 *
+	 *
 	 *	@param	ViewFamily		The ViewFamily to pre-render for
 	 *	@param	VisibilityMap	A BitArray that indicates whether the primitive was visible in that view (index)
 	 *	@param	FrameNumber		The frame number of this pre-render
diff --git a/Engine/Source/Runtime/Engine/Public/SceneInterface.h b/Engine/Source/Runtime/Engine/Public/SceneInterface.h
--- a/Engine/Source/Runtime/Engine/Public/SceneInterface.h
+++ b/Engine/Source/Runtime/Engine/Public/SceneInterface.h
@@ -231,7 +231,7 @@ public:
 	/**
 	 * Get the optional UWorld that is associated with this scene
 	 * 
- 	 * @return UWorld instance used by this scene
+	 * @return UWorld instance used by this scene
 	 */
 	virtual class UWorld* GetWorld() const = 0;
 	/**
@@ -300,6 +300,11 @@ public:
 
 	virtual ERHIFeatureLevel::Type GetFeatureLevel() const { return GRHIFeatureLevel; }
 
+	bool ShouldUseDeferredRenderer() const
+	{
+		return GetFeatureLevel() >= ERHIFeatureLevel::SM4;
+	}
+
 protected:
 	virtual ~FSceneInterface() {}
-};
\ No newline at end of file
+};
diff --git a/Engine/Source/Runtime/OpenGLDrv/Private/OpenGLES2.cpp b/Engine/Source/Runtime/OpenGLDrv/Private/OpenGLES2.cpp
--- a/Engine/Source/Runtime/OpenGLDrv/Private/OpenGLES2.cpp
+++ b/Engine/Source/Runtime/OpenGLDrv/Private/OpenGLES2.cpp
@@ -137,8 +137,6 @@ void FOpenGLES2::ProcessQueryGLInt()
 	}
 	MaxVertexUniformComponents = FMath::Max<GLint>(MaxVertexUniformComponents, RequiredMaxVertexUniformComponents);
 
-	MaxVertexUniformComponents *= 4;
-	MaxPixelUniformComponents *= 4;
 	MaxGeometryUniformComponents = 0;
 
 	MaxGeometryTextureImageUnits = 0;
diff --git a/Engine/Source/Runtime/OpenGLDrv/Public/OpenGL3.h b/Engine/Source/Runtime/OpenGLDrv/Public/OpenGL3.h
--- a/Engine/Source/Runtime/OpenGLDrv/Public/OpenGL3.h
+++ b/Engine/Source/Runtime/OpenGLDrv/Public/OpenGL3.h
@@ -526,7 +526,7 @@ struct FOpenGL3 : public FOpenGLBase
 
 	static FORCEINLINE ERHIFeatureLevel::Type GetFeatureLevel()
 	{
-		if (FParse::Param(FCommandLine::Get(), TEXT("FeatureLevelES2")))
+		if (FParse::Param(FCommandLine::Get(), TEXT("FeatureLevelES2")) && !GIsEditor)
 		{
 			return ERHIFeatureLevel::ES2;
 		}
@@ -546,7 +546,7 @@ struct FOpenGL3 : public FOpenGLBase
 
 	static FORCEINLINE EShaderPlatform GetShaderPlatform()
 	{
-		static bool bForceFeatureLevelES2 = FParse::Param(FCommandLine::Get(), TEXT("FeatureLevelES2"));
+		static bool bForceFeatureLevelES2 = FParse::Param(FCommandLine::Get(), TEXT("FeatureLevelES2")) && !GIsEditor;
 		if (bForceFeatureLevelES2)
 		{
 			return SP_OPENGL_PCES2;
diff --git a/Engine/Source/Runtime/OpenGLDrv/Public/OpenGLState.h b/Engine/Source/Runtime/OpenGLDrv/Public/OpenGLState.h
--- a/Engine/Source/Runtime/OpenGLDrv/Public/OpenGLState.h
+++ b/Engine/Source/Runtime/OpenGLDrv/Public/OpenGLState.h
@@ -473,19 +473,19 @@ struct FOpenGLRHIState : public FOpenGLCommonState
 		check(!ShaderParameters);
 		FOpenGLCommonState::InitializeResources(NumCombinedTextures, NumComputeUAVUnits);
 		ShaderParameters = new FOpenGLShaderParameterCache[CrossCompiler::NUM_SHADER_STAGES];
-		ShaderParameters[CrossCompiler::SHADER_STAGE_VERTEX].InitializeResources(FOpenGL::GetMaxVertexUniformComponents() * sizeof(float));
-		ShaderParameters[CrossCompiler::SHADER_STAGE_PIXEL].InitializeResources(FOpenGL::GetMaxPixelUniformComponents() * sizeof(float));
-		ShaderParameters[CrossCompiler::SHADER_STAGE_GEOMETRY].InitializeResources(FOpenGL::GetMaxGeometryUniformComponents() * sizeof(float));
+		ShaderParameters[CrossCompiler::SHADER_STAGE_VERTEX].InitializeResources(FOpenGL::GetMaxVertexUniformComponents() * 4 * sizeof(float));
+		ShaderParameters[CrossCompiler::SHADER_STAGE_PIXEL].InitializeResources(FOpenGL::GetMaxPixelUniformComponents() * 4 * sizeof(float));
+		ShaderParameters[CrossCompiler::SHADER_STAGE_GEOMETRY].InitializeResources(FOpenGL::GetMaxGeometryUniformComponents() * 4 * sizeof(float));
 		
 		if ( FOpenGL::SupportsTessellation() )
 		{
-			ShaderParameters[CrossCompiler::SHADER_STAGE_HULL].InitializeResources(FOpenGL::GetMaxHullUniformComponents() * sizeof(float));
-			ShaderParameters[CrossCompiler::SHADER_STAGE_DOMAIN].InitializeResources(FOpenGL::GetMaxDomainUniformComponents() * sizeof(float));
+			ShaderParameters[CrossCompiler::SHADER_STAGE_HULL].InitializeResources(FOpenGL::GetMaxHullUniformComponents() * 4 * sizeof(float));
+			ShaderParameters[CrossCompiler::SHADER_STAGE_DOMAIN].InitializeResources(FOpenGL::GetMaxDomainUniformComponents() * 4 * sizeof(float));
 		}
 
 		if ( FOpenGL::SupportsComputeShaders() )
 		{
-			ShaderParameters[CrossCompiler::SHADER_STAGE_COMPUTE].InitializeResources(FOpenGL::GetMaxComputeUniformComponents() * sizeof(float));
+			ShaderParameters[CrossCompiler::SHADER_STAGE_COMPUTE].InitializeResources(FOpenGL::GetMaxComputeUniformComponents() * 4 * sizeof(float));
 		}
 
 		for (int32 Frequency = 0; Frequency < SF_NumFrequencies; ++Frequency)
diff --git a/Engine/Source/Runtime/RHI/Public/RHIDefinitions.h b/Engine/Source/Runtime/RHI/Public/RHIDefinitions.h
--- a/Engine/Source/Runtime/RHI/Public/RHIDefinitions.h
+++ b/Engine/Source/Runtime/RHI/Public/RHIDefinitions.h
@@ -528,6 +528,12 @@ inline bool IsES2Platform(const EShaderPlatform Platform)
 	return Platform == SP_PCD3D_ES2 || Platform == SP_OPENGL_PCES2 || Platform == SP_OPENGL_ES2 || Platform == SP_OPENGL_ES2_WEBGL || Platform == SP_OPENGL_ES2_IOS || Platform == SP_METAL; 
 }
 
+/** Whether the shader platform corresponds to the ES2 feature level. */
+inline bool IsMobilePlatform(const EShaderPlatform Platform)
+{
+	return IsES2Platform(Platform) || Platform == SP_METAL;
+}
+
 inline bool IsOpenGLPlatform(const EShaderPlatform Platform)
 {
 	return Platform == SP_OPENGL_SM4 || Platform == SP_OPENGL_SM4_MAC || Platform == SP_OPENGL_SM5 || Platform == SP_OPENGL_PCES2 || Platform == SP_OPENGL_ES2 || Platform == SP_OPENGL_ES2_WEBGL || Platform == SP_OPENGL_ES2_IOS;
@@ -618,7 +624,7 @@ inline bool RHINeedsToSwitchVerticalAxis(EShaderPlatform Platform)
 inline bool RHISupportsInstancing(const EShaderPlatform Platform)
 {
 	//@todo-rco: Add Metal support
-	return IsFeatureLevelSupported(Platform, ERHIFeatureLevel::SM3);// || (Platform == SP_METAL);
+	return IsFeatureLevelSupported(Platform, ERHIFeatureLevel::SM4);// || (Platform == SP_METAL);
 }
 
 inline bool RHISupportsSeparateMSAAAndResolveTextures(const EShaderPlatform Platform)
diff --git a/Engine/Source/Runtime/RHI/Public/RHIResources.h b/Engine/Source/Runtime/RHI/Public/RHIResources.h
--- a/Engine/Source/Runtime/RHI/Public/RHIResources.h
+++ b/Engine/Source/Runtime/RHI/Public/RHIResources.h
@@ -129,7 +129,8 @@ struct FRHIUniformBufferLayout
 		if (!bComputedHash)
 		{
 			uint32 TmpHash = ConstantBufferSize;
-			TmpHash ^= ResourceOffset;
+			// This is to account for 32vs64 bits difference in pointer sizes.
+			TmpHash ^= Align(ResourceOffset, 8);
 			uint32 N = Resources.Num();
 			while (N >= 4)
 			{
diff --git a/Engine/Source/Runtime/Renderer/Private/AtmosphereRendering.cpp b/Engine/Source/Runtime/Renderer/Private/AtmosphereRendering.cpp
--- a/Engine/Source/Runtime/Renderer/Private/AtmosphereRendering.cpp
+++ b/Engine/Source/Runtime/Renderer/Private/AtmosphereRendering.cpp
@@ -285,7 +285,7 @@ public:
 
 	static bool ShouldCache(EShaderPlatform Platform)
 	{
-		return IsFeatureLevelSupported(Platform, ERHIFeatureLevel::SM3);
+		return IsFeatureLevelSupported(Platform, ERHIFeatureLevel::SM4);
 	}
 
 	FAtmosphericVS( )	{ }
diff --git a/Engine/Source/Runtime/Renderer/Private/CompositionLighting/PostProcessAmbient.cpp b/Engine/Source/Runtime/Renderer/Private/CompositionLighting/PostProcessAmbient.cpp
--- a/Engine/Source/Runtime/Renderer/Private/CompositionLighting/PostProcessAmbient.cpp
+++ b/Engine/Source/Runtime/Renderer/Private/CompositionLighting/PostProcessAmbient.cpp
@@ -18,7 +18,7 @@ class FPostProcessAmbientPS : public FGlobalShader
 
 	static bool ShouldCache(EShaderPlatform Platform)
 	{
-		return IsFeatureLevelSupported(Platform, ERHIFeatureLevel::SM3);
+		return IsFeatureLevelSupported(Platform, ERHIFeatureLevel::SM4);
 	}
 
 	static void ModifyCompilationEnvironment(EShaderPlatform Platform, FShaderCompilerEnvironment& OutEnvironment)
diff --git a/Engine/Source/Runtime/Renderer/Private/CompositionLighting/PostProcessPassThrough.cpp b/Engine/Source/Runtime/Renderer/Private/CompositionLighting/PostProcessPassThrough.cpp
--- a/Engine/Source/Runtime/Renderer/Private/CompositionLighting/PostProcessPassThrough.cpp
+++ b/Engine/Source/Runtime/Renderer/Private/CompositionLighting/PostProcessPassThrough.cpp
@@ -18,7 +18,7 @@ class FPostProcessPassThroughPS : public FGlobalShader
 
 	static bool ShouldCache(EShaderPlatform Platform)
 	{
-		return IsFeatureLevelSupported(Platform, ERHIFeatureLevel::SM3);
+		return IsFeatureLevelSupported(Platform, ERHIFeatureLevel::SM4);
 	}
 
 	/** Default constructor. */
diff --git a/Engine/Source/Runtime/Renderer/Private/DeferredShadingRenderer.cpp b/Engine/Source/Runtime/Renderer/Private/DeferredShadingRenderer.cpp
--- a/Engine/Source/Runtime/Renderer/Private/DeferredShadingRenderer.cpp
+++ b/Engine/Source/Runtime/Renderer/Private/DeferredShadingRenderer.cpp
@@ -1293,7 +1293,7 @@ public:
 
 	static bool ShouldCache(EShaderPlatform Platform)
 	{ 
-		return IsFeatureLevelSupported(Platform, ERHIFeatureLevel::SM3);
+		return IsFeatureLevelSupported(Platform, ERHIFeatureLevel::SM4);
 	}
 
 	FDownsampleSceneDepthPS(const ShaderMetaType::CompiledShaderInitializerType& Initializer):
@@ -1350,7 +1350,7 @@ FGlobalBoundShaderState DownsampleDepthBoundShaderState;
 /** Updates the downsized depth buffer with the current full resolution depth buffer. */
 void FDeferredShadingSceneRenderer::UpdateDownsampledDepthSurface(FRHICommandList& RHICmdList)
 {
-	if (GSceneRenderTargets.UseDownsizedOcclusionQueries() && IsFeatureLevelSupported(GRHIShaderPlatform, ERHIFeatureLevel::SM3))
+	if (GSceneRenderTargets.UseDownsizedOcclusionQueries() && IsFeatureLevelSupported(GRHIShaderPlatform, ERHIFeatureLevel::SM4))
 	{
 		SetRenderTarget(RHICmdList, NULL, GSceneRenderTargets.GetSmallDepthSurface());
 
diff --git a/Engine/Source/Runtime/Renderer/Private/DepthRendering.cpp b/Engine/Source/Runtime/Renderer/Private/DepthRendering.cpp
--- a/Engine/Source/Runtime/Renderer/Private/DepthRendering.cpp
+++ b/Engine/Source/Runtime/Renderer/Private/DepthRendering.cpp
@@ -29,12 +29,12 @@ public:
 		if (bUsePositionOnlyStream)
 		{
 			return VertexFactoryType->SupportsPositionOnly() && Material->IsSpecialEngineMaterial() 
-				&& IsFeatureLevelSupported(Platform, ERHIFeatureLevel::SM3);
+				&& IsFeatureLevelSupported(Platform, ERHIFeatureLevel::SM4);
 		}
 
 		// Only compile for the default material and masked materials
 		return (Material->IsSpecialEngineMaterial() || Material->IsMasked() || Material->MaterialMayModifyMeshPosition())
-			 && IsFeatureLevelSupported(Platform, ERHIFeatureLevel::SM3);
+			 && IsFeatureLevelSupported(Platform, ERHIFeatureLevel::SM4);
 	}
 
 	void SetParameters(FRHICommandList& RHICmdList, const FMaterialRenderProxy* MaterialRenderProxy,const FMaterial& MaterialResource,const FSceneView& View)
@@ -106,7 +106,7 @@ public:
 	static bool ShouldCache(EShaderPlatform Platform,const FMaterial* Material,const FVertexFactoryType* VertexFactoryType)
 	{
 		// Compile for materials that are masked.
-		return Material->IsMasked() && IsFeatureLevelSupported(Platform, ERHIFeatureLevel::SM3);
+		return Material->IsMasked() && IsFeatureLevelSupported(Platform, ERHIFeatureLevel::SM4);
 	}
 
 	FDepthOnlyPS(const ShaderMetaType::CompiledShaderInitializerType& Initializer):
diff --git a/Engine/Source/Runtime/Renderer/Private/FogRendering.cpp b/Engine/Source/Runtime/Renderer/Private/FogRendering.cpp
--- a/Engine/Source/Runtime/Renderer/Private/FogRendering.cpp
+++ b/Engine/Source/Runtime/Renderer/Private/FogRendering.cpp
@@ -71,7 +71,7 @@ public:
 
 	static bool ShouldCache(EShaderPlatform Platform)
 	{
-		return IsFeatureLevelSupported(Platform, ERHIFeatureLevel::SM3);
+		return IsFeatureLevelSupported(Platform, ERHIFeatureLevel::SM4);
 	}
 
 	FHeightFogVS( )	{ }
@@ -132,7 +132,7 @@ public:
 
 	static bool ShouldCache(EShaderPlatform Platform)
 	{
-		return IsFeatureLevelSupported(Platform, ERHIFeatureLevel::SM3);
+		return IsFeatureLevelSupported(Platform, ERHIFeatureLevel::SM4);
 	}
 
 	FExponentialHeightFogPS( )	{ }
diff --git a/Engine/Source/Runtime/Renderer/Private/GPUBenchmark.cpp b/Engine/Source/Runtime/Renderer/Private/GPUBenchmark.cpp
--- a/Engine/Source/Runtime/Renderer/Private/GPUBenchmark.cpp
+++ b/Engine/Source/Runtime/Renderer/Private/GPUBenchmark.cpp
@@ -26,7 +26,7 @@ class FPostProcessBenchmarkPS : public FGlobalShader
 
 	static bool ShouldCache(EShaderPlatform Platform)
 	{
-		return IsFeatureLevelSupported(Platform, ERHIFeatureLevel::SM3);
+		return IsFeatureLevelSupported(Platform, ERHIFeatureLevel::SM4);
 	}
 
 	static void ModifyCompilationEnvironment(EShaderPlatform Platform, FShaderCompilerEnvironment& OutEnvironment)
diff --git a/Engine/Source/Runtime/Renderer/Private/IndirectLightingCache.cpp b/Engine/Source/Runtime/Renderer/Private/IndirectLightingCache.cpp
--- a/Engine/Source/Runtime/Renderer/Private/IndirectLightingCache.cpp
+++ b/Engine/Source/Runtime/Renderer/Private/IndirectLightingCache.cpp
@@ -134,7 +134,7 @@ bool IsIndirectLightingCacheAllowed(ERHIFeatureLevel::Type InFeatureLevel)
 bool CanIndirectLightingCacheUseVolumeTexture(ERHIFeatureLevel::Type InFeatureLevel)
 {
 	// @todo Mac OS X/OpenGL: For OpenGL devices which don't support volume-texture rendering we need to use the simpler point indirect lighting shaders.
-	return InFeatureLevel >= ERHIFeatureLevel::SM3 && GSupportsVolumeTextureRendering;
+	return InFeatureLevel >= ERHIFeatureLevel::SM4 && GSupportsVolumeTextureRendering;
 }
 
 FIndirectLightingCache::FIndirectLightingCache(ERHIFeatureLevel::Type InFeatureLevel)
diff --git a/Engine/Source/Runtime/Renderer/Private/LightFunctionRendering.cpp b/Engine/Source/Runtime/Renderer/Private/LightFunctionRendering.cpp
--- a/Engine/Source/Runtime/Renderer/Private/LightFunctionRendering.cpp
+++ b/Engine/Source/Runtime/Renderer/Private/LightFunctionRendering.cpp
@@ -24,7 +24,7 @@ public:
 	  */
 	static bool ShouldCache(EShaderPlatform Platform, const FMaterial* Material)
 	{
-		return Material->IsLightFunction() && IsFeatureLevelSupported(Platform, ERHIFeatureLevel::SM3);
+		return Material->IsLightFunction() && IsFeatureLevelSupported(Platform, ERHIFeatureLevel::SM4);
 	}
 
 	FLightFunctionVS( )	{ }
@@ -81,7 +81,7 @@ public:
 	  */
 	static bool ShouldCache(EShaderPlatform Platform, const FMaterial* Material)
 	{
-		return Material->IsLightFunction() && IsFeatureLevelSupported(Platform, ERHIFeatureLevel::SM3);
+		return Material->IsLightFunction() && IsFeatureLevelSupported(Platform, ERHIFeatureLevel::SM4);
 	}
 
 	FLightFunctionPS() {}
diff --git a/Engine/Source/Runtime/Renderer/Private/LightMapDensityRendering.cpp b/Engine/Source/Runtime/Renderer/Private/LightMapDensityRendering.cpp
--- a/Engine/Source/Runtime/Renderer/Private/LightMapDensityRendering.cpp
+++ b/Engine/Source/Runtime/Renderer/Private/LightMapDensityRendering.cpp
@@ -38,7 +38,7 @@ IMPLEMENT_DENSITY_LIGHTMAPPED_SHADER_TYPE( TLightMapPolicy<HQ_LIGHTMAP>, TLightM
 bool FDeferredShadingSceneRenderer::RenderLightMapDensities(FRHICommandListImmediate& RHICmdList)
 {
 	bool bDirty=0;
-	if (Scene->GetFeatureLevel() >= ERHIFeatureLevel::SM3)
+	if (Scene->GetFeatureLevel() >= ERHIFeatureLevel::SM4)
 	{
 		SCOPED_DRAW_EVENT(LightMapDensity, DEC_SCENE_ITEMS);
 
diff --git a/Engine/Source/Runtime/Renderer/Private/LightMapDensityRendering.h b/Engine/Source/Runtime/Renderer/Private/LightMapDensityRendering.h
--- a/Engine/Source/Runtime/Renderer/Private/LightMapDensityRendering.h
+++ b/Engine/Source/Runtime/Renderer/Private/LightMapDensityRendering.h
@@ -20,7 +20,7 @@ public:
 	{
 		return (Material->IsSpecialEngineMaterial() || Material->IsMasked() || Material->MaterialMayModifyMeshPosition())
 				&& LightMapPolicyType::ShouldCache(Platform,Material,VertexFactoryType)
-				&& IsFeatureLevelSupported(Platform, ERHIFeatureLevel::SM3);
+				&& IsFeatureLevelSupported(Platform, ERHIFeatureLevel::SM4);
 	}
 
 	static void ModifyCompilationEnvironment(EShaderPlatform Platform, const FMaterial* Material, FShaderCompilerEnvironment& OutEnvironment)
@@ -134,7 +134,7 @@ public:
 	{
 		return (Material->IsSpecialEngineMaterial() || Material->IsMasked() || Material->MaterialMayModifyMeshPosition())
 				&& LightMapPolicyType::ShouldCache(Platform,Material,VertexFactoryType)
-				&& IsFeatureLevelSupported(Platform, ERHIFeatureLevel::SM3);
+				&& IsFeatureLevelSupported(Platform, ERHIFeatureLevel::SM4);
 	}
 
 	static void ModifyCompilationEnvironment(EShaderPlatform Platform, const FMaterial* Material, FShaderCompilerEnvironment& OutEnvironment)
@@ -220,7 +220,7 @@ public:
 	}
 
 private:
- 	FShaderParameter LightMapDensity;
+	FShaderParameter LightMapDensity;
 	FShaderParameter BuiltLightingAndSelectedFlags;
 	FShaderParameter DensitySelectedColor;
 	FShaderParameter LightMapResolutionScale;
diff --git a/Engine/Source/Runtime/Renderer/Private/LightShaftRendering.cpp b/Engine/Source/Runtime/Renderer/Private/LightShaftRendering.cpp
--- a/Engine/Source/Runtime/Renderer/Private/LightShaftRendering.cpp
+++ b/Engine/Source/Runtime/Renderer/Private/LightShaftRendering.cpp
@@ -226,7 +226,7 @@ public:
 
 	static bool ShouldCache(EShaderPlatform Platform)
 	{
-		return IsFeatureLevelSupported(Platform, ERHIFeatureLevel::SM3);
+		return IsFeatureLevelSupported(Platform, ERHIFeatureLevel::SM4);
 	}
 
 	/** Default constructor. */
@@ -259,7 +259,7 @@ public:
 
 	static bool ShouldCache(EShaderPlatform Platform)
 	{
-		return IsFeatureLevelSupported(Platform, ERHIFeatureLevel::SM3); 
+		return IsFeatureLevelSupported(Platform, ERHIFeatureLevel::SM4); 
 	}
 
 	static void ModifyCompilationEnvironment(EShaderPlatform Platform, FShaderCompilerEnvironment& OutEnvironment)
@@ -333,7 +333,7 @@ public:
 
 	static bool ShouldCache(EShaderPlatform Platform)
 	{
-		return IsFeatureLevelSupported(Platform, ERHIFeatureLevel::SM3); 
+		return IsFeatureLevelSupported(Platform, ERHIFeatureLevel::SM4); 
 	}
 
 	static void ModifyCompilationEnvironment(EShaderPlatform Platform, FShaderCompilerEnvironment& OutEnvironment)
@@ -390,7 +390,7 @@ public:
 
 	static bool ShouldCache(EShaderPlatform Platform)
 	{
-		return IsFeatureLevelSupported(Platform, ERHIFeatureLevel::SM3); 
+		return IsFeatureLevelSupported(Platform, ERHIFeatureLevel::SM4); 
 	}
 
 	/** Default constructor. */
@@ -790,7 +790,7 @@ public:
 
 	static bool ShouldCache(EShaderPlatform Platform)
 	{
-		return IsFeatureLevelSupported(Platform, ERHIFeatureLevel::SM3); 
+		return IsFeatureLevelSupported(Platform, ERHIFeatureLevel::SM4); 
 	}
 
 	/** Default constructor. */
diff --git a/Engine/Source/Runtime/Renderer/Private/PostProcess/PostProcessAA.cpp b/Engine/Source/Runtime/Renderer/Private/PostProcess/PostProcessAA.cpp
--- a/Engine/Source/Runtime/Renderer/Private/PostProcess/PostProcessAA.cpp
+++ b/Engine/Source/Runtime/Renderer/Private/PostProcess/PostProcessAA.cpp
@@ -164,7 +164,7 @@ public:
 
 	static bool ShouldCache(EShaderPlatform Platform)
 	{
-		return IsFeatureLevelSupported(Platform, ERHIFeatureLevel::SM3);
+		return IsFeatureLevelSupported(Platform, ERHIFeatureLevel::SM4);
 	}
 
 	/** Default constructor. */
diff --git a/Engine/Source/Runtime/Renderer/Private/PostProcess/PostProcessBokehDOFRecombine.cpp b/Engine/Source/Runtime/Renderer/Private/PostProcess/PostProcessBokehDOFRecombine.cpp
--- a/Engine/Source/Runtime/Renderer/Private/PostProcess/PostProcessBokehDOFRecombine.cpp
+++ b/Engine/Source/Runtime/Renderer/Private/PostProcess/PostProcessBokehDOFRecombine.cpp
@@ -25,7 +25,7 @@ class FPostProcessBokehDOFRecombinePS : public FGlobalShader
 
 	static bool ShouldCache(EShaderPlatform Platform)
 	{
-		return IsFeatureLevelSupported(Platform, ERHIFeatureLevel::SM3);
+		return IsFeatureLevelSupported(Platform, ERHIFeatureLevel::SM4);
 	}
 
 	static void ModifyCompilationEnvironment(EShaderPlatform Platform, FShaderCompilerEnvironment& OutEnvironment)
diff --git a/Engine/Source/Runtime/Renderer/Private/PostProcess/PostProcessBusyWait.cpp b/Engine/Source/Runtime/Renderer/Private/PostProcess/PostProcessBusyWait.cpp
--- a/Engine/Source/Runtime/Renderer/Private/PostProcess/PostProcessBusyWait.cpp
+++ b/Engine/Source/Runtime/Renderer/Private/PostProcess/PostProcessBusyWait.cpp
@@ -28,7 +28,7 @@ class FPostProcessBusyWaitPS : public FGlobalShader
 
 	static bool ShouldCache(EShaderPlatform Platform)
 	{
-		return IsFeatureLevelSupported(Platform, ERHIFeatureLevel::SM3);
+		return IsFeatureLevelSupported(Platform, ERHIFeatureLevel::SM4);
 	}
 
 	/** Default constructor. */
diff --git a/Engine/Source/Runtime/Renderer/Private/PostProcess/PostProcessCombineLUTs.cpp b/Engine/Source/Runtime/Renderer/Private/PostProcess/PostProcessCombineLUTs.cpp
--- a/Engine/Source/Runtime/Renderer/Private/PostProcess/PostProcessCombineLUTs.cpp
+++ b/Engine/Source/Runtime/Renderer/Private/PostProcess/PostProcessCombineLUTs.cpp
@@ -111,7 +111,7 @@ class FLUTBlenderPS : public FGlobalShader
 	DECLARE_SHADER_TYPE(FLUTBlenderPS,Global);
 public:
 
-	static bool ShouldCache(EShaderPlatform Platform) { return IsFeatureLevelSupported(Platform, ERHIFeatureLevel::SM3); }
+	static bool ShouldCache(EShaderPlatform Platform) { return IsFeatureLevelSupported(Platform, ERHIFeatureLevel::SM4); }
 
 	FLUTBlenderPS(const ShaderMetaType::CompiledShaderInitializerType& Initializer)
 		: FGlobalShader(Initializer) 
diff --git a/Engine/Source/Runtime/Renderer/Private/PostProcess/PostProcessDownsample.cpp b/Engine/Source/Runtime/Renderer/Private/PostProcess/PostProcessDownsample.cpp
--- a/Engine/Source/Runtime/Renderer/Private/PostProcess/PostProcessDownsample.cpp
+++ b/Engine/Source/Runtime/Renderer/Private/PostProcess/PostProcessDownsample.cpp
@@ -20,7 +20,7 @@ class FPostProcessDownsamplePS : public FGlobalShader
 
 	static bool ShouldCache(EShaderPlatform Platform)
 	{
-		return Method != 2 || IsFeatureLevelSupported(Platform, ERHIFeatureLevel::SM3);
+		return Method != 2 || IsFeatureLevelSupported(Platform, ERHIFeatureLevel::SM4);
 	}
 
 	static void ModifyCompilationEnvironment(EShaderPlatform Platform, FShaderCompilerEnvironment& OutEnvironment)
diff --git a/Engine/Source/Runtime/Renderer/Private/PostProcess/PostProcessLensFlares.cpp b/Engine/Source/Runtime/Renderer/Private/PostProcess/PostProcessLensFlares.cpp
--- a/Engine/Source/Runtime/Renderer/Private/PostProcess/PostProcessLensFlares.cpp
+++ b/Engine/Source/Runtime/Renderer/Private/PostProcess/PostProcessLensFlares.cpp
@@ -18,7 +18,7 @@ class FPostProcessLensFlareBasePS : public FGlobalShader
 
 	static bool ShouldCache(EShaderPlatform Platform)
 	{
-		return IsFeatureLevelSupported(Platform, ERHIFeatureLevel::SM3);
+		return IsFeatureLevelSupported(Platform, ERHIFeatureLevel::SM4);
 	}
 
 	/** Default constructor. */
@@ -62,7 +62,7 @@ class FPostProcessLensFlaresPS : public FGlobalShader
 
 	static bool ShouldCache(EShaderPlatform Platform)
 	{
-		return IsFeatureLevelSupported(Platform, ERHIFeatureLevel::SM3);
+		return IsFeatureLevelSupported(Platform, ERHIFeatureLevel::SM4);
 	}
 
 	/** Default constructor. */
diff --git a/Engine/Source/Runtime/Renderer/Private/PostProcess/PostProcessMaterial.cpp b/Engine/Source/Runtime/Renderer/Private/PostProcess/PostProcessMaterial.cpp
--- a/Engine/Source/Runtime/Renderer/Private/PostProcess/PostProcessMaterial.cpp
+++ b/Engine/Source/Runtime/Renderer/Private/PostProcess/PostProcessMaterial.cpp
@@ -23,7 +23,7 @@ public:
 	  */
 	static bool ShouldCache(EShaderPlatform Platform, const FMaterial* Material)
 	{
-		return (Material->GetMaterialDomain() == MD_PostProcess) && IsFeatureLevelSupported(Platform, ERHIFeatureLevel::SM3);
+		return (Material->GetMaterialDomain() == MD_PostProcess) && IsFeatureLevelSupported(Platform, ERHIFeatureLevel::SM4);
 	}
 
 	FPostProcessMaterialVS( )	{ }
@@ -61,7 +61,7 @@ public:
 	  */
 	static bool ShouldCache(EShaderPlatform Platform, const FMaterial* Material)
 	{
-		return (Material->GetMaterialDomain() == MD_PostProcess) && IsFeatureLevelSupported(Platform, ERHIFeatureLevel::SM3);
+		return (Material->GetMaterialDomain() == MD_PostProcess) && IsFeatureLevelSupported(Platform, ERHIFeatureLevel::SM4);
 	}
 
 	FPostProcessMaterialPS() {}
diff --git a/Engine/Source/Runtime/Renderer/Private/PostProcess/PostProcessMotionBlur.cpp b/Engine/Source/Runtime/Renderer/Private/PostProcess/PostProcessMotionBlur.cpp
--- a/Engine/Source/Runtime/Renderer/Private/PostProcess/PostProcessMotionBlur.cpp
+++ b/Engine/Source/Runtime/Renderer/Private/PostProcess/PostProcessMotionBlur.cpp
@@ -29,7 +29,7 @@ class FPostProcessMotionBlurSetupPS : public FGlobalShader
 
 	static bool ShouldCache(EShaderPlatform Platform)
 	{
-		return IsFeatureLevelSupported(Platform, ERHIFeatureLevel::SM3);
+		return IsFeatureLevelSupported(Platform, ERHIFeatureLevel::SM4);
 	}
 
 	/** Default constructor. */
diff --git a/Engine/Source/Runtime/Renderer/Private/PostProcess/PostProcessSceneColorFringe.cpp b/Engine/Source/Runtime/Renderer/Private/PostProcess/PostProcessSceneColorFringe.cpp
--- a/Engine/Source/Runtime/Renderer/Private/PostProcess/PostProcessSceneColorFringe.cpp
+++ b/Engine/Source/Runtime/Renderer/Private/PostProcess/PostProcessSceneColorFringe.cpp
@@ -22,7 +22,7 @@ class FPostProcessSceneColorFringeVS : public FGlobalShader
 
 	static bool ShouldCache(EShaderPlatform Platform)
 	{
-		return IsFeatureLevelSupported(Platform, ERHIFeatureLevel::SM3);
+		return IsFeatureLevelSupported(Platform, ERHIFeatureLevel::SM4);
 	}
 
 	/** Default constructor. */
@@ -89,7 +89,7 @@ class FPostProcessSceneColorFringePS : public FGlobalShader
 
 	static bool ShouldCache(EShaderPlatform Platform)
 	{
-		return IsFeatureLevelSupported(Platform, ERHIFeatureLevel::SM3);
+		return IsFeatureLevelSupported(Platform, ERHIFeatureLevel::SM4);
 	}
 
 	/** Default constructor. */
diff --git a/Engine/Source/Runtime/Renderer/Private/PostProcess/PostProcessSelectionOutline.cpp b/Engine/Source/Runtime/Renderer/Private/PostProcess/PostProcessSelectionOutline.cpp
--- a/Engine/Source/Runtime/Renderer/Private/PostProcess/PostProcessSelectionOutline.cpp
+++ b/Engine/Source/Runtime/Renderer/Private/PostProcess/PostProcessSelectionOutline.cpp
@@ -206,7 +206,7 @@ class FPostProcessSelectionOutlinePS : public FGlobalShader
 			}
 		}
 
-		return IsFeatureLevelSupported(Platform, ERHIFeatureLevel::SM3);
+		return IsFeatureLevelSupported(Platform, ERHIFeatureLevel::SM4);
 	}
 
 	static void ModifyCompilationEnvironment(EShaderPlatform Platform, FShaderCompilerEnvironment& OutEnvironment)
diff --git a/Engine/Source/Runtime/Renderer/Private/PostProcess/PostProcessTonemap.cpp b/Engine/Source/Runtime/Renderer/Private/PostProcess/PostProcessTonemap.cpp
--- a/Engine/Source/Runtime/Renderer/Private/PostProcess/PostProcessTonemap.cpp
+++ b/Engine/Source/Runtime/Renderer/Private/PostProcess/PostProcessTonemap.cpp
@@ -940,7 +940,7 @@ class FPostProcessTonemapPS : public FGlobalShader
 
 	static bool ShouldCache(EShaderPlatform Platform)
 	{
-		return IsFeatureLevelSupported(Platform, ERHIFeatureLevel::SM3);
+		return IsFeatureLevelSupported(Platform, ERHIFeatureLevel::SM4);
 	}
 
 	static void ModifyCompilationEnvironment(EShaderPlatform Platform, FShaderCompilerEnvironment& OutEnvironment)
diff --git a/Engine/Source/Runtime/Renderer/Private/PostProcess/PostProcessUIBlur.cpp b/Engine/Source/Runtime/Renderer/Private/PostProcess/PostProcessUIBlur.cpp
--- a/Engine/Source/Runtime/Renderer/Private/PostProcess/PostProcessUIBlur.cpp
+++ b/Engine/Source/Runtime/Renderer/Private/PostProcess/PostProcessUIBlur.cpp
@@ -19,7 +19,7 @@ class FPostProcessUIBlurPS : public FGlobalShader
 
 	static bool ShouldCache(EShaderPlatform Platform)
 	{
-		return IsFeatureLevelSupported(Platform, ERHIFeatureLevel::SM3);
+		return IsFeatureLevelSupported(Platform, ERHIFeatureLevel::SM4);
 	}
 
 	/** Default constructor. */
diff --git a/Engine/Source/Runtime/Renderer/Private/PostProcess/PostProcessUpscale.cpp b/Engine/Source/Runtime/Renderer/Private/PostProcess/PostProcessUpscale.cpp
--- a/Engine/Source/Runtime/Renderer/Private/PostProcess/PostProcessUpscale.cpp
+++ b/Engine/Source/Runtime/Renderer/Private/PostProcess/PostProcessUpscale.cpp
@@ -28,7 +28,7 @@ class FPostProcessUpscalePS : public FGlobalShader
 
 	static bool ShouldCache(EShaderPlatform Platform)
 	{
-		return IsFeatureLevelSupported(Platform, ERHIFeatureLevel::SM3);
+		return IsFeatureLevelSupported(Platform, ERHIFeatureLevel::SM4);
 	}
 
 	static void ModifyCompilationEnvironment(EShaderPlatform Platform, FShaderCompilerEnvironment& OutEnvironment)
diff --git a/Engine/Source/Runtime/Renderer/Private/PostProcess/PostProcessWeightedSampleSum.cpp b/Engine/Source/Runtime/Renderer/Private/PostProcess/PostProcessWeightedSampleSum.cpp
--- a/Engine/Source/Runtime/Renderer/Private/PostProcess/PostProcessWeightedSampleSum.cpp
+++ b/Engine/Source/Runtime/Renderer/Private/PostProcess/PostProcessWeightedSampleSum.cpp
@@ -26,7 +26,7 @@ public:
 		{
 			return true;
 		}
-		else if( IsFeatureLevelSupported(Platform, ERHIFeatureLevel::SM3) )
+		else if( IsFeatureLevelSupported(Platform, ERHIFeatureLevel::SM4) )
 		{
 			return NumSamples <= 16;
 		}
@@ -617,7 +617,7 @@ uint32 FRCPassPostProcessWeightedSampleSum::GetMaxNumSamples(ERHIFeatureLevel::T
 	{
 		MaxNumSamples = MAX_FILTER_SAMPLES;
 	}
-	else if (InFeatureLevel >= ERHIFeatureLevel::SM3)
+	else if (InFeatureLevel >= ERHIFeatureLevel::SM4)
 	{
 		MaxNumSamples = 16;
 	}
diff --git a/Engine/Source/Runtime/Renderer/Private/PostProcess/PostProcessWeightedSampleSum.h b/Engine/Source/Runtime/Renderer/Private/PostProcess/PostProcessWeightedSampleSum.h
--- a/Engine/Source/Runtime/Renderer/Private/PostProcess/PostProcessWeightedSampleSum.h
+++ b/Engine/Source/Runtime/Renderer/Private/PostProcess/PostProcessWeightedSampleSum.h
@@ -26,7 +26,7 @@ public:
 		{
 			return true;
 		}
-		else if (IsFeatureLevelSupported(Platform, ERHIFeatureLevel::SM3))
+		else if (IsFeatureLevelSupported(Platform, ERHIFeatureLevel::SM4))
 		{
 			return NumSamples <= 16;
 		}
diff --git a/Engine/Source/Runtime/Renderer/Private/SceneOcclusion.cpp b/Engine/Source/Runtime/Renderer/Private/SceneOcclusion.cpp
--- a/Engine/Source/Runtime/Renderer/Private/SceneOcclusion.cpp
+++ b/Engine/Source/Runtime/Renderer/Private/SceneOcclusion.cpp
@@ -518,7 +518,7 @@ void FHZBOcclusionTester::SetInvalidFrameNumber()
 
 void FHZBOcclusionTester::InitDynamicRHI()
 {
-	if (GetFeatureLevel() >= ERHIFeatureLevel::SM3)
+	if (GetFeatureLevel() >= ERHIFeatureLevel::SM4)
 	{
 #if PLATFORM_MAC // Workaround radr://16096028 Texture Readback via glReadPixels + PBOs stalls on Nvidia GPUs
 		FPooledRenderTargetDesc Desc( FPooledRenderTargetDesc::Create2DDesc( FIntPoint( SizeX, SizeY ), PF_R8G8B8A8, TexCreate_CPUReadback | TexCreate_HideInVisualizeTexture, TexCreate_None, false ) );
@@ -531,7 +531,7 @@ void FHZBOcclusionTester::InitDynamicRHI()
 
 void FHZBOcclusionTester::ReleaseDynamicRHI()
 {
-	if (GetFeatureLevel() >= ERHIFeatureLevel::SM3)
+	if (GetFeatureLevel() >= ERHIFeatureLevel::SM4)
 	{
 		GRenderTargetPool.FreeUnusedResource( ResultsTextureCPU );
 	}
@@ -609,7 +609,7 @@ class FHZBTestPS : public FGlobalShader
 
 	static bool ShouldCache(EShaderPlatform Platform)
 	{
-		return IsFeatureLevelSupported(Platform, ERHIFeatureLevel::SM3);
+		return IsFeatureLevelSupported(Platform, ERHIFeatureLevel::SM4);
 	}
 
 	static void ModifyCompilationEnvironment(EShaderPlatform Platform, FShaderCompilerEnvironment& OutEnvironment)
@@ -782,7 +782,7 @@ class THZBBuildPS : public FGlobalShader
 
 	static bool ShouldCache(EShaderPlatform Platform)
 	{
-		return IsFeatureLevelSupported(Platform, ERHIFeatureLevel::SM3);
+		return IsFeatureLevelSupported(Platform, ERHIFeatureLevel::SM4);
 	}
 
 	static void ModifyCompilationEnvironment(EShaderPlatform Platform, FShaderCompilerEnvironment& OutEnvironment)
@@ -989,11 +989,11 @@ void FDeferredShadingSceneRenderer::BeginOcclusionTests(FRHICommandListImmediate
 		{
 			RHICmdList.SetViewport(View.ViewRect.Min.X, View.ViewRect.Min.Y, 0.0f, View.ViewRect.Max.X, View.ViewRect.Max.Y, 1.0f);
 		}
-    
-	    FSceneViewState* ViewState = (FSceneViewState*)View.State;
-    
-	    if(ViewState && !View.bDisableQuerySubmissions)
-	    {
+	
+		FSceneViewState* ViewState = (FSceneViewState*)View.State;
+	
+		if(ViewState && !View.bDisableQuerySubmissions)
+		{
 			// Depth tests, no depth writes, no color writes, opaque, solid rasterization wo/ backface culling.
 			// Note, this is a reversed Z depth surface, using CF_GreaterEqual.
 			RHICmdList.SetDepthStencilState(TStaticDepthStencilState<false, CF_GreaterEqual>::GetRHI());
@@ -1089,12 +1089,12 @@ void FDeferredShadingSceneRenderer::BeginOcclusionTests(FRHICommandListImmediate
 					}
 				}
 			}
-    
-		    // Don't do primitive occlusion if we have a view parent or are frozen.
-    #if !(UE_BUILD_SHIPPING || UE_BUILD_TEST)
-		    if ( !ViewState->HasViewParent() && !ViewState->bIsFrozen )
-    #endif
-		    {
+	
+			// Don't do primitive occlusion if we have a view parent or are frozen.
+	#if !(UE_BUILD_SHIPPING || UE_BUILD_TEST)
+			if ( !ViewState->HasViewParent() && !ViewState->bIsFrozen )
+	#endif
+			{
 				VertexShader->SetParameters(RHICmdList, View);
 
 				{
@@ -1105,9 +1105,9 @@ void FDeferredShadingSceneRenderer::BeginOcclusionTests(FRHICommandListImmediate
 					SCOPED_DRAW_EVENT(GroupedQueries, DEC_SCENE_ITEMS);
 					View.GroupedOcclusionQueries.Flush(RHICmdList);
 				}
-		    }
-	    }
-    }
+			}
+		}
+	}
 
 	for( int32 ViewIndex = 0; ViewIndex < Views.Num(); ViewIndex++ )
 	{
diff --git a/Engine/Source/Runtime/Renderer/Private/SceneOcclusion.h b/Engine/Source/Runtime/Renderer/Private/SceneOcclusion.h
--- a/Engine/Source/Runtime/Renderer/Private/SceneOcclusion.h
+++ b/Engine/Source/Runtime/Renderer/Private/SceneOcclusion.h
@@ -13,7 +13,7 @@ class FOcclusionQueryVS : public FGlobalShader
 {
 	DECLARE_SHADER_TYPE(FOcclusionQueryVS,Global);
 public:
-	static bool ShouldCache(EShaderPlatform Platform) { return IsFeatureLevelSupported(Platform, ERHIFeatureLevel::SM3); }
+	static bool ShouldCache(EShaderPlatform Platform) { return IsFeatureLevelSupported(Platform, ERHIFeatureLevel::SM4); }
 
 	FOcclusionQueryVS(const ShaderMetaType::CompiledShaderInitializerType& Initializer):
 		FGlobalShader(Initializer)
diff --git a/Engine/Source/Runtime/Renderer/Private/ShaderBaseClasses.cpp b/Engine/Source/Runtime/Renderer/Private/ShaderBaseClasses.cpp
--- a/Engine/Source/Runtime/Renderer/Private/ShaderBaseClasses.cpp
+++ b/Engine/Source/Runtime/Renderer/Private/ShaderBaseClasses.cpp
@@ -213,7 +213,7 @@ void FMaterialShader::SetParameters(
 
 	AtmosphericFogTextureParameters.Set(RHICmdList, ShaderRHI, View);
 
-	if (FeatureLevel >= ERHIFeatureLevel::SM3)
+	if (FeatureLevel >= ERHIFeatureLevel::SM4)
 	{
 		// for copied scene color
 		if(LightAttenuation.IsBound())
diff --git a/Engine/Source/Runtime/Renderer/Private/SystemTextures.cpp b/Engine/Source/Runtime/Renderer/Private/SystemTextures.cpp
--- a/Engine/Source/Runtime/Renderer/Private/SystemTextures.cpp
+++ b/Engine/Source/Runtime/Renderer/Private/SystemTextures.cpp
@@ -204,7 +204,7 @@ void FSystemTextures::InitializeTextures(FRHICommandListImmediate& RHICmdList, E
 	}
 
 	// Create the SSAO randomization texture
-	if (FeatureLevelInitializedTo < ERHIFeatureLevel::SM3 && InFeatureLevel >= ERHIFeatureLevel::SM3)
+	if (FeatureLevelInitializedTo < ERHIFeatureLevel::SM4 && InFeatureLevel >= ERHIFeatureLevel::SM4)
 	{
 		float g_AngleOff1 = 127;
 		float g_AngleOff2 = 198;
diff --git a/Engine/Source/Runtime/Renderer/Private/TranslucentRendering.cpp b/Engine/Source/Runtime/Renderer/Private/TranslucentRendering.cpp
--- a/Engine/Source/Runtime/Renderer/Private/TranslucentRendering.cpp
+++ b/Engine/Source/Runtime/Renderer/Private/TranslucentRendering.cpp
@@ -139,7 +139,7 @@ class FCopySceneColorPS : public FGlobalShader
 	DECLARE_SHADER_TYPE(FCopySceneColorPS,Global);
 public:
 
-	static bool ShouldCache(EShaderPlatform Platform) { return IsFeatureLevelSupported(Platform, ERHIFeatureLevel::SM3); }
+	static bool ShouldCache(EShaderPlatform Platform) { return IsFeatureLevelSupported(Platform, ERHIFeatureLevel::SM4); }
 
 	FCopySceneColorPS(const ShaderMetaType::CompiledShaderInitializerType& Initializer):
 		FGlobalShader(Initializer)
@@ -739,7 +739,7 @@ void FTranslucentPrimSet::AddScenePrimitive(FPrimitiveSceneInfo* PrimitiveSceneI
 	const auto FeatureLevel = ViewInfo.GetFeatureLevel();
 
 	if(bUseSeparateTranslucency 
-		&& FeatureLevel >= ERHIFeatureLevel::SM3)
+		&& FeatureLevel >= ERHIFeatureLevel::SM4)
 	{
 		// add to list of translucent prims that use scene color
 		new(SortedSeparateTranslucencyPrims) FSortedPrim(PrimitiveSceneInfo,SortKey,PrimitiveSceneInfo->Proxy->GetTranslucencySortPriority());
@@ -747,7 +747,7 @@ void FTranslucentPrimSet::AddScenePrimitive(FPrimitiveSceneInfo* PrimitiveSceneI
 
 	if (bUseNormalTranslucency 
 		// Force separate translucency to be rendered normally if the feature level does not support separate translucency
-		|| (bUseSeparateTranslucency && FeatureLevel < ERHIFeatureLevel::SM3))
+		|| (bUseSeparateTranslucency && FeatureLevel < ERHIFeatureLevel::SM4))
 	{
 		// add to list of translucent prims
 		new(SortedPrims) FSortedPrim(PrimitiveSceneInfo,SortKey,PrimitiveSceneInfo->Proxy->GetTranslucencySortPriority());
diff --git a/Engine/Source/Runtime/Windows/D3D11RHI/Private/D3D11Device.cpp b/Engine/Source/Runtime/Windows/D3D11RHI/Private/D3D11Device.cpp
--- a/Engine/Source/Runtime/Windows/D3D11RHI/Private/D3D11Device.cpp
+++ b/Engine/Source/Runtime/Windows/D3D11RHI/Private/D3D11Device.cpp
@@ -79,7 +79,7 @@ FD3D11DynamicRHI::FD3D11DynamicRHI(IDXGIFactory* InDXGIFactory,D3D_FEATURE_LEVEL
 	}
 
 	// ES2 feature level emulation in D3D11
-	if (FParse::Param(FCommandLine::Get(), TEXT("FeatureLevelES2")))
+	if (FParse::Param(FCommandLine::Get(), TEXT("FeatureLevelES2")) && !GIsEditor)
 	{
 		GMaxRHIFeatureLevelValue = GRHIFeatureLevelValue = ERHIFeatureLevel::ES2;
 		GRHIShaderPlatformValue = SP_PCD3D_ES2;
diff --git a/Engine/Source/Runtime/Windows/D3D11RHI/Private/D3D11RHI.cpp b/Engine/Source/Runtime/Windows/D3D11RHI/Private/D3D11RHI.cpp
--- a/Engine/Source/Runtime/Windows/D3D11RHI/Private/D3D11RHI.cpp
+++ b/Engine/Source/Runtime/Windows/D3D11RHI/Private/D3D11RHI.cpp
@@ -172,7 +172,7 @@ public:
 
 	static bool ShouldCache(EShaderPlatform Platform)
 	{
-		return IsFeatureLevelSupported(Platform, ERHIFeatureLevel::SM3);
+		return IsFeatureLevelSupported(Platform, ERHIFeatureLevel::SM4);
 	}
 };
 
@@ -182,7 +182,7 @@ FGlobalBoundShaderState LongGPUTaskBoundShaderState;
 
 void FD3D11DynamicRHI::IssueLongGPUTask()
 {
-	if (GMaxRHIFeatureLevel >= ERHIFeatureLevel::SM3)
+	if (GMaxRHIFeatureLevel >= ERHIFeatureLevel::SM4)
 	{
 		int32 LargestViewportIndex = INDEX_NONE;
 		int32 LargestViewportPixels = 0;
