UE4 - Safe integration SM3->SM4

diff --git a/Engine/Shaders/BasePassForForwardShadingCommon.usf b/Engine/Shaders/BasePassForForwardShadingCommon.usf
--- a/Engine/Shaders/BasePassForForwardShadingCommon.usf
+++ b/Engine/Shaders/BasePassForForwardShadingCommon.usf
@@ -11,7 +11,7 @@
 
 #define USE_VERTEX_FOG					(!OUTPUT_GAMMA_SPACE)
 
-#define PACK_INTERPOLANTS (USE_VERTEX_FOG && NUM_VF_PACKED_INTERPOLANTS > 0 && ES2_PROFILE)
+#define PACK_INTERPOLANTS (USE_VERTEX_FOG && NUM_VF_PACKED_INTERPOLANTS > 0 && (ES2_PROFILE || ES3_1_PROFILE))
 #define LANDSCAPE_BUG_WORKAROUND (IS_FORWARD_BASEPASS_VERTEX_SHADER && PACK_INTERPOLANTS)
 
 #if METAL_PROFILE 
diff --git a/Engine/Shaders/LightmapCommon.usf b/Engine/Shaders/LightmapCommon.usf
--- a/Engine/Shaders/LightmapCommon.usf
+++ b/Engine/Shaders/LightmapCommon.usf
@@ -140,6 +140,11 @@ half GetPrimaryPrecomputedShadowMask(FVertexFactoryInterpolantsVSToPS Interpolan
 		// This is necessary because objects inside a light's influence that were determined to be completely shadowed won't be rendered with STATICLIGHTING_TEXTUREMASK==1
 		return 0;
 
+	#elif MOVABLE_DIRECTIONAL_LIGHT
+		//@mw todo
+		// TODO: Filter shadowmap
+		return 1;
+
 	#else
 
 		// Mark as unshadowed by default
diff --git a/Engine/Shaders/ShaderVersion.usf b/Engine/Shaders/ShaderVersion.usf
--- a/Engine/Shaders/ShaderVersion.usf
+++ b/Engine/Shaders/ShaderVersion.usf
@@ -2,4 +2,4 @@
 // Each time the console command is executed it generates a new GUID. As this file is included
 // in common.usf (which should be included in any shader) it allows to invalidate the shader DDC.
 // 
-// GUID = B08E107742F276154CF002A94C779319
\ No newline at end of file
+// GUID = A7702884495089DC221506B8C4961E6C1
diff --git a/Engine/Shaders/ShadowFilteringCommon.usf b/Engine/Shaders/ShadowFilteringCommon.usf
new file mode 100644
--- /dev/null
+++ b/Engine/Shaders/ShadowFilteringCommon.usf
@@ -0,0 +1,443 @@
+// Copyright 1998-2014 Epic Games, Inc. All Rights Reserved.
+
+/*=============================================================================
+	ShadowFilteringCommon.usf: Contains functions to filter a shadowmap, shared between forward/deferred shading.
+=========================================================================*/
+
+struct FPCFSamplerSettings
+{
+	Texture2D		ShadowDepthTexture;
+	SamplerState	ShadowDepthTextureSampler;
+
+	//XY - Pixel size of shadowmap
+	//ZW - Inverse pixel size of shadowmap
+	float4			ShadowBufferSize;
+
+	// SceneDepth in lightspace.
+	float			SceneDepth;
+
+	float			TransitionScale;
+
+	// set by the caller, constant for the code so only one code branch should be compiled
+	bool			bSubsurface;
+
+	// only used if bSubsurface is true
+	float			DensityMulConstant;
+
+	// only used if bSubsurface is true
+	float2			ProjectionDepthBiasParameters;
+};
+
+// linear PCF, input 3x3
+// @param Values0 in row 0 from left to right: x,y,z,w
+// @param Values1 in row 1 from left to right: x,y,z,w
+// @param Values2 in row 2 from left to right: x,y,z,w
+// can be optimized
+float PCF2x2(float2 Fraction, float3 Values0, float3 Values1, float3 Values2)
+{
+	float2 VerticalLerp00 = lerp(float2(Values0.x, Values1.x), float2(Values0.y, Values1.y), Fraction.xx);
+	float PCFResult00 = lerp(VerticalLerp00.x, VerticalLerp00.y, Fraction.y);
+
+	float2 VerticalLerp10 = lerp(float2(Values0.y, Values1.y), float2(Values0.z, Values1.z), Fraction.xx);
+	float PCFResult10 = lerp(VerticalLerp10.x, VerticalLerp10.y, Fraction.y);
+
+	float2 VerticalLerp01 = lerp(float2(Values1.x, Values2.x), float2(Values1.y, Values2.y), Fraction.xx);
+	float PCFResult01 = lerp(VerticalLerp01.x, VerticalLerp01.y, Fraction.y);
+
+	float2 VerticalLerp11 = lerp(float2(Values1.y, Values2.y), float2(Values1.z, Values2.z), Fraction.xx);
+	float PCFResult11 = lerp(VerticalLerp11.x, VerticalLerp11.y, Fraction.y);	
+
+	return saturate((PCFResult00 + PCFResult10 + PCFResult01 + PCFResult11) * 0.25f);
+}
+
+// linear PCF, input 4x4
+// @param Values0 in row 0 from left to right: x,y,z,w
+// @param Values1 in row 1 from left to right: x,y,z,w
+// @param Values2 in row 2 from left to right: x,y,z,w
+// @param Values3 in row 3 from left to right: x,y,z,w
+// can be optimized
+float PCF3x3(float2 Fraction, float4 Values0, float4 Values1, float4 Values2, float4 Values3)
+{
+	float2 VerticalLerp00 = lerp(float2(Values0.x, Values1.x), float2(Values0.y, Values1.y), Fraction.xx);
+	float PCFResult00 = lerp(VerticalLerp00.x, VerticalLerp00.y, Fraction.y);
+
+	float2 VerticalLerp10 = lerp(float2(Values0.y, Values1.y), float2(Values0.z, Values1.z), Fraction.xx);
+	float PCFResult10 = lerp(VerticalLerp10.x, VerticalLerp10.y, Fraction.y);
+
+	float2 VerticalLerp20 = lerp(float2(Values0.z, Values1.z), float2(Values0.w, Values1.w), Fraction.xx);
+	float PCFResult20 = lerp(VerticalLerp20.x, VerticalLerp20.y, Fraction.y);
+
+	float2 VerticalLerp01 = lerp(float2(Values1.x, Values2.x), float2(Values1.y, Values2.y), Fraction.xx);
+	float PCFResult01 = lerp(VerticalLerp01.x, VerticalLerp01.y, Fraction.y);
+
+	float2 VerticalLerp11 = lerp(float2(Values1.y, Values2.y), float2(Values1.z, Values2.z), Fraction.xx);
+	float PCFResult11 = lerp(VerticalLerp11.x, VerticalLerp11.y, Fraction.y);
+
+	float2 VerticalLerp21 = lerp(float2(Values1.z, Values2.z), float2(Values1.w, Values2.w), Fraction.xx);
+	float PCFResult21 = lerp(VerticalLerp21.x, VerticalLerp21.y, Fraction.y);
+
+	float2 VerticalLerp02 = lerp(float2(Values2.x, Values3.x), float2(Values2.y, Values3.y), Fraction.xx);
+	float PCFResult02 = lerp(VerticalLerp02.x, VerticalLerp02.y, Fraction.y);
+
+	float2 VerticalLerp12 = lerp(float2(Values2.y, Values3.y), float2(Values2.z, Values3.z), Fraction.xx);
+	float PCFResult12 = lerp(VerticalLerp12.x, VerticalLerp12.y, Fraction.y);
+
+	float2 VerticalLerp22 = lerp(float2(Values2.z, Values3.z), float2(Values2.w, Values3.w), Fraction.xx);
+	float PCFResult22 = lerp(VerticalLerp22.x, VerticalLerp22.y, Fraction.y);
+
+	return saturate((PCFResult00 + PCFResult10 + PCFResult20 + PCFResult01 + PCFResult11 + PCFResult21 + PCFResult02 + PCFResult12 + PCFResult22) * .11111f);
+}
+
+
+// linear PCF, input 4x4
+// using Gather: xyzw in counter clockwise order starting with the sample to the lower left of the queried location
+// @param Values0 left top
+// @param Values1 right top
+// @param Values2 left bottom
+// @param Values3 right bottom
+float PCF3x3gather(float2 Fraction, float4 Values0, float4 Values1, float4 Values2, float4 Values3)
+{
+	// mostly vectorized
+	float4 PCFResult0 = float4(lerp(Values0.wx, Values0.zy, Fraction.xx), lerp(Values2.wx, Values2.zy, Fraction.x));
+	float4 PCFResult1 = float4(lerp(Values0.zy, Values1.wx, Fraction.xx), lerp(Values2.zy, Values3.wx, Fraction.x));
+	float4 PCFResult2 = float4(lerp(Values1.wx, Values1.zy, Fraction.xx), lerp(Values3.wx, Values3.zy, Fraction.x));
+
+	float4 PCFResult012 = PCFResult0 + PCFResult1 + PCFResult2;
+
+	float3 PCFResults = lerp(PCFResult012.xyz, PCFResult012.yzw, Fraction.y);
+
+	return dot(PCFResults, 1.0f / 9.0f);
+}
+
+float PCF5x2(float2 Fraction, float4 Values00, float4 Values20, float4 Values40, float4 Values02, float4 Values22, float4 Values42)
+{
+	float4 Results0;
+	float2 Results1;
+	float2 HorizontalLerp00 = lerp(float2(Values00.w, Values00.x), float2(Values00.z, Values00.y), Fraction.xx);
+	Results0.x = lerp(HorizontalLerp00.x, HorizontalLerp00.y, Fraction.y);
+
+	float2 HorizontalLerp10 = lerp(float2(Values00.z, Values00.y), float2(Values20.w, Values20.x), Fraction.xx);
+	Results0.y = lerp(HorizontalLerp10.x, HorizontalLerp10.y, Fraction.y);
+
+	float2 HorizontalLerp20 = lerp(float2(Values20.w, Values20.x), float2(Values20.z, Values20.y), Fraction.xx);
+	Results0.z = lerp(HorizontalLerp20.x, HorizontalLerp20.y, Fraction.y);
+
+	float2 HorizontalLerp30 = lerp(float2(Values20.z, Values20.y), float2(Values40.w, Values40.x), Fraction.xx);
+	Results0.w = lerp(HorizontalLerp30.x, HorizontalLerp30.y, Fraction.y);
+
+	float2 HorizontalLerp40 = lerp(float2(Values40.w, Values40.x), float2(Values40.z, Values40.y), Fraction.xx);
+	Results1.x = lerp(HorizontalLerp40.x, HorizontalLerp40.y, Fraction.y);
+
+	float4 Results2;
+	float2 HorizontalLerp01 = lerp(float2(Values00.x, Values02.w), float2(Values00.y, Values02.z), Fraction.xx);
+	Results2.x = lerp(HorizontalLerp01.x, HorizontalLerp01.y, Fraction.y);
+
+	float2 HorizontalLerp11 = lerp(float2(Values00.y, Values02.z), float2(Values20.x, Values22.w), Fraction.xx);
+	Results2.y = lerp(HorizontalLerp11.x, HorizontalLerp11.y, Fraction.y);
+
+	float2 HorizontalLerp21 = lerp(float2(Values20.x, Values22.w), float2(Values20.y, Values22.z), Fraction.xx);
+	Results2.z = lerp(HorizontalLerp21.x, HorizontalLerp21.y, Fraction.y);
+
+	float2 HorizontalLerp31 = lerp(float2(Values20.y, Values22.z), float2(Values40.x, Values42.w), Fraction.xx);
+	Results2.w = lerp(HorizontalLerp31.x, HorizontalLerp31.y, Fraction.y);
+
+	float2 HorizontalLerp41 = lerp(float2(Values40.x, Values42.w), float2(Values40.y, Values42.z), Fraction.xx);
+	Results1.y = lerp(HorizontalLerp41.x, HorizontalLerp41.y, Fraction.y);
+
+	return dot(Results0, .04f) + dot(Results1, .04f) + dot(Results2, .04f);
+}
+
+float PCF5x1(float2 Fraction, float4 Values00, float4 Values20, float4 Values40)
+{
+	float4 Results0;
+	float2 HorizontalLerp00 = lerp(float2(Values00.w, Values00.x), float2(Values00.z, Values00.y), Fraction.xx);
+	Results0.x = lerp(HorizontalLerp00.x, HorizontalLerp00.y, Fraction.y);
+
+	float2 HorizontalLerp10 = lerp(float2(Values00.z, Values00.y), float2(Values20.w, Values20.x), Fraction.xx);
+	Results0.y = lerp(HorizontalLerp10.x, HorizontalLerp10.y, Fraction.y);
+
+	float2 HorizontalLerp20 = lerp(float2(Values20.w, Values20.x), float2(Values20.z, Values20.y), Fraction.xx);
+	Results0.z = lerp(HorizontalLerp20.x, HorizontalLerp20.y, Fraction.y);
+
+	float2 HorizontalLerp30 = lerp(float2(Values20.z, Values20.y), float2(Values40.w, Values40.x), Fraction.xx);
+	Results0.w = lerp(HorizontalLerp30.x, HorizontalLerp30.y, Fraction.y);
+
+	float2 HorizontalLerp40 = lerp(float2(Values40.w, Values40.x), float2(Values40.z, Values40.y), Fraction.xx);
+	float Results1 = lerp(HorizontalLerp40.x, HorizontalLerp40.y, Fraction.y);
+
+	return dot(Results0, .04f) + Results1 * .04f;
+}
+
+// lowest quality ith PCF
+float PCF1x1(float2 Fraction, float4 Values00)
+{
+	float2 HorizontalLerp00 = lerp(Values00.wx, Values00.zy, Fraction.xx);
+
+	return lerp(HorizontalLerp00.x, HorizontalLerp00.y, Fraction.y);
+}
+
+float4 CalculateOcclusion(float4 ShadowmapDepth, FPCFSamplerSettings Settings)
+{
+	if (Settings.bSubsurface)
+	{
+		// Determine the distance that the light traveled through the subsurface object
+		// This assumes that anything between this subsurface pixel and the light was also a subsurface material,
+		// As a result, subsurface materials receive leaked light based on how transparent they are
+		float4 Thickness = max(Settings.SceneDepth - ShadowmapDepth, 0);
+		float4 Occlusion = saturate(exp(-Thickness * Settings.DensityMulConstant));
+		// Never shadow from depths that were never written to (max depth value)
+		return ShadowmapDepth > .99f ? 1 : Occlusion;
+	}
+	else
+	{
+		// The standard comparison is SceneDepth < ShadowmapDepth
+		// Using a soft transition based on depth difference
+		// Offsets shadows a bit but reduces self shadowing artifacts considerably
+		float TransitionScale = Settings.TransitionScale; //SoftTransitionScale.z;
+		return saturate((ShadowmapDepth - Settings.SceneDepth) * TransitionScale + 1);
+	}
+}
+
+float3 CalculateOcclusion(float3 ShadowmapDepth, FPCFSamplerSettings Settings)
+{
+	if (Settings.bSubsurface)
+	{
+		// Determine the distance that the light traveled through the subsurface object
+		// This assumes that anything between this subsurface pixel and the light was also a subsurface material,
+		// As a result, subsurface materials receive leaked light based on how transparent they are
+		float3 Thickness = max(Settings.SceneDepth - (ShadowmapDepth - Settings.ProjectionDepthBiasParameters.x), 0);
+		float3 Occlusion = saturate(exp(-Thickness * Settings.DensityMulConstant));
+		// Never shadow from depths that were never written to (max depth value)
+		return ShadowmapDepth > .99f ? 1 : Occlusion;
+	}
+	else
+	{
+		// The standard comparison is Settings.SceneDepth < ShadowmapDepth
+		// Using a soft transition based on depth difference
+		// Offsets shadows a bit but reduces self shadowing artifacts considerably
+		float TransitionScale = Settings.TransitionScale; //SoftTransitionScale.z;
+		return saturate((ShadowmapDepth - Settings.SceneDepth) * TransitionScale + 1);
+	}
+}
+
+void FetchRowOfThree(float2 Sample00TexelCenter, float VerticalOffset, out float3 Values0, FPCFSamplerSettings Settings)
+{
+	Values0.x = Texture2DSampleLevel(Settings.ShadowDepthTexture, Settings.ShadowDepthTextureSampler, (Sample00TexelCenter + float2(0, VerticalOffset)) * Settings.ShadowBufferSize.zw, 0).r;
+	Values0.y = Texture2DSampleLevel(Settings.ShadowDepthTexture, Settings.ShadowDepthTextureSampler, (Sample00TexelCenter + float2(1, VerticalOffset)) * Settings.ShadowBufferSize.zw, 0).r;
+	Values0.z = Texture2DSampleLevel(Settings.ShadowDepthTexture, Settings.ShadowDepthTextureSampler, (Sample00TexelCenter + float2(2, VerticalOffset)) * Settings.ShadowBufferSize.zw, 0).r;	
+	Values0 = CalculateOcclusion(Values0, Settings);
+}
+
+void FetchRowOfFour(float2 Sample00TexelCenter, float VerticalOffset, out float4 Values0, FPCFSamplerSettings Settings)
+{
+	Values0.x = Texture2DSampleLevel(Settings.ShadowDepthTexture, Settings.ShadowDepthTextureSampler, (Sample00TexelCenter + float2(0, VerticalOffset)) * Settings.ShadowBufferSize.zw, 0).r;
+	Values0.y = Texture2DSampleLevel(Settings.ShadowDepthTexture, Settings.ShadowDepthTextureSampler, (Sample00TexelCenter + float2(1, VerticalOffset)) * Settings.ShadowBufferSize.zw, 0).r;
+	Values0.z = Texture2DSampleLevel(Settings.ShadowDepthTexture, Settings.ShadowDepthTextureSampler, (Sample00TexelCenter + float2(2, VerticalOffset)) * Settings.ShadowBufferSize.zw, 0).r;
+	Values0.w = Texture2DSampleLevel(Settings.ShadowDepthTexture, Settings.ShadowDepthTextureSampler, (Sample00TexelCenter + float2(3, VerticalOffset)) * Settings.ShadowBufferSize.zw, 0).r;
+	Values0 = CalculateOcclusion(Values0, Settings);
+}
+
+void FetchRowOfThreeAfterFour(float2 Sample00TexelCenter, float VerticalOffset, out float3 Values1, FPCFSamplerSettings Settings)
+{
+	Values1.x = Texture2DSampleLevel(Settings.ShadowDepthTexture, Settings.ShadowDepthTextureSampler, (Sample00TexelCenter + float2(4, VerticalOffset)) * Settings.ShadowBufferSize.zw, 0).r;
+	Values1.y = Texture2DSampleLevel(Settings.ShadowDepthTexture, Settings.ShadowDepthTextureSampler, (Sample00TexelCenter + float2(5, VerticalOffset)) * Settings.ShadowBufferSize.zw, 0).r;
+	Values1.z = Texture2DSampleLevel(Settings.ShadowDepthTexture, Settings.ShadowDepthTextureSampler, (Sample00TexelCenter + float2(6, VerticalOffset)) * Settings.ShadowBufferSize.zw, 0).r;
+	Values1 = CalculateOcclusion(Values1, Settings);
+}
+
+// break this out for forward rendering as it's not part of ManualPCF's set of shadowquality settings.
+float Manual2x2PCF(float2 ShadowPosition, FPCFSamplerSettings Settings)
+{
+	float2 TexelPos = ShadowPosition * Settings.ShadowBufferSize.xy - 0.5f;	// bias to be consistent with texture filtering hardware
+	float2 Fraction = frac(TexelPos);
+	float2 TexelCenter = floor(TexelPos) + 0.5f;	// bias to get reliable texel center content
+	
+	float3 SamplesValues0, SamplesValues1, SamplesValues2;
+
+	FetchRowOfThree(TexelCenter, 0, SamplesValues0, Settings);
+	FetchRowOfThree(TexelCenter, 1, SamplesValues1, Settings);
+	FetchRowOfThree(TexelCenter, 2, SamplesValues2, Settings);
+	
+	return PCF2x2(Fraction, SamplesValues0, SamplesValues1, SamplesValues2);
+}
+
+float ManualNoFiltering(float2 ShadowPosition, FPCFSamplerSettings Settings)
+{
+	// very low quality but very good performance, useful to profile, 1 sample, not using gather4
+	return CalculateOcclusion(Texture2DSampleLevel(Settings.ShadowDepthTexture, Settings.ShadowDepthTextureSampler, ShadowPosition, 0).rrr, Settings).r;
+}
+
+float Manual1x1PCF(float2 ShadowPosition, FPCFSamplerSettings Settings)
+{
+	float2 TexelPos = ShadowPosition * Settings.ShadowBufferSize.xy - 0.5f;	// bias to be consistent with texture filtering hardware
+	float2 Fraction = frac(TexelPos);
+	float2 TexelCenter = floor(TexelPos) + 0.5f;	// bias to get reliable texel center content
+
+	// using Gather: xyzw in counter clockwise order starting with the sample to the lower left of the queried location
+	float4 Samples;
+
+#if FEATURE_GATHER4
+	Samples = Settings.ShadowDepthTexture.Gather(Settings.ShadowDepthTextureSampler, TexelCenter * Settings.ShadowBufferSize.zw);
+#else
+	Samples.x = Texture2DSampleLevel(Settings.ShadowDepthTexture, Settings.ShadowDepthTextureSampler, (TexelCenter.xy + float2(0, 1)) * Settings.ShadowBufferSize.zw, 0).r;
+	Samples.y = Texture2DSampleLevel(Settings.ShadowDepthTexture, Settings.ShadowDepthTextureSampler, (TexelCenter.xy + float2(1, 1)) * Settings.ShadowBufferSize.zw, 0).r;
+	Samples.z = Texture2DSampleLevel(Settings.ShadowDepthTexture, Settings.ShadowDepthTextureSampler, (TexelCenter.xy + float2(1, 0)) * Settings.ShadowBufferSize.zw, 0).r;
+	Samples.w = Texture2DSampleLevel(Settings.ShadowDepthTexture, Settings.ShadowDepthTextureSampler, (TexelCenter.xy + float2(0, 0)) * Settings.ShadowBufferSize.zw, 0).r;
+#endif
+
+	float4 Values00 = CalculateOcclusion(Samples, Settings);
+	return PCF1x1(Fraction, Values00);
+}
+
+float ManualPCF(float2 ShadowPosition, FPCFSamplerSettings Settings)
+{
+
+#if SHADOW_QUALITY == 1
+	return ManualNoFiltering(ShadowPosition, Settings);
+#endif
+
+#if SHADOW_QUALITY == 2
+	// low quality, 2x2 samples, using and not using gather4
+	return Manual1x1PCF(ShadowPosition, Settings);
+#endif
+
+#if SHADOW_QUALITY == 3
+	// medium quality, 4x4 samples, using and not using gather4
+	{
+		float2 TexelPos = ShadowPosition * Settings.ShadowBufferSize.xy - 0.5f;	// bias to be consistent with texture filtering hardware
+		float2 Fraction = frac(TexelPos);
+		float2 TexelCenter = floor(TexelPos) + 0.5f;	// bias to get reliable texel center content
+		{
+			float2 Sample00TexelCenter = TexelCenter - float2(1, 1);
+
+			float4 SampleValues0, SampleValues1, SampleValues2, SampleValues3;
+
+#if FEATURE_GATHER4
+			float2 SamplePos = TexelCenter * Settings.ShadowBufferSize.zw;	// bias to get reliable texel center content
+			SampleValues0 = CalculateOcclusion(Settings.ShadowDepthTexture.Gather(ShadowDepthTextureSampler, SamplePos, int2(-1, -1)), Settings);
+			SampleValues1 = CalculateOcclusion(Settings.ShadowDepthTexture.Gather(ShadowDepthTextureSampler, SamplePos, int2(1, -1)), Settings);
+			SampleValues2 = CalculateOcclusion(Settings.ShadowDepthTexture.Gather(ShadowDepthTextureSampler, SamplePos, int2(-1, 1)), Settings);
+			SampleValues3 = CalculateOcclusion(Settings.ShadowDepthTexture.Gather(ShadowDepthTextureSampler, SamplePos, int2(1, 1)), Settings);
+			return PCF3x3gather(Fraction, SampleValues0, SampleValues1, SampleValues2, SampleValues3);
+#else // FEATURE_GATHER4 
+			FetchRowOfFour(Sample00TexelCenter, 0, SampleValues0, Settings);
+			FetchRowOfFour(Sample00TexelCenter, 1, SampleValues1, Settings);
+			FetchRowOfFour(Sample00TexelCenter, 2, SampleValues2, Settings);
+			FetchRowOfFour(Sample00TexelCenter, 3, SampleValues3, Settings);
+			return PCF3x3(Fraction, SampleValues0, SampleValues1, SampleValues2, SampleValues3);
+#endif // FEATURE_GATHER4
+		}
+	}
+#endif
+
+#if FEATURE_GATHER4
+	// high quality, 6x6 samples, using gather4
+	{
+		float2 TexelPos = ShadowPosition * Settings.ShadowBufferSize.xy - 0.5f;	// bias to be consistent with texture filtering hardware
+		float2 Fraction = frac(TexelPos);
+		float2 TexelCenter = floor(TexelPos);
+		float2 SamplePos = (TexelCenter + 0.5f) * Settings.ShadowBufferSize.zw;	// bias to get reliable texel center content
+
+		float3 Results;
+
+		float4 Values00 = CalculateOcclusion(Settings.ShadowDepthTexture.Gather(Settings.ShadowDepthTextureSampler, SamplePos, int2(-2, -2)), Settings);
+		float4 Values20 = CalculateOcclusion(Settings.ShadowDepthTexture.Gather(Settings.ShadowDepthTextureSampler, SamplePos, int2(0, -2)), Settings);
+		float4 Values40 = CalculateOcclusion(Settings.ShadowDepthTexture.Gather(Settings.ShadowDepthTextureSampler, SamplePos, int2(2, -2)), Settings);
+
+		float4 Values02 = CalculateOcclusion(Settings.ShadowDepthTexture.Gather(Settings.ShadowDepthTextureSampler, SamplePos, int2(-2, 0)), Settings);
+		float4 Values22 = CalculateOcclusion(Settings.ShadowDepthTexture.Gather(Settings.ShadowDepthTextureSampler, SamplePos, int2(0, 0)), Settings);
+		float4 Values42 = CalculateOcclusion(Settings.ShadowDepthTexture.Gather(Settings.ShadowDepthTextureSampler, SamplePos, int2(2, 0)), Settings);
+
+		Results.x = PCF5x2(Fraction, Values00, Values20, Values40, Values02, Values22, Values42);
+
+		float4 Values04 = CalculateOcclusion(Settings.ShadowDepthTexture.Gather(Settings.ShadowDepthTextureSampler, SamplePos, int2(-2, 2)), Settings);
+		float4 Values24 = CalculateOcclusion(Settings.ShadowDepthTexture.Gather(Settings.ShadowDepthTextureSampler, SamplePos, int2(0, 2)), Settings);
+		float4 Values44 = CalculateOcclusion(Settings.ShadowDepthTexture.Gather(Settings.ShadowDepthTextureSampler, SamplePos, int2(2, 2)), Settings);
+
+		Results.y = PCF5x2(Fraction, Values02, Values22, Values42, Values04, Values24, Values44);
+		Results.z = PCF5x1(Fraction, Values04, Values24, Values44);
+
+		return dot(Results, 1);
+	}
+
+#else // FEATURE_GATHER4
+
+	// high quality, 7x7 samples, not using gather4 (todo: ideally we make this 6x6 to get same results with gather code)
+	{
+		float2 Fraction = frac(ShadowPosition * Settings.ShadowBufferSize.xy);
+		float2 Sample00TexelCenter = floor(ShadowPosition * Settings.ShadowBufferSize.xy) - float2(3, 3);
+
+		// Fetch 7x7 shadowmap point samples
+		// Do 6x6 PCF samples, sharing the point samples between neighboring PCF samples
+		float4 Results;
+
+		float4 SampleValues03;
+		float4 SampleValues13;
+
+		{
+			float4 SampleValues10;
+			float4 SampleValues11;
+			float4 SampleValues12;
+
+			// Group work to minimize temporary registers and to split texture work with PCF ALU operations to hide texture latency
+			// Without this layout (all texture lookups at the beginning, PCF ALU's at the end) this shader was 4x slower on Nvidia cards
+			{
+				float4 SampleValues00;
+				FetchRowOfFour(Sample00TexelCenter, 0, SampleValues00, Settings);
+				SampleValues10.x = SampleValues00.w;
+
+				float4 SampleValues01;
+				FetchRowOfFour(Sample00TexelCenter, 1, SampleValues01, Settings);
+				SampleValues11.x = SampleValues01.w;
+
+				float4 SampleValues02;
+				FetchRowOfFour(Sample00TexelCenter, 2, SampleValues02, Settings);
+				SampleValues12.x = SampleValues02.w;
+
+				FetchRowOfFour(Sample00TexelCenter, 3, SampleValues03, Settings);
+				SampleValues13.x = SampleValues03.w;
+				Results.x = PCF3x3(Fraction, SampleValues00, SampleValues01, SampleValues02, SampleValues03);
+			}
+
+			{
+				FetchRowOfThreeAfterFour(Sample00TexelCenter, 0, SampleValues10.yzw, Settings);
+				FetchRowOfThreeAfterFour(Sample00TexelCenter, 1, SampleValues11.yzw, Settings);
+				FetchRowOfThreeAfterFour(Sample00TexelCenter, 2, SampleValues12.yzw, Settings);
+				FetchRowOfThreeAfterFour(Sample00TexelCenter, 3, SampleValues13.yzw, Settings);
+				Results.y = PCF3x3(Fraction, SampleValues10, SampleValues11, SampleValues12, SampleValues13);
+			}
+		}
+
+		{
+			float4 SampleValues14;
+			float4 SampleValues15;
+			float4 SampleValues16;
+
+			{
+				float4 SampleValues04;
+				FetchRowOfFour(Sample00TexelCenter, 4, SampleValues04, Settings);
+				SampleValues14.x = SampleValues04.w;
+
+				float4 SampleValues05;
+				FetchRowOfFour(Sample00TexelCenter, 5, SampleValues05, Settings);
+				SampleValues15.x = SampleValues05.w;
+
+				float4 SampleValues06;
+				FetchRowOfFour(Sample00TexelCenter, 6, SampleValues06, Settings);
+				SampleValues16.x = SampleValues06.w;
+
+				Results.z = PCF3x3(Fraction, SampleValues03, SampleValues04, SampleValues05, SampleValues06);
+			}
+
+			{
+				FetchRowOfThreeAfterFour(Sample00TexelCenter, 4, SampleValues14.yzw, Settings);
+				FetchRowOfThreeAfterFour(Sample00TexelCenter, 5, SampleValues15.yzw, Settings);
+				FetchRowOfThreeAfterFour(Sample00TexelCenter, 6, SampleValues16.yzw, Settings);
+				Results.w = PCF3x3(Fraction, SampleValues13, SampleValues14, SampleValues15, SampleValues16);
+			}
+		}
+
+		return dot(Results, .25f);
+	}
+#endif	// FEATURE_GATHER4
+}
diff --git a/Engine/Source/Runtime/Core/Private/WinRT/WinRTMisc.cpp b/Engine/Source/Runtime/Core/Private/WinRT/WinRTMisc.cpp
--- a/Engine/Source/Runtime/Core/Private/WinRT/WinRTMisc.cpp
+++ b/Engine/Source/Runtime/Core/Private/WinRT/WinRTMisc.cpp
@@ -281,8 +281,3 @@ const TCHAR* FWinRTMisc::RootDir()
 	}
 	return *Path;
 }
-
-const TCHAR* FWinRTMisc::GetNullRHIShaderFormat()
-{
-	return TEXT("PCD3D_SM5");
-}
diff --git a/Engine/Source/Runtime/Core/Private/Windows/WindowsPlatformMisc.cpp b/Engine/Source/Runtime/Core/Private/Windows/WindowsPlatformMisc.cpp
--- a/Engine/Source/Runtime/Core/Private/Windows/WindowsPlatformMisc.cpp
+++ b/Engine/Source/Runtime/Core/Private/Windows/WindowsPlatformMisc.cpp
@@ -1763,11 +1763,6 @@ void FWindowsPlatformMisc::LoadStartupModules()
 #endif	//WITH_EDITOR
 }
 
-const TCHAR* FWindowsPlatformMisc::GetNullRHIShaderFormat()
-{
-	return TEXT("PCD3D_SM5");
-}
-
 bool FWindowsPlatformMisc::OsExecute(const TCHAR* CommandType, const TCHAR* Command)
 {
 	HINSTANCE hApp = ShellExecute(NULL,
diff --git a/Engine/Source/Runtime/Core/Public/WinRT/WinRTMisc.h b/Engine/Source/Runtime/Core/Public/WinRT/WinRTMisc.h
--- a/Engine/Source/Runtime/Core/Public/WinRT/WinRTMisc.h
+++ b/Engine/Source/Runtime/Core/Public/WinRT/WinRTMisc.h
@@ -42,13 +42,6 @@ struct CORE_API FWinRTMisc : public FGenericPlatformMisc
 
 	/** Get the application root directory. */
 	static const TCHAR* RootDir();
-
-	/** 
-	 * Determines the shader format for the plarform
-	 *
-	 * @return	Returns the shader format to be used by that platform
-	 */
-	static const TCHAR* GetNullRHIShaderFormat();
 };
 
 typedef FWinRTMisc FPlatformMisc;
diff --git a/Engine/Source/Runtime/Core/Public/Windows/WindowsPlatformMisc.h b/Engine/Source/Runtime/Core/Public/Windows/WindowsPlatformMisc.h
--- a/Engine/Source/Runtime/Core/Public/Windows/WindowsPlatformMisc.h
+++ b/Engine/Source/Runtime/Core/Public/Windows/WindowsPlatformMisc.h
@@ -80,13 +80,6 @@ struct CORE_API FWindowsPlatformMisc : public FGenericPlatformMisc
 	static bool CoInitialize();
 	static void CoUninitialize();
 
-	/** 
-	 * Determines the shader format for the plarform
-	 *
-	 * @return	Returns the shader format to be used by that platform
-	 */
-	static const TCHAR* GetNullRHIShaderFormat();
-
 	/**
 	 * Has the OS execute a command and path pair (such as launch a browser)
 	 *
diff --git a/Engine/Source/Runtime/Engine/Private/Particles/ParticleSystemRender.cpp b/Engine/Source/Runtime/Engine/Private/Particles/ParticleSystemRender.cpp
--- a/Engine/Source/Runtime/Engine/Private/Particles/ParticleSystemRender.cpp
+++ b/Engine/Source/Runtime/Engine/Private/Particles/ParticleSystemRender.cpp
@@ -1250,7 +1250,7 @@ void FDynamicSpriteEmitterData::GetDynamicMeshElementsEmitter(const FParticleSys
 	SCOPE_CYCLE_COUNTER(STAT_SpriteRenderingTime);
 
 	const auto FeatureLevel = View->GetFeatureLevel();
-	const bool bInstanced = FeatureLevel >= ERHIFeatureLevel::SM3;
+	const bool bInstanced = FeatureLevel >= ERHIFeatureLevel::SM4;
 
 	// Sort and generate particles for this view.
 	const FDynamicSpriteEmitterReplayDataBase* SourceData = GetSourceData();
@@ -1827,7 +1827,7 @@ void FDynamicMeshEmitterData::GetDynamicMeshElementsEmitter(const FParticleSyste
 {
 	SCOPE_CYCLE_COUNTER(STAT_MeshRenderingTime);
 
-	const bool bInstanced = GRHIFeatureLevel >= ERHIFeatureLevel::SM3;
+	const bool bInstanced = GRHIFeatureLevel >= ERHIFeatureLevel::SM4;
 
 	if (bValid)
 	{
diff --git a/Engine/Source/Runtime/Engine/Private/ReflectionCaptureComponent.cpp b/Engine/Source/Runtime/Engine/Private/ReflectionCaptureComponent.cpp
--- a/Engine/Source/Runtime/Engine/Private/ReflectionCaptureComponent.cpp
+++ b/Engine/Source/Runtime/Engine/Private/ReflectionCaptureComponent.cpp
@@ -961,7 +961,7 @@ void UReflectionCaptureComponent::PostLoad()
 	Super::PostLoad();
 
 	static const auto CVar = IConsoleManager::Get().FindTConsoleVariableDataInt(TEXT("r.FeatureLevelPreview"));
-	bool bRetainAllFeatureLevelData = CVar->GetValueOnGameThread() == 1 && GIsEditor && GMaxRHIFeatureLevel >= ERHIFeatureLevel::SM3;
+	bool bRetainAllFeatureLevelData = CVar->GetValueOnGameThread() == 1 && GIsEditor && GMaxRHIFeatureLevel >= ERHIFeatureLevel::SM4;
 
 	// If we're loading on a platform that doesn't require cooked data, attempt to load missing data from the DDC
 	if (!FPlatformProperties::RequiresCookedData())
diff --git a/Engine/Source/Runtime/Renderer/Private/PostProcess/SceneRenderTargets.cpp b/Engine/Source/Runtime/Renderer/Private/PostProcess/SceneRenderTargets.cpp
--- a/Engine/Source/Runtime/Renderer/Private/PostProcess/SceneRenderTargets.cpp
+++ b/Engine/Source/Runtime/Renderer/Private/PostProcess/SceneRenderTargets.cpp
@@ -64,9 +64,9 @@ static TAutoConsoleVariable<int32> CVarMobileMSAA(
 	TEXT("r.MobileMSAA"),
 	0,
 	TEXT("Use MSAA instead of Temporal AA on mobile:\n")
-    TEXT("1: Use Temporal AA (MSAA disabled)\n")
-    TEXT("2: Use 2x MSAA (Temporal AA disabled)\n")
-    TEXT("4: Use 4x MSAA (Temporal AA disabled)\n"),
+	TEXT("1: Use Temporal AA (MSAA disabled)\n")
+	TEXT("2: Use 2x MSAA (Temporal AA disabled)\n")
+	TEXT("4: Use 4x MSAA (Temporal AA disabled)\n"),
 	ECVF_RenderThreadSafe
 	);
 
@@ -201,7 +201,7 @@ void FSceneRenderTargets::Allocate(const FSceneViewFamily& ViewFamily)
 		(CurrentSceneColorFormat != SceneColorFormat) ||
 		(bAllowStaticLighting != bNewAllowStaticLighting) ||
 		(bUseDownsizedOcclusionQueries != bDownsampledOcclusionQueries) ||
- 		(CurrentMaxShadowResolution != MaxShadowResolution) ||
+		(CurrentMaxShadowResolution != MaxShadowResolution) ||
 		(CurrentTranslucencyLightingVolumeDim != TranslucencyLightingVolumeDim) ||
 		(CurrentMobile32bpp != Mobile32bpp) ||
 		(CurrentMobileMSAA != MobileMSAA) ||
@@ -332,12 +332,13 @@ void FSceneRenderTargets::AllocSceneColor()
 			Desc.TargetableFlags |= TexCreate_UAV;
 		}
 
-		if (CurrentFeatureLevel == ERHIFeatureLevel::ES2)
+		if (CurrentFeatureLevel < ERHIFeatureLevel::SM4)
 		{
-		    uint16 NumSamples = GRHIShaderPlatform == SP_OPENGL_ES2_IOS ? 1 : CVarMobileMSAA.GetValueOnRenderThread();
-	        if (NumSamples != 1 && NumSamples != 2 && NumSamples != 4)
+			//@todo-rco: Fix when OpenGL supports MSAA
+			uint16 NumSamples = GRHIShaderPlatform == SP_OPENGL_ES2_IOS ? 1 : CVarMobileMSAA.GetValueOnRenderThread();
+			if (NumSamples != 1 && NumSamples != 2 && NumSamples != 4)
 			{
-	            NumSamples = 1;
+				NumSamples = 1;
 			}
 			Desc.NumSamples = NumSamples;
 		}
@@ -640,7 +641,7 @@ void FSceneRenderTargets::FinishRenderingCustomDepth(FRHICommandListImmediate& R
 */
 void FSceneRenderTargets::ResolveSceneColor(FRHICommandList& RHICmdList, const FResolveRect& ResolveRect)
 {
-    SCOPED_DRAW_EVENT(ResolveSceneColor, DEC_SCENE_ITEMS);
+	SCOPED_DRAW_EVENT(ResolveSceneColor, DEC_SCENE_ITEMS);
 
 	uint32 samples = GSceneRenderTargets.SceneColor->GetDesc().NumSamples;
 
@@ -1063,6 +1064,14 @@ void FSceneRenderTargets::AllocateForwardShadingPathRenderTargets()
 		FPooledRenderTargetDesc Desc(FPooledRenderTargetDesc::Create2DDesc(BufferSize, PF_DepthStencil, TexCreate_None, TexCreate_DepthStencilTargetable, false));
 		GRenderTargetPool.FindFreeElement(Desc, AuxiliarySceneDepthZ, TEXT("AuxiliarySceneDepthZ"));
 	}
+
+	AllocateReflectionTargets();	
+}
+
+void FSceneRenderTargets::AllocateForwardShadingShadowDepthTarget(const FIntPoint& ShadowBufferResolution)
+{
+	FPooledRenderTargetDesc Desc(FPooledRenderTargetDesc::Create2DDesc(ShadowBufferResolution, PF_ShadowDepth, TexCreate_None, TexCreate_DepthStencilTargetable, false));
+	GRenderTargetPool.FindFreeElement(Desc, ShadowDepthZ, TEXT("ShadowDepthZ"));
 }
 
 // for easier use of "VisualizeTexture"
@@ -1097,6 +1106,41 @@ static TCHAR* const GetTranslucencyShadowTransmissionName(uint32 Id)
 	return (TCHAR*)TEXT("InvalidName");
 }
 
+void FSceneRenderTargets::AllocateReflectionTargets()
+{	
+	if (CurrentFeatureLevel >= ERHIFeatureLevel::ES2 && GSupportsRenderTargetFormat_PF_FloatRGBA)
+	{
+		extern ENGINE_API int32 GReflectionCaptureSize;
+		const int32 NumReflectionCaptureMips = FMath::CeilLogTwo(GReflectionCaptureSize) + 1;
+
+		{
+			uint32 TexFlags = TexCreate_TargetArraySlicesIndependently;
+
+			// Create scratch cubemaps for filtering passes
+			FPooledRenderTargetDesc Desc2(FPooledRenderTargetDesc::CreateCubemapDesc(GReflectionCaptureSize, PF_FloatRGBA, TexFlags, TexCreate_RenderTargetable, false, 1, NumReflectionCaptureMips));
+			GRenderTargetPool.FindFreeElement(Desc2, ReflectionColorScratchCubemap[0], TEXT("ReflectionColorScratchCubemap0"));
+			GRenderTargetPool.FindFreeElement(Desc2, ReflectionColorScratchCubemap[1], TEXT("ReflectionColorScratchCubemap1"));
+
+			EPixelFormat BrightnessFormat = CurrentFeatureLevel > ERHIFeatureLevel::ES2 ? PF_R32_FLOAT : PF_FloatRGBA;
+			FPooledRenderTargetDesc Desc3(FPooledRenderTargetDesc::Create2DDesc(FIntPoint(1, 1), BrightnessFormat, TexCreate_None, TexCreate_RenderTargetable, false));
+			GRenderTargetPool.FindFreeElement(Desc3, ReflectionBrightness, TEXT("ReflectionBrightness"));
+		}
+
+		extern int32 GDiffuseIrradianceCubemapSize;
+		const int32 NumDiffuseIrradianceMips = FMath::CeilLogTwo(GDiffuseIrradianceCubemapSize) + 1;
+
+		{
+			FPooledRenderTargetDesc Desc2(FPooledRenderTargetDesc::CreateCubemapDesc(GDiffuseIrradianceCubemapSize, PF_FloatRGBA, TexCreate_TargetArraySlicesIndependently, TexCreate_RenderTargetable, false, 1, NumDiffuseIrradianceMips));
+			GRenderTargetPool.FindFreeElement(Desc2, DiffuseIrradianceScratchCubemap[0], TEXT("DiffuseIrradianceScratchCubemap0"));
+			GRenderTargetPool.FindFreeElement(Desc2, DiffuseIrradianceScratchCubemap[1], TEXT("DiffuseIrradianceScratchCubemap1"));
+		}
+
+		{
+			FPooledRenderTargetDesc Desc(FPooledRenderTargetDesc::Create2DDesc(FIntPoint(FSHVector3::MaxSHBasis, 1), PF_FloatRGBA, TexCreate_None, TexCreate_RenderTargetable, false));
+			GRenderTargetPool.FindFreeElement(Desc, SkySHIrradianceMap, TEXT("SkySHIrradianceMap"));
+		}
+	}
+}
 
 void FSceneRenderTargets::AllocateDeferredShadingPathRenderTargets()
 {
@@ -1206,39 +1250,7 @@ void FSceneRenderTargets::AllocateDeferredShadingPathRenderTargets()
 		}
 	}
 
-	extern ENGINE_API int32 GReflectionCaptureSize;
-	const int32 NumReflectionCaptureMips = FMath::CeilLogTwo(GReflectionCaptureSize) + 1;
-
-	{
-		uint32 TexFlags = TexCreate_None;
-
-		if (!GSupportsGSRenderTargetLayerSwitchingToMips)
-		{
-			TexFlags = TexCreate_TargetArraySlicesIndependently;
-		}
-
-		// Create scratch cubemaps for filtering passes
-		FPooledRenderTargetDesc Desc2(FPooledRenderTargetDesc::CreateCubemapDesc(GReflectionCaptureSize, PF_FloatRGBA, TexFlags, TexCreate_RenderTargetable, false, 1, NumReflectionCaptureMips));
-		GRenderTargetPool.FindFreeElement(Desc2, ReflectionColorScratchCubemap[0], TEXT("ReflectionColorScratchCubemap0"));
-		GRenderTargetPool.FindFreeElement(Desc2, ReflectionColorScratchCubemap[1], TEXT("ReflectionColorScratchCubemap1"));
-
-		FPooledRenderTargetDesc Desc3(FPooledRenderTargetDesc::Create2DDesc(FIntPoint(1, 1), PF_R32_FLOAT, TexCreate_None, TexCreate_RenderTargetable, false));
-		GRenderTargetPool.FindFreeElement(Desc3, ReflectionBrightness, TEXT("ReflectionBrightness"));
-	}
-
-	extern int32 GDiffuseIrradianceCubemapSize;
-	const int32 NumDiffuseIrradianceMips = FMath::CeilLogTwo(GDiffuseIrradianceCubemapSize) + 1;
-
-	{
-		FPooledRenderTargetDesc Desc2(FPooledRenderTargetDesc::CreateCubemapDesc(GDiffuseIrradianceCubemapSize, PF_FloatRGBA, TexCreate_None, TexCreate_RenderTargetable, false, 1, NumDiffuseIrradianceMips));
-		GRenderTargetPool.FindFreeElement(Desc2, DiffuseIrradianceScratchCubemap[0], TEXT("DiffuseIrradianceScratchCubemap0"));
-		GRenderTargetPool.FindFreeElement(Desc2, DiffuseIrradianceScratchCubemap[1], TEXT("DiffuseIrradianceScratchCubemap1"));
-	}
-
-	{
-		FPooledRenderTargetDesc Desc(FPooledRenderTargetDesc::Create2DDesc(FIntPoint(FSHVector3::MaxSHBasis, 1), PF_FloatRGBA, TexCreate_None, TexCreate_RenderTargetable, false));
-		GRenderTargetPool.FindFreeElement(Desc, SkySHIrradianceMap, TEXT("SkySHIrradianceMap"));
-	}
+	AllocateReflectionTargets();
 
 	if (CurrentFeatureLevel >= ERHIFeatureLevel::SM5)
 	{
@@ -1267,7 +1279,7 @@ EPixelFormat FSceneRenderTargets::GetSceneColorFormat() const
 {
 	EPixelFormat SceneColorBufferFormat = PF_FloatRGBA;
 
-	if (CurrentFeatureLevel <= ERHIFeatureLevel::ES2)
+	if (CurrentFeatureLevel < ERHIFeatureLevel::SM4)
 	{
 		// Potentially allocate an alpha channel in th -fe scene color texture to store the resolved scene depth.
 		SceneColorBufferFormat = GSupportsRenderTargetFormat_PF_FloatRGBA ? PF_FloatRGBA : PF_B8G8R8A8;
@@ -1308,7 +1320,7 @@ void FSceneRenderTargets::InitDynamicRHI()
 		FRHICommandListImmediate& RHICmdList = FRHICommandListExecutor::GetImmediateCommandList();
 		RHICmdList.SetScissorRect(false, 0, 0, 0, 0);
 
-		if (CurrentFeatureLevel == ERHIFeatureLevel::ES2)
+		if (CurrentFeatureLevel < ERHIFeatureLevel::SM4)
 		{
 			AllocateForwardShadingPathRenderTargets();
 		}
@@ -1739,7 +1751,7 @@ void FDeferredPixelShaderParameters::Set(FRHICommandList& RHICmdList, const Shad
 
 	const auto FeatureLevel = View.GetFeatureLevel();
 
-	if (TextureMode == ESceneRenderTargetsMode::SetTextures && FeatureLevel >= ERHIFeatureLevel::SM3)
+	if (TextureMode == ESceneRenderTargetsMode::SetTextures && FeatureLevel >= ERHIFeatureLevel::SM4)
 	{
 		// if there is no ambient occlusion it's better to have white there
 		IPooledRenderTarget* ScreenSpaceAO = GSceneRenderTargets.ScreenSpaceAO;
diff --git a/Engine/Source/Runtime/Renderer/Private/PostProcess/SceneRenderTargets.h b/Engine/Source/Runtime/Renderer/Private/PostProcess/SceneRenderTargets.h
--- a/Engine/Source/Runtime/Renderer/Private/PostProcess/SceneRenderTargets.h
+++ b/Engine/Source/Runtime/Renderer/Private/PostProcess/SceneRenderTargets.h
@@ -135,6 +135,12 @@ public:
 	 * If the allocated render targets are too small, they are reallocated.
 	 */
 	void Allocate(const FSceneViewFamily& ViewFamily);
+
+	/**
+	 * Forward shading can't know how big the optimal atlased shadow buffer will be, so provide a set it up per frame.
+	 */
+	void AllocateForwardShadingShadowDepthTarget(const FIntPoint& ShadowBufferResolution);
+
 	/**
 	 *
 	 */
@@ -502,11 +508,13 @@ private:
 	void InitEditorPrimitivesDepth();
 
 	/** Allocates render targets for use with the forward shading path. */
-	void AllocateForwardShadingPathRenderTargets();
+	void AllocateForwardShadingPathRenderTargets();	
 
 	/** Allocates render targets for use with the deferred shading path. */
 	void AllocateDeferredShadingPathRenderTargets();
 
+	void AllocateReflectionTargets();
+
 	/** Determine the appropriate render target dimensions. */
 	FIntPoint GetSceneRenderTargetSize(const FSceneViewFamily & ViewFamily) const;
 
diff --git a/Engine/Source/Runtime/Renderer/Private/SceneHitProxyRendering.cpp b/Engine/Source/Runtime/Renderer/Private/SceneHitProxyRendering.cpp
--- a/Engine/Source/Runtime/Renderer/Private/SceneHitProxyRendering.cpp
+++ b/Engine/Source/Runtime/Renderer/Private/SceneHitProxyRendering.cpp
@@ -37,6 +37,7 @@ public:
 	{
 		// Only compile the hit proxy vertex shader on PC
 		return IsPCPlatform(Platform)
+			&& IsFeatureLevelSupported(Platform, ERHIFeatureLevel::SM4)
 			// and only compile for the default material or materials that are masked.
 			&& (Material->IsSpecialEngineMaterial() || Material->IsMasked() || Material->MaterialMayModifyMeshPosition() || Material->IsTwoSided());
 	}
@@ -108,7 +109,8 @@ public:
 	static bool ShouldCache(EShaderPlatform Platform,const FMaterial* Material,const FVertexFactoryType* VertexFactoryType)
 	{
 		// Only compile the hit proxy vertex shader on PC
-		return IsPCPlatform(Platform)
+		return IsPCPlatform(Platform) 
+			&& IsFeatureLevelSupported(Platform, ERHIFeatureLevel::SM4) 
 			// and only compile for default materials or materials that are masked.
 			&& (Material->IsSpecialEngineMaterial() || Material->IsMasked() || Material->MaterialMayModifyMeshPosition() || Material->IsTwoSided());
 	}
diff --git a/Engine/Source/Runtime/ShaderCore/Private/ShaderCore.cpp b/Engine/Source/Runtime/ShaderCore/Private/ShaderCore.cpp
--- a/Engine/Source/Runtime/ShaderCore/Private/ShaderCore.cpp
+++ b/Engine/Source/Runtime/ShaderCore/Private/ShaderCore.cpp
@@ -108,20 +108,29 @@ TMap<FString, FString> GShaderFileCache;
 /** The shader file hash cache, used to minimize loading and hashing shader files */
 TMap<FString, FSHAHash> GShaderHashCache;
 
+static TAutoConsoleVariable<int32> CVarForceDebugViewModes(
+	TEXT("r.ForceDebugViewModes"),
+	0,
+	TEXT("0: Setting has no effect.\n")
+	TEXT("1: Forces debug view modes to be available, even on cooked builds.")
+	TEXT("2: Forces debug view modes to be unavailable, even on editor builds.  Removes many shader permutations for faster shader iteration."),
+	ECVF_RenderThreadSafe | ECVF_ReadOnly);
+
 /** Returns true if debug viewmodes are allowed for the current platform. */
 bool AllowDebugViewmodes()
-{
-	static bool bAllowDebugViewmodesOnConsoles = false;
-	static bool bReadFromIni = false;
-	if (!bReadFromIni)
-	{
-		bReadFromIni = true;
-		GConfig->GetBool( TEXT("/Script/Engine.Engine"), TEXT("bAllowDebugViewmodesOnConsoles"), bAllowDebugViewmodesOnConsoles, GEngineIni );
-	}
+{	
+	int32 ForceDebugViewValue = CVarForceDebugViewModes.GetValueOnAnyThread();
 
-	// To use debug viewmodes on consoles, bAllowDebugViewmodesOnConsoles in the engine ini must be set to true, 
+	// To use debug viewmodes on consoles, r.ForceDebugViewModes must be set to 2 in ConsoleVariables.ini
 	// And EngineDebugMaterials must be in the StartupPackages for the target platform.
-	return bAllowDebugViewmodesOnConsoles || (!IsRunningCommandlet() && !FPlatformProperties::RequiresCookedData());
+	bool bForceEnable = ForceDebugViewValue == 1;
+	bool bForceDisable = ForceDebugViewValue == 2;
+
+	//bool bEnabled = (!IsRunningCommandlet() && !FPlatformProperties::RequiresCookedData());	
+	//bEnabled |= bForceEnable;
+	//bEnabled &= !bForceDisable;
+
+	return (!bForceDisable) && (bForceEnable || (!IsRunningCommandlet() && !FPlatformProperties::RequiresCookedData()));
 }
 
 bool FShaderParameterMap::FindParameterAllocation(const TCHAR* ParameterName,uint16& OutBufferIndex,uint16& OutBaseIndex,uint16& OutSize) const
