Possible bug fix for Metal for HDR custom resolve.

diff --git a/Engine/Shaders/PostProcessMobile.usf b/Engine/Shaders/PostProcessMobile.usf
--- a/Engine/Shaders/PostProcessMobile.usf
+++ b/Engine/Shaders/PostProcessMobile.usf
@@ -71,7 +71,7 @@ void SunMaskPS_ES2(
 	out half4 OutColor : SV_Target0
 	)
 {
-	#if COMPILER_GLSL_ES2 && ES2_USE_FETCH
+	#if (COMPILER_GLSL_ES2 && ES2_USE_FETCH) || METAL_PROFILE
 		OutColor = FramebufferFetchES2();
 	#else
 		OutColor = PostprocessInput0.Sample(PostprocessInput0Sampler, InUVPos.xy);
