Fixed Slate materials not rendering

diff --git a/Engine/Shaders/SlateElementPixelShader.usf b/Engine/Shaders/SlateElementPixelShader.usf
--- a/Engine/Shaders/SlateElementPixelShader.usf
+++ b/Engine/Shaders/SlateElementPixelShader.usf
@@ -51,7 +51,8 @@ struct VertexOut
 float4 GetMaterialColor( VertexOut InVertex, float2 UV )
 {
 	FMaterialPixelParameters Parameters = MakeInitializedMaterialPixelParameters();
-	CalcMaterialParameters(Parameters, true, float4(InVertex.WorldPosition,0,1));
+	CalcMaterialParameters(Parameters, true, InVertex.Position);
+
 #if NUM_MATERIAL_TEXCOORDS == 1
 	Parameters.TexCoords[0] = UV;
 #elif NUM_MATERIAL_TEXCOORDS == 2
