Removed general purpose tweak usage

diff --git a/Engine/Shaders/ReflectionEnvironmentShaders.usf b/Engine/Shaders/ReflectionEnvironmentShaders.usf
--- a/Engine/Shaders/ReflectionEnvironmentShaders.usf
+++ b/Engine/Shaders/ReflectionEnvironmentShaders.usf
@@ -533,10 +533,7 @@ void ReflectionApplyPS(
 			SkyLighting *= GBuffer.IndirectIrradiance;
 		}
 
-		float Visibility = 1;
-		
-		if (View.GeneralPurposeTweak > 0)
-		Visibility = GetDistanceFieldAOSpecularOcclusion(UV, ReflectionVector, GBuffer.Roughness); 
+		float Visibility = GetDistanceFieldAOSpecularOcclusion(UV, ReflectionVector, GBuffer.Roughness); 
 
 		SpecularLighting.rgb += (Visibility * SpecularLighting.a) * SkyLighting;
 #endif
