SM4 feature level now supports reflection captures on translucency
Fixed bug with skylight reflections blending into reflection captures on translucent materials

diff --git a/Engine/Shaders/BasePassPixelShader.usf b/Engine/Shaders/BasePassPixelShader.usf
--- a/Engine/Shaders/BasePassPixelShader.usf
+++ b/Engine/Shaders/BasePassPixelShader.usf
@@ -61,17 +61,24 @@ float4 PointSkyBentNormal;
 #if TRANSLUCENCY_LIGHTING_SURFACE
 
 #if FEATURE_LEVEL >= FEATURE_LEVEL_SM5
-/** Prenormalized capture of the scene that's closest to the object being rendered, used for reflection environment on translucency. */
-TextureCubeArray ReflectionCubemap;
-SamplerState ReflectionCubemapSampler;
-int CubemapArrayIndex;
+	/** Prenormalized capture of the scene that's closest to the object being rendered, used for reflection environment on translucency. */
+	TextureCubeArray ReflectionCubemap;
+	SamplerState ReflectionCubemapSampler;
+	int CubemapArrayIndex;
+#else
+	TextureCube ReflectionCubemap;
+	SamplerState ReflectionCubemapSampler;
 #endif
 
 half3 GetImageBasedReflectionLighting(FMaterialPixelParameters MaterialParameters, half Roughness, half3 SpecularColor, half IndirectIrradiance)
 {
-#if FEATURE_LEVEL >= FEATURE_LEVEL_SM5
 	half AbsoluteSpecularMip = ComputeReflectionCaptureMipFromRoughness(Roughness);
+
+#if FEATURE_LEVEL >= FEATURE_LEVEL_SM5
 	float4 SpecularIBL = TextureCubeArraySampleLevel(ReflectionCubemap, ReflectionCubemapSampler, MaterialParameters.ReflectionVector, CubemapArrayIndex, AbsoluteSpecularMip);
+#else
+	float4 SpecularIBL = TextureCubeSampleLevel(ReflectionCubemap, ReflectionCubemapSampler, MaterialParameters.ReflectionVector, AbsoluteSpecularMip);
+#endif
 
 	FLATTEN
 	if (View.UseLightmaps > 0)
@@ -82,7 +89,7 @@ half3 GetImageBasedReflectionLighting(FMaterialPixelParameters MaterialParameter
 
 #if ENABLE_SKY_LIGHT
 	BRANCH
-	if (SkyLightParameters.y > 0 && SpecularIBL.a < .001f)
+	if (SkyLightParameters.y > 0 && SpecularIBL.a < .999f)
 	{
 		// Normalize for static skylight types which mix with lightmaps
 		bool bNormalize = SkyLightParameters.z < 1 && View.UseLightmaps;
@@ -103,9 +110,6 @@ half3 GetImageBasedReflectionLighting(FMaterialPixelParameters MaterialParameter
 	SpecularColor = EnvBRDFApprox(SpecularColor, Roughness, NoV);
 
 	return SpecularIBL.rgb * SpecularColor;
-#else
-	return 0;
-#endif
 }
 
 #endif
diff --git a/Engine/Shaders/ReflectionEnvironmentShaders.usf b/Engine/Shaders/ReflectionEnvironmentShaders.usf
--- a/Engine/Shaders/ReflectionEnvironmentShaders.usf
+++ b/Engine/Shaders/ReflectionEnvironmentShaders.usf
@@ -522,20 +522,25 @@ void ReflectionApplyPS(
 #endif
 
 #if APPLY_SKYLIGHT
-		float3 ReflectionVector = reflect(-V, GBuffer.WorldNormal);
-		// Normalize for static skylight types which mix with lightmaps
-		bool bNormalize = SkyLightParameters.z < 1 && View.UseLightmaps > 0;
-		float3 SkyLighting = GetSkyLightReflection(ReflectionVector, GBuffer.Roughness, bNormalize);
 
-		FLATTEN
-		if (bNormalize)
+		BRANCH
+		if (SpecularLighting.a > .001f)
 		{
-			SkyLighting *= GBuffer.IndirectIrradiance;
-		}
+			float3 ReflectionVector = reflect(-V, GBuffer.WorldNormal);
+			// Normalize for static skylight types which mix with lightmaps
+			bool bNormalize = SkyLightParameters.z < 1 && View.UseLightmaps > 0;
+			float3 SkyLighting = GetSkyLightReflection(ReflectionVector, GBuffer.Roughness, bNormalize);
 
-		float Visibility = GetDistanceFieldAOSpecularOcclusion(UV, ReflectionVector, GBuffer.Roughness); 
+			FLATTEN
+			if (bNormalize)
+			{
+				SkyLighting *= GBuffer.IndirectIrradiance;
+			}
+
+			float Visibility = GetDistanceFieldAOSpecularOcclusion(UV, ReflectionVector, GBuffer.Roughness); 
 
-		SpecularLighting.rgb += (Visibility * SpecularLighting.a) * SkyLighting;
+			SpecularLighting.rgb += (Visibility * SpecularLighting.a) * SkyLighting;
+		}
 #endif
 
 #if APPLY_SSR
diff --git a/Engine/Source/Runtime/Renderer/Private/BasePassRendering.cpp b/Engine/Source/Runtime/Renderer/Private/BasePassRendering.cpp
--- a/Engine/Source/Runtime/Renderer/Private/BasePassRendering.cpp
+++ b/Engine/Source/Runtime/Renderer/Private/BasePassRendering.cpp
@@ -109,10 +109,10 @@ void FTranslucentLightingParameters::Set(FRHICommandList& RHICmdList, FShader* S
 	SkyLightReflectionParameters.SetParameters(RHICmdList, Shader->GetPixelShader(), (const FScene*)(View->Family->Scene), true);
 }
 
-void FTranslucentLightingParameters::SetMesh(FRHICommandList& RHICmdList, FShader* Shader, const FPrimitiveSceneProxy* Proxy)
+void FTranslucentLightingParameters::SetMesh(FRHICommandList& RHICmdList, FShader* Shader, const FPrimitiveSceneProxy* Proxy, ERHIFeatureLevel::Type FeatureLevel)
 {
 	// Note: GBlackCubeArrayTexture has an alpha of 0, which is needed to represent invalid data so the sky cubemap can still be applied
-	FTextureRHIParamRef CubeArrayTexture = GBlackCubeArrayTexture->TextureRHI;
+	FTextureRHIParamRef CubeArrayTexture = FeatureLevel >= ERHIFeatureLevel::SM5 ? GBlackCubeArrayTexture->TextureRHI : GBlackTextureCube->TextureRHI;
 	int32 ArrayIndex = 0;
 	const FPrimitiveSceneInfo* PrimitiveSceneInfo = Proxy ? Proxy->GetPrimitiveSceneInfo() : NULL;
 
diff --git a/Engine/Source/Runtime/Renderer/Private/BasePassRendering.h b/Engine/Source/Runtime/Renderer/Private/BasePassRendering.h
--- a/Engine/Source/Runtime/Renderer/Private/BasePassRendering.h
+++ b/Engine/Source/Runtime/Renderer/Private/BasePassRendering.h
@@ -257,7 +257,7 @@ public:
 
 	void Set(FRHICommandList& RHICmdList, FShader* Shader, const FSceneView* View);
 
-	void SetMesh(FRHICommandList& RHICmdList, FShader* Shader, const FPrimitiveSceneProxy* Proxy);
+	void SetMesh(FRHICommandList& RHICmdList, FShader* Shader, const FPrimitiveSceneProxy* Proxy, ERHIFeatureLevel::Type FeatureLevel);
 
 	/** Serializer. */
 	friend FArchive& operator<<(FArchive& Ar,FTranslucentLightingParameters& P)
@@ -386,7 +386,7 @@ public:
 		if (View.GetFeatureLevel() >= ERHIFeatureLevel::SM4
 			&& IsTranslucentBlendMode(BlendMode))
 		{
-			TranslucentLightingParameters.SetMesh(RHICmdList, this, Proxy);
+			TranslucentLightingParameters.SetMesh(RHICmdList, this, Proxy, View.GetFeatureLevel());
 		}
 
 		FMeshMaterialShader::SetMesh(RHICmdList, GetPixelShader(),VertexFactory,View,Proxy,BatchElement);
diff --git a/Engine/Source/Runtime/Renderer/Private/RendererScene.cpp b/Engine/Source/Runtime/Renderer/Private/RendererScene.cpp
--- a/Engine/Source/Runtime/Renderer/Private/RendererScene.cpp
+++ b/Engine/Source/Runtime/Renderer/Private/RendererScene.cpp
@@ -908,12 +908,22 @@ const FReflectionCaptureProxy* FScene::FindClosestReflectionCapture(FVector Posi
 
 void FScene::GetCaptureParameters(const FReflectionCaptureProxy* ReflectionProxy, FTextureRHIParamRef& ReflectionCubemapArray, int32& ArrayIndex) const
 {
-	const FCaptureComponentSceneState* FoundState = ReflectionSceneData.AllocatedReflectionCaptureState.Find(ReflectionProxy->Component);
+	ERHIFeatureLevel::Type LocalFeatureLevel = GetFeatureLevel();
 
-	if (FoundState)
+	if (LocalFeatureLevel >= ERHIFeatureLevel::SM5)
 	{
-		ReflectionCubemapArray = ReflectionSceneData.CubemapArray.GetRenderTarget().ShaderResourceTexture;
-		ArrayIndex = FoundState->CaptureIndex;
+		const FCaptureComponentSceneState* FoundState = ReflectionSceneData.AllocatedReflectionCaptureState.Find(ReflectionProxy->Component);
+
+		if (FoundState)
+		{
+			ReflectionCubemapArray = ReflectionSceneData.CubemapArray.GetRenderTarget().ShaderResourceTexture;
+			ArrayIndex = FoundState->CaptureIndex;
+		}
+	}
+	else if (ReflectionProxy->SM4FullHDRCubemap)
+	{
+		ReflectionCubemapArray = ReflectionProxy->SM4FullHDRCubemap->TextureRHI;
+		ArrayIndex = 0;
 	}
 }
 
