optimized motionblur postprocess

diff --git a/Engine/Shaders/PostProcessMotionBlur.usf b/Engine/Shaders/PostProcessMotionBlur.usf
--- a/Engine/Shaders/PostProcessMotionBlur.usf
+++ b/Engine/Shaders/PostProcessMotionBlur.usf
@@ -139,26 +139,37 @@ void OverrideWithTestChart(float2 ScreenPos, inout float2 ObjectVelocity, inout
 
 // ------------------------------------------
 
+// motion blur setup vertex shader
+void SetupVS(
+	in float4 InPosition : ATTRIBUTE0,
+	in float2 InTexCoord : ATTRIBUTE1,
+	out float2 OutUV: TEXCOORD0,
+	out float2 OutUVs[4]: TEXCOORD1,
+	out float4 OutPosition : SV_POSITION
+	)
+{
+	DrawRectangle(InPosition, InTexCoord, OutPosition, OutUV);
+
+//	float2 ScreenPos =  OutPosition.xy;
+
+	float2 HalfPixelOffset = PostprocessInput1Size.zw * float2(0.5f, 0.5f);
+
+	// no filtering (2x2 kernel) to get no leaking in Depth of Field, lefttop,righttop,leftbottom,rightbottom
+	OutUVs[0] = OutUV + HalfPixelOffset * float2(-1, -1);
+	OutUVs[1] = OutUV + HalfPixelOffset * float2( 1, -1);
+	OutUVs[2] = OutUV + HalfPixelOffset * float2(-1,  1);
+	OutUVs[3] = OutUV + HalfPixelOffset * float2( 1,  1);
+}
+
 // motion blur setup pixel shader, downsamples to half resolution
 // MRT0: velocity in normalized_motionblur_velocity and mask in 0:background..1:object (needed for soft masked)
 // MRT1: color and depth in alpha
-void SetupPS(float4 UVAndScreenPos : TEXCOORD0,
+void SetupPS(
+	float2 UV : TEXCOORD0,				// UV into the full resolution source RT
+	float2 UVs[4] : TEXCOORD1,			// UV for high quality 4 samples with blurring
 	out float4 OutVelocity : SV_Target0,
 	out float4 OutColor : SV_Target1)
 {
-	float2 UV = UVAndScreenPos.xy;
-	// UV into the full resolution source RT
-	// high quality 4 samples with blurring
-	float2 UVs[4];
-
-	float2 HalfPixelOffset = PostprocessInput1Size.zw * float2(0.5f, 0.5f);
-
-	// no filtering (2x2 kernel) to get no leaking in Depth of Field, lefttop,righttop,leftbottom,rightbottom
-	UVs[0] = UV + HalfPixelOffset * float2(-1, -1);
-	UVs[1] = UV + HalfPixelOffset * float2( 1, -1);
-	UVs[2] = UV + HalfPixelOffset * float2(-1,  1);
-	UVs[3] = UV + HalfPixelOffset * float2( 1,  1);
-	
 	// detect find objects
 	float4 ObjectMask;
 	float2 VelocitySamples[4];
@@ -219,44 +230,49 @@ void SetupPS(float4 UVAndScreenPos : TEXCOORD0,
 
 	OutColor.a = ConvertFromDeviceZ(OutColor.a);
 
-	float4 TempVelocity;
 	{
 		float2 SumVelocity = 0;
 
-		for( int i = 0; i < 4; i++ )
+		FLATTEN for( int i = 0; i < 4; i++ )
 		{
 			SumVelocity += VelocitySamples[i] * ObjectMask[i];
 		}
 
-		TempVelocity.xy = SumVelocity * InvTotalWeight * VelocityScale.xy;
+		OutVelocity.xy = SumVelocity * InvTotalWeight * VelocityScale.xy;
 	}
 
 	// for debugging
+/*
 	{
 		float2 BackgroundVelocity = 0;
 
-		OverrideWithTestChart(UVAndScreenPos.zw, TempVelocity.xy, BackgroundVelocity, OutColor, AvgObject);
+		OverrideWithTestChart(UVAndScreenPos.zw, OutVelocity.xy, BackgroundVelocity, OutColor, AvgObject);
 	}
+*/
+
+	// 0:background, 1:object layer
+	OutVelocity.b = AvgObject > 0.5f;
 
 	// clamp motion vector in a disc from -1 to 1 (the maximum motion vector range)
-	half Len = dot(TempVelocity.xy, TempVelocity.xy);
-	FLATTEN if(Len > 1)
-	{
-		TempVelocity.xy *= rsqrt(Len);
-	}
-	FLATTEN if(Len < 0.01f)
 	{
-		TempVelocity = 0;
+		half Len = dot(OutVelocity.xy, OutVelocity.xy);
+
+		float ScaleFix = rsqrt(Len);
+		
+		FLATTEN if(Len < 1)
+		{
+			ScaleFix = 1.0f;
+		}
+		FLATTEN if(Len < 0.01f)
+		{
+			ScaleFix = 0;
+		}
+
+		OutVelocity.xy *= ScaleFix * OutVelocity.b;
 	}
 
-	// 0:background, 1:object layer
-	TempVelocity.b = AvgObject > 0.5f;
 	// alpha is used to normalize the velocity after blurring (0:background, 1:object)
-	TempVelocity.a = 1;
-
-	// 
-	TempVelocity.rg *= TempVelocity.b;
-	OutVelocity = TempVelocity;
+	OutVelocity.a = 1;
 
 	// debug, uncomment to geneate small and large horizontal motion and the objects
 //	OutVelocity.xy = lerp(float2(0,0), (InUV.y > 0.8f) ? float2(1,0) : float2(0.02f,0), OutVelocity.b);
diff --git a/Engine/Source/Runtime/Renderer/Private/PostProcess/PostProcessMotionBlur.cpp b/Engine/Source/Runtime/Renderer/Private/PostProcess/PostProcessMotionBlur.cpp
--- a/Engine/Source/Runtime/Renderer/Private/PostProcess/PostProcessMotionBlur.cpp
+++ b/Engine/Source/Runtime/Renderer/Private/PostProcess/PostProcessMotionBlur.cpp
@@ -22,6 +22,49 @@ static TAutoConsoleVariable<int32> CVarMotionBlurFiltering(
 #endif
 
 
+/** Encapsulates the post processing motion blur vertex shader. */
+class FPostProcessMotionBlurSetupVS : public FGlobalShader
+{
+	DECLARE_SHADER_TYPE(FPostProcessMotionBlurSetupVS,Global);
+
+	static bool ShouldCache(EShaderPlatform Platform)
+	{
+		return true;
+	}
+
+	/** Default constructor. */
+	FPostProcessMotionBlurSetupVS() {}
+
+	// FShader interface.
+	virtual bool Serialize(FArchive& Ar)
+	{
+		bool bShaderHasOutdatedParameters = FGlobalShader::Serialize(Ar);
+		Ar << PostprocessParameter;
+		return bShaderHasOutdatedParameters;
+	}
+
+	/** to have a similar interface as all other shaders */
+	void SetParameters(const FRenderingCompositePassContext& Context)
+	{
+		const auto ShaderRHI = GetVertexShader();
+		FGlobalShader::SetParameters(Context.RHICmdList, ShaderRHI, Context.View);
+		PostprocessParameter.SetVS(ShaderRHI, Context, TStaticSamplerState<SF_Point, AM_Clamp, AM_Clamp, AM_Clamp>::GetRHI());
+	}
+
+public:
+	FPostProcessPassParameters PostprocessParameter;
+
+	/** Initialization constructor. */
+	FPostProcessMotionBlurSetupVS(const ShaderMetaType::CompiledShaderInitializerType& Initializer)
+		: FGlobalShader(Initializer)
+	{
+		PostprocessParameter.Bind(Initializer.ParameterMap);
+	}
+};
+
+IMPLEMENT_SHADER_TYPE(,FPostProcessMotionBlurSetupVS,TEXT("PostProcessMotionBlur"),TEXT("SetupVS"),SF_Vertex);
+
+
 /** Encapsulates the post processing motion blur pixel shader. */
 class FPostProcessMotionBlurSetupPS : public FGlobalShader
 {
@@ -135,9 +178,9 @@ void FRCPassPostProcessMotionBlurSetup::Process(FRenderingCompositePassContext&
 	Context.RHICmdList.SetRasterizerState(TStaticRasterizerState<>::GetRHI());
 	Context.RHICmdList.SetDepthStencilState(TStaticDepthStencilState<false, CF_Always>::GetRHI());
 
-	TShaderMapRef<FPostProcessVS> VertexShader(Context.GetShaderMap());
+	TShaderMapRef<FPostProcessMotionBlurSetupVS> VertexShader(Context.GetShaderMap());
 
-	{		
+	{
 		TShaderMapRef<FPostProcessMotionBlurSetupPS > PixelShader(Context.GetShaderMap());
 		static FGlobalBoundShaderState BoundShaderState;
 		
