Short term fix for grid drawing on top of things

diff --git a/Engine/Shaders/BasePassPixelShader.usf b/Engine/Shaders/BasePassPixelShader.usf
--- a/Engine/Shaders/BasePassPixelShader.usf
+++ b/Engine/Shaders/BasePassPixelShader.usf
@@ -457,7 +457,7 @@ float ClipForEditorPrimitives(FMaterialPixelParameters MaterialParameters)
 #if EDITOR_PRIMITIVE_MATERIAL && (FEATURE_LEVEL >= FEATURE_LEVEL_SM4 || ES2_EMULATION)
 	// Depth test manually if compositing editor primitives since the depth buffer is different (MSAA only)
 	BRANCH
-	if (Primitive.UseEditorDepthTest)
+	if (bEnableEditorPrimitiveDepthTest)
 	{
 		bool bIsPerspective = (View.ViewToClip._m33 < 1.0f);
  		
