Fixed warning

diff --git a/Engine/Shaders/MeshParticleVertexFactory.usf b/Engine/Shaders/MeshParticleVertexFactory.usf
--- a/Engine/Shaders/MeshParticleVertexFactory.usf
+++ b/Engine/Shaders/MeshParticleVertexFactory.usf
@@ -261,7 +261,7 @@ half2 GetShadowMapCoordinate(FVertexFactoryInterpolantsVSToPS Interpolants)
 float4 CalcWorldPosition(FVertexFactoryInput Input)
 {
 	float4x4 Transform = GetParticleTransform(Input);
-	float3 WorldPosition = mul(Input.Position, Transform) + View.PreViewTranslation;
+	float3 WorldPosition = mul(Input.Position, Transform).xyz + View.PreViewTranslation;
 	return float4(WorldPosition, Input.Position.w);
 }
 
