Custom UVs for sprite particles.

#ttp 344988 - Customized UVs do not work with particle vertex factories.

diff --git a/Engine/Shaders/ParticleGPUSpriteVertexFactory.usf b/Engine/Shaders/ParticleGPUSpriteVertexFactory.usf
--- a/Engine/Shaders/ParticleGPUSpriteVertexFactory.usf
+++ b/Engine/Shaders/ParticleGPUSpriteVertexFactory.usf
@@ -50,14 +50,18 @@ struct FVertexFactoryInput
  */
 struct FVertexFactoryInterpolantsVSToPS
 {
-	/** Texture coordinate. */
-	float4 TexCoord					: TEXCOORD0;
+#if NUM_MATERIAL_TEXCOORDS
+	float4	TexCoords[(NUM_MATERIAL_TEXCOORDS + 1) / 2]	: TEXCOORD0;
+#elif USE_PARTICLE_SUBUVS
+	float4	TexCoords[1]	: TEXCOORD0;
+#endif
+
 	/** X: SubImageLerp Y:RelativeTime Z:MotionBlurFade */
-	float3 Misc						: TEXCOORD1;
+	float3 Misc						: TEXCOORD4;
 
 #if NEEDS_PARTICLE_COLOR
 	/** Sprite color. */
-	float4 Color					: TEXCOORD2;
+	float4 Color					: TEXCOORD5;
 #endif
 
 	/** Particle position and size. */
@@ -298,6 +302,27 @@ FVertexFactoryIntermediates GetVertexFactoryIntermediates( FVertexFactoryInput I
 	return Intermediates;
 }
 
+#if NUM_MATERIAL_TEXCOORDS || USE_PARTICLE_SUBUVS
+float2 GetUV(FVertexFactoryInterpolantsVSToPS Interpolants, int UVIndex)
+{
+	float4 UVVector = Interpolants.TexCoords[UVIndex / 2];
+	return UVIndex % 2 ? UVVector.zw : UVVector.xy;
+}
+
+void SetUV(inout FVertexFactoryInterpolantsVSToPS Interpolants, int UVIndex, float2 InValue)
+{
+	FLATTEN
+	if (UVIndex % 2)
+	{
+		Interpolants.TexCoords[UVIndex / 2].zw = InValue;
+	}
+	else
+	{
+		Interpolants.TexCoords[UVIndex / 2].xy = InValue;
+	}
+}
+#endif
+
 /**
  * Computes material parameterss for a given pixel.
  * @param Interpolants - Attributes interpolated from the vertex shader.
@@ -309,21 +334,12 @@ FMaterialPixelParameters GetMaterialPixelParameters( FVertexFactoryInterpolantsV
 	FMaterialPixelParameters Result = MakeInitializedMaterialPixelParameters();
 
 #if NUM_MATERIAL_TEXCOORDS
-	#if NUM_MATERIAL_TEXCOORDS >= 1
-		Result.TexCoords[0] = Interpolants.TexCoord.xy;
-	#if NUM_MATERIAL_TEXCOORDS >= 2
-		Result.TexCoords[1] = Interpolants.TexCoord.zw;
-	#endif	// >= 2
-	#endif	// >= 1
-
-	#if NUM_MATERIAL_TEXCOORDS > 3
-		UNROLL
-		for(int CoordinateIndex = 3;CoordinateIndex < NUM_MATERIAL_TEXCOORDS;CoordinateIndex++)
-		{
-			Result.TexCoords[CoordinateIndex] = Interpolants.TexCoord.xy;
-		}
-	#endif	// > 3
-#endif	//NUM_MATERIAL_TEXCOORDS
+	UNROLL
+	for (int CoordinateIndex = 0; CoordinateIndex < NUM_MATERIAL_TEXCOORDS; CoordinateIndex++)
+	{
+		Result.TexCoords[CoordinateIndex] = GetUV(Interpolants, CoordinateIndex);
+	}
+#endif	// NUM_MATERIAL_TEXCOORDS
 
 #if USE_PARTICLE_TIME
 	Result.Particle.RelativeTime = Interpolants.Misc.y;
@@ -363,8 +379,8 @@ FMaterialPixelParameters GetMaterialPixelParameters( FVertexFactoryInterpolantsV
 #endif
 
 #if USE_PARTICLE_SUBUVS
-	Result.Particle.SubUVCoords[0] = Interpolants.TexCoord.xy;
-	Result.Particle.SubUVCoords[1] = Interpolants.TexCoord.zw;
+	Result.Particle.SubUVCoords[0] = GetUV(Interpolants, 0);
+	Result.Particle.SubUVCoords[1] = GetUV(Interpolants, 1);
 	Result.Particle.SubUVLerp = Interpolants.Misc.xx;
 #endif
 
@@ -474,7 +490,22 @@ float3x3 VertexFactoryGetTangentToLocal( FVertexFactoryInput Input, FVertexFacto
 FVertexFactoryInterpolantsVSToPS VertexFactoryGetInterpolantsVSToPS( FVertexFactoryInput Input, FVertexFactoryIntermediates Intermediates, FMaterialVertexParameters VertexParameters )
 {
 	FVertexFactoryInterpolantsVSToPS Interpolants;
-	Interpolants.TexCoord = Intermediates.TexCoord;
+
+	// Initialize the whole struct to 0
+	Interpolants = (FVertexFactoryInterpolantsVSToPS)0;
+
+#if NUM_MATERIAL_TEXCOORDS
+
+	float2 CustomizedUVs[NUM_MATERIAL_TEXCOORDS];
+	GetMaterialCustomizedUVs(VertexParameters, CustomizedUVs);
+
+	UNROLL
+	for (int CoordinateIndex = 0; CoordinateIndex < NUM_MATERIAL_TEXCOORDS; CoordinateIndex++)
+	{
+		SetUV(Interpolants, CoordinateIndex, CustomizedUVs[CoordinateIndex]);
+	}
+#endif
+
 	Interpolants.Misc.x = Intermediates.SubImageLerp;
 	Interpolants.Misc.y = Intermediates.RelativeTime;
 	Interpolants.Misc.z = Intermediates.MotionBlurFade;
diff --git a/Engine/Shaders/ParticleSpriteVertexFactory.usf b/Engine/Shaders/ParticleSpriteVertexFactory.usf
--- a/Engine/Shaders/ParticleSpriteVertexFactory.usf
+++ b/Engine/Shaders/ParticleSpriteVertexFactory.usf
@@ -41,15 +41,19 @@ struct FVertexFactoryInterpolantsVSToPS
 	float SubImageLerp : TEXCOORD0;
 #endif
 
-	float4	TexCoord			: TEXCOORD2;
+#if USE_DYNAMIC_PARAMETERS
+	float4 DynamicParameter	: TEXCOORD1;
+#endif	//USE_DYNAMIC_PARAMETERS
 
 #if NEEDS_PARTICLE_COLOR
-	float4	Color				: TEXCOORD3;
+	float4	Color				: TEXCOORD2;
 #endif
 
-#if USE_DYNAMIC_PARAMETERS
-	float4 DynamicParameter	: TEXCOORD4;
-#endif	//USE_DYNAMIC_PARAMETERS
+#if NUM_MATERIAL_TEXCOORDS
+	float4	TexCoords[(NUM_MATERIAL_TEXCOORDS + 1) / 2]	: TEXCOORD3;
+#elif LIGHTMAP_UV_ACCESS || USE_PARTICLE_SUBUVS
+	float4	TexCoords[1]	: TEXCOORD3;	
+#endif
 
 #if USE_PARTICLE_POSITION
 	float4 ParticlePositionAndSize	: PARTICLE_POSITION;
@@ -107,6 +111,27 @@ struct FVertexFactoryIntermediates
 	float2 ParticleSize;
 };
 
+#if NUM_MATERIAL_TEXCOORDS || LIGHTMAP_UV_ACCESS || USE_PARTICLE_SUBUVS
+float2 GetUV(FVertexFactoryInterpolantsVSToPS Interpolants, int UVIndex)
+{
+	float4 UVVector = Interpolants.TexCoords[UVIndex / 2];
+	return UVIndex % 2 ? UVVector.zw : UVVector.xy;
+}
+
+void SetUV(inout FVertexFactoryInterpolantsVSToPS Interpolants, int UVIndex, float2 InValue)
+{
+	FLATTEN
+	if (UVIndex % 2)
+	{
+		Interpolants.TexCoords[UVIndex / 2].zw = InValue;
+	}
+	else
+	{
+		Interpolants.TexCoords[UVIndex / 2].xy = InValue;
+	}
+}
+#endif
+
 /** Converts from vertex factory specific interpolants to a FMaterialPixelParameters, which is used by material inputs. */
 FMaterialPixelParameters GetMaterialPixelParameters(FVertexFactoryInterpolantsVSToPS Interpolants, float4 PixelPosition)
 {
@@ -114,21 +139,13 @@ FMaterialPixelParameters GetMaterialPixelParameters(FVertexFactoryInterpolantsVS
 	FMaterialPixelParameters Result = MakeInitializedMaterialPixelParameters();
 
 #if NUM_MATERIAL_TEXCOORDS
-	#if NUM_MATERIAL_TEXCOORDS >= 1
-		Result.TexCoords[0] = Interpolants.TexCoord.xy;
-	#if NUM_MATERIAL_TEXCOORDS >= 2
-		Result.TexCoords[1] = Interpolants.TexCoord.zw;
-	#endif	// >= 2
-	#endif	// >= 1
+	UNROLL
+	for (int CoordinateIndex = 0; CoordinateIndex < NUM_MATERIAL_TEXCOORDS; CoordinateIndex++)
+	{
+		Result.TexCoords[CoordinateIndex] = GetUV(Interpolants, CoordinateIndex);
+	}
+#endif	// NUM_MATERIAL_TEXCOORDS
 
-	#if NUM_MATERIAL_TEXCOORDS > 3
-		UNROLL
-		for(int CoordinateIndex = 3;CoordinateIndex < NUM_MATERIAL_TEXCOORDS;CoordinateIndex++)
-		{
-			Result.TexCoords[CoordinateIndex] = Interpolants.TexCoord.xy;
-		}
-	#endif	// > 3
-#endif	//NUM_MATERIAL_TEXCOORDS
 	Result.VertexColor = 1;
 
 #if NEEDS_TANGENT_BASIS
@@ -171,7 +188,7 @@ FMaterialPixelParameters GetMaterialPixelParameters(FVertexFactoryInterpolantsVS
 	Result.Particle.MotionBlurFade = 1.0f;
 
 #if LIGHTMAP_UV_ACCESS
-	Result.LightmapUVs = Interpolants.TexCoord.xy;
+	Result.LightmapUVs = GetUV(Interpolants,0);
 #endif
 
 #if USE_PARTICLE_LIGHTING_OFFSET
@@ -179,8 +196,8 @@ FMaterialPixelParameters GetMaterialPixelParameters(FVertexFactoryInterpolantsVS
 #endif
 
 #if USE_PARTICLE_SUBUVS
-	Result.Particle.SubUVCoords[0] = Interpolants.TexCoord.xy;
-	Result.Particle.SubUVCoords[1] = Interpolants.TexCoord.zw;
+	Result.Particle.SubUVCoords[0] = GetUV(Interpolants, 0);
+	Result.Particle.SubUVCoords[1] = GetUV(Interpolants, 1);
 	Result.Particle.SubUVLerp = SubImageLerp;
 #endif
 
@@ -432,7 +449,21 @@ FVertexFactoryInterpolantsVSToPS VertexFactoryGetInterpolantsVSToPS(FVertexFacto
 {
 	FVertexFactoryInterpolantsVSToPS Interpolants;
 
-	Interpolants.TexCoord = Intermediates.TexCoord;
+	// Initialize the whole struct to 0
+	Interpolants = (FVertexFactoryInterpolantsVSToPS)0;
+
+#if NUM_MATERIAL_TEXCOORDS
+
+	float2 CustomizedUVs[NUM_MATERIAL_TEXCOORDS];
+	GetMaterialCustomizedUVs(VertexParameters, CustomizedUVs);
+
+	UNROLL
+	for (int CoordinateIndex = 0; CoordinateIndex < NUM_MATERIAL_TEXCOORDS; CoordinateIndex++)
+	{
+		SetUV(Interpolants, CoordinateIndex, CustomizedUVs[CoordinateIndex]);
+	}
+#endif
+
 #if NEEDS_TANGENT_BASIS
 	// Calculate the transform from tangent to world space.
 	// Note that "local" space for particles is actually oriented in world space! Therefore no rotation is needed.
