UE4 - Fix ES 2 shaders (modulo doesn't work on GLSL ES2)

diff --git a/Engine/Shaders/ParticleSpriteVertexFactory.usf b/Engine/Shaders/ParticleSpriteVertexFactory.usf
--- a/Engine/Shaders/ParticleSpriteVertexFactory.usf
+++ b/Engine/Shaders/ParticleSpriteVertexFactory.usf
@@ -112,16 +112,25 @@ struct FVertexFactoryIntermediates
 };
 
 #if NUM_MATERIAL_TEXCOORDS || LIGHTMAP_UV_ACCESS || USE_PARTICLE_SUBUVS
+bool UVIndexUseZW(int UVIndex)
+{
+	#if COMPILER_GLSL_ES2
+	return (frac((float)UVIndex / 2.0f) == 0.5f);
+	#else
+	return (UVIndex % 2) != 0;
+	#endif
+}
+
 float2 GetUV(FVertexFactoryInterpolantsVSToPS Interpolants, int UVIndex)
 {
 	float4 UVVector = Interpolants.TexCoords[UVIndex / 2];
-	return UVIndex % 2 ? UVVector.zw : UVVector.xy;
+	return UVIndexUseZW(UVIndex) ? UVVector.zw : UVVector.xy;
 }
 
 void SetUV(inout FVertexFactoryInterpolantsVSToPS Interpolants, int UVIndex, float2 InValue)
 {
 	FLATTEN
-	if (UVIndex % 2)
+	if (UVIndexUseZW(UVIndex))
 	{
 		Interpolants.TexCoords[UVIndex / 2].zw = InValue;
 	}
