Fix for shader errors caused by using extra texcoords interpolants in some particle shaders.

#ttp 346312 - Pri 0 - Do Not Leave Desk (blocker!)) TTP #346312 was assigned to Nick Penwarden (BLOCKER: CRITICAL: Regression: Many particles with translucent or additive materials do not render correctly

diff --git a/Engine/Shaders/TranslucentShadowDepthShaders.usf b/Engine/Shaders/TranslucentShadowDepthShaders.usf
--- a/Engine/Shaders/TranslucentShadowDepthShaders.usf
+++ b/Engine/Shaders/TranslucentShadowDepthShaders.usf
@@ -10,8 +10,8 @@
 struct FTranslucencyShadowDepthVSToPS
 {
 	FVertexFactoryInterpolantsVSToPS FactoryInterpolants;
-	float ShadowDepth : TEXCOORD5;
-	float4 PixelPosition	: TEXCOORD6;
+	float ShadowDepth : TEXCOORD6;
+	float4 PixelPosition	: TEXCOORD7;
 };
 
 float4x4 ProjectionMatrix;
