Distance field AO is now using cone tracing instead of cone stepping
* Provides better quality with fewer steps, reduces artifacts in heavly occluded situations (indoors)
* Saves about .5ms on a 7870 at 1080p (4.6ms -> 4.1ms)

diff --git a/Engine/Shaders/DistanceFieldLightingShared.usf b/Engine/Shaders/DistanceFieldLightingShared.usf
--- a/Engine/Shaders/DistanceFieldLightingShared.usf
+++ b/Engine/Shaders/DistanceFieldLightingShared.usf
@@ -22,8 +22,8 @@
 #define SORTTHREADGROUP_SIZEY 1
 #endif
 
-#ifndef NUM_SAMPLES
-#define NUM_SAMPLES 1
+#ifndef NUM_CONE_DIRECTIONS
+#define NUM_CONE_DIRECTIONS 1
 #endif
 
 #ifndef DOWNSAMPLE_FACTOR
@@ -48,7 +48,7 @@ float AOMaxViewDistance;
 
 #define NUM_CULLED_OBJECT_LISTS 3
 
-float GetPhaseParameters(uint Phase, out uint StartIndex, out uint EndIndex)
+void GetPhaseParameters(uint Phase, out uint StartIndex, out uint EndIndex)
 {
 	if (Phase == 0)
 	{
@@ -65,13 +65,11 @@ float GetPhaseParameters(uint Phase, out uint StartIndex, out uint EndIndex)
 		StartIndex = 8;
 		EndIndex = 9;
 	}
-
-	return AOStepScale;
 }
 
-float GetStepOffset(uint StepIndex, float StepScale)
+float GetStepOffset(float StepIndex)
 {
-	return StepScale * exp2(AOStepExponentScale * StepIndex);
+	return AOStepScale * exp2(AOStepExponentScale * StepIndex);
 }
 
 float3 DistanceFieldVolumePositionToUV(float3 VolumePosition, float3 UVScale, float3 UVAdd)
diff --git a/Engine/Shaders/DistanceFieldShadowing.usf b/Engine/Shaders/DistanceFieldShadowing.usf
--- a/Engine/Shaders/DistanceFieldShadowing.usf
+++ b/Engine/Shaders/DistanceFieldShadowing.usf
@@ -157,97 +157,115 @@ void DistanceFieldShadowingCS(
 
 	GroupMemoryBarrierWithGroupSync();
 
+	uint GroupIndex = SceneDepth > MaxTileZ2 ? 1 : 0;
+	uint NumIntersectingObjects = GroupIndex == 0 ? TileNumObjects0 : TileNumObjects1;
+	
 	float Result = 0;
 
+#define COMPUTE_SHADOWING 1
+#if COMPUTE_SHADOWING
 	BRANCH
 	if (SceneDepth > MinDepth)
 	{
-		float4 HomogeneousWorldPosition = mul(float4(ScreenPosition * SceneDepth, SceneDepth, 1), View.ScreenToWorld);
-		float3 OpaqueWorldPosition = HomogeneousWorldPosition.xyz / HomogeneousWorldPosition.w;
-
-		float TraceDistance = 10000;
-		float RayStartOffset = 10;
-		float3 WorldRayStart = OpaqueWorldPosition - LightDirection * RayStartOffset;
-		float3 WorldRayEnd = OpaqueWorldPosition - LightDirection * TraceDistance;
-		float3 WorldRayDirection = WorldRayEnd - WorldRayStart;
-		float3 UnitWorldRayDirection = normalize(WorldRayDirection);
-		uint GroupIndex = SceneDepth > MaxTileZ2 ? 1 : 0;
-		uint NumIntersectingObjects = GroupIndex == 0 ? TileNumObjects0 : TileNumObjects1;
+		FGBufferData GBufferData = GetGBufferData(ScreenUV);
+		float3 WorldNormal = GBufferData.WorldNormal;
 
 		//@todo - expose
-		float SunAngle = 1 * PI / 180.0f;
-		float TanSunAngle = tan(SunAngle);
-		float MinRayTime = TraceDistance;
-		float ClosestDistanceToIntersection = 1000;
-		float TotalStepsTaken = 0;
-		float MinConeVisibility = 1;
-
-		LOOP
-		for (uint ListObjectIndex = 0; ListObjectIndex < NumIntersectingObjects && MinRayTime >= TraceDistance; ListObjectIndex++)
-		{
-			uint ObjectIndex = IntersectingObjectIndices[MAX_INTERSECTING_OBJECTS * GroupIndex + ListObjectIndex];
-			float3 LocalPositionExtent = LoadObjectLocalPositionExtent(ObjectIndex);
-			float4x4 WorldToVolume = LoadObjectWorldToVolume(ObjectIndex);
-			float4 UVScaleAndVolumeScale = LoadObjectUVScale(ObjectIndex);
-			float3 UVAdd = LoadObjectUVAdd(ObjectIndex);
-
-			float3 VolumeRayStart = mul(float4(WorldRayStart, 1), WorldToVolume).xyz;
-			float3 VolumeRayEnd = mul(float4(WorldRayEnd, 1), WorldToVolume).xyz;
-			float3 VolumeRayDirection = VolumeRayEnd - VolumeRayStart;
-			float VolumeRayLength = length(VolumeRayDirection);
-			VolumeRayDirection /= VolumeRayLength;
+		float SunAngle = 2 * PI / 180.0f;
 
-			float2 IntersectionTimes = LineBoxIntersect(VolumeRayStart, VolumeRayEnd, -LocalPositionExtent, LocalPositionExtent);
-
-			if (IntersectionTimes.x < IntersectionTimes.y && IntersectionTimes.x < 1)
+		BRANCH
+		if (dot(WorldNormal, -LightDirection) > cos(PI / 2 + SunAngle))
+		{
+			float4 HomogeneousWorldPosition = mul(float4(ScreenPosition * SceneDepth, SceneDepth, 1), View.ScreenToWorld);
+			float3 OpaqueWorldPosition = HomogeneousWorldPosition.xyz / HomogeneousWorldPosition.w;
+
+			float TraceDistance = 10000;
+			float RayStartOffset = 2;
+			float3 WorldRayStart = OpaqueWorldPosition - LightDirection * RayStartOffset;
+			float3 WorldRayEnd = OpaqueWorldPosition - LightDirection * TraceDistance;
+			float3 WorldRayDirection = WorldRayEnd - WorldRayStart;
+			float3 UnitWorldRayDirection = normalize(WorldRayDirection);
+
+			float TanSunAngle = tan(SunAngle);
+			float MinRayTime = TraceDistance;
+			float ClosestDistanceToIntersection = 1000;
+			float TotalStepsTaken = 0;
+			float MinConeVisibility = 1;
+
+			LOOP
+			for (uint ListObjectIndex = 0; ListObjectIndex < NumIntersectingObjects && MinRayTime >= TraceDistance; ListObjectIndex++)
 			{
-				float SampleRayTime = IntersectionTimes.x * VolumeRayLength;
+				uint ObjectIndex = IntersectingObjectIndices[MAX_INTERSECTING_OBJECTS * GroupIndex + ListObjectIndex];
+				float3 LocalPositionExtent = LoadObjectLocalPositionExtent(ObjectIndex);
+				float4x4 WorldToVolume = LoadObjectWorldToVolume(ObjectIndex);
+				float4 UVScaleAndVolumeScale = LoadObjectUVScale(ObjectIndex);
+				float3 UVAdd = LoadObjectUVAdd(ObjectIndex);
 
-				float MinDistance = 1000000;
-				float3 IntersectionPosition = float3(0, 0, 0);
+				float3 VolumeRayStart = mul(float4(WorldRayStart, 1), WorldToVolume).xyz;
+				float3 VolumeRayEnd = mul(float4(WorldRayEnd, 1), WorldToVolume).xyz;
+				float3 VolumeRayDirection = VolumeRayEnd - VolumeRayStart;
+				float VolumeRayLength = length(VolumeRayDirection);
+				VolumeRayDirection /= VolumeRayLength;
 
-				uint StepIndex = 0;
-				uint MaxSteps = 256;
+				float2 IntersectionTimes = LineBoxIntersect(VolumeRayStart, VolumeRayEnd, -LocalPositionExtent, LocalPositionExtent);
 
-				LOOP
-				for (; StepIndex < MaxSteps; StepIndex++)
+				if (IntersectionTimes.x < IntersectionTimes.y && IntersectionTimes.x < 1)
 				{
-					float3 SampleVolumePosition = VolumeRayStart + VolumeRayDirection * SampleRayTime;
-					float3 ClampedSamplePosition = clamp(SampleVolumePosition, -LocalPositionExtent, LocalPositionExtent);
-					float3 VolumeUV = DistanceFieldVolumePositionToUV(ClampedSamplePosition, UVScaleAndVolumeScale.xyz, UVAdd);
-					float DistanceField = Texture3DSampleLevel(DistanceFieldTexture, DistanceFieldSampler, VolumeUV, 0).x;
-					MinDistance = min(MinDistance, DistanceField);
-					MinConeVisibility = min(MinConeVisibility, saturate(DistanceField / (TanSunAngle * SampleRayTime)));
-					IntersectionPosition = SampleVolumePosition;
-
+					float SampleRayTime = IntersectionTimes.x * VolumeRayLength;
+					uint MaxSteps = 64;
 					float MinStepSize = 1.0f / (4 * MaxSteps);
-					float StepDistance = max(DistanceField, MinStepSize);
-					SampleRayTime += StepDistance;
 
-					// Terminate the trace if we reached a negative area or went past the end of the ray
-					if (DistanceField < 0 
-						|| SampleRayTime > IntersectionTimes.y * VolumeRayLength)
+					float MinDistance = 1000000;
+					float3 IntersectionPosition = float3(0, 0, 0);
+
+					uint StepIndex = 0;
+
+					LOOP
+					for (; StepIndex < MaxSteps; StepIndex++)
 					{
-						break;
+						float3 SampleVolumePosition = VolumeRayStart + VolumeRayDirection * SampleRayTime;
+						float3 ClampedSamplePosition = clamp(SampleVolumePosition, -LocalPositionExtent, LocalPositionExtent);
+						float3 VolumeUV = DistanceFieldVolumePositionToUV(ClampedSamplePosition, UVScaleAndVolumeScale.xyz, UVAdd);
+						float DistanceField = Texture3DSampleLevel(DistanceFieldTexture, DistanceFieldSampler, VolumeUV, 0).x;
+
+						MinDistance = min(MinDistance, DistanceField);
+						float SphereRadius = TanSunAngle * SampleRayTime;
+
+						MinConeVisibility = min(MinConeVisibility, saturate(DistanceField / SphereRadius));
+						IntersectionPosition = SampleVolumePosition;
+
+						float StepDistance = max(DistanceField, MinStepSize);
+						SampleRayTime += StepDistance;
+
+						// Terminate the trace if we reached a negative area or went past the end of the ray
+						if (DistanceField < 0 
+							|| SampleRayTime > IntersectionTimes.y * VolumeRayLength)
+						{
+							break;
+						}
 					}
-				}
 
-				if (MinDistance * UVScaleAndVolumeScale.w < 0 || StepIndex == MaxSteps)
-				{
-					MinRayTime = min(MinRayTime, SampleRayTime * UVScaleAndVolumeScale.w);
-				}
+					if (MinDistance * UVScaleAndVolumeScale.w < 0 || StepIndex == MaxSteps)
+					{
+						MinRayTime = min(MinRayTime, SampleRayTime * UVScaleAndVolumeScale.w);
+					}
 
-				if (SampleRayTime * UVScaleAndVolumeScale.w > 100)
-				{
-					ClosestDistanceToIntersection = min(ClosestDistanceToIntersection, MinDistance * UVScaleAndVolumeScale.w);
+					if (SampleRayTime * UVScaleAndVolumeScale.w > 100)
+					{
+						ClosestDistanceToIntersection = min(ClosestDistanceToIntersection, MinDistance * UVScaleAndVolumeScale.w);
+					}
+					TotalStepsTaken += StepIndex;
 				}
-				TotalStepsTaken += StepIndex;
 			}
-		}
 
-		Result = MinConeVisibility;
+			Result = MinConeVisibility;
+		}
 	}
 
+#else
+	Result = .0001f * NumIntersectingObjects;
+#endif
+
 	//Result = (GroupIndex == 0 ? Group0BoundingRadius : Group1BoundingRadius) / 1000.0f;
 	//Result = TotalStepsTaken / 1000.0f;
 	//Result = NumIntersectingObjects / 100.0f;
diff --git a/Engine/Shaders/DistanceFieldSurfaceCacheLighting.usf b/Engine/Shaders/DistanceFieldSurfaceCacheLighting.usf
--- a/Engine/Shaders/DistanceFieldSurfaceCacheLighting.usf
+++ b/Engine/Shaders/DistanceFieldSurfaceCacheLighting.usf
@@ -71,8 +71,8 @@ void ObjectCullPS(
 	{
 		uint StartIndex;
 		uint EndIndex;
-		float StepScale = GetPhaseParameters(GroupIndex, StartIndex, EndIndex);
-		float GroupMaxSampleRadius = GetStepOffset(EndIndex, StepScale) * 2 * RadiusScale;
+		GetPhaseParameters(GroupIndex, StartIndex, EndIndex);
+		float GroupMaxSampleRadius = GetStepOffset(EndIndex) * 2 * RadiusScale;
 	
 		BRANCH
 		if (SphereIntersectConeWithDepthRanges(float4(ViewSpaceSphereCenter, SphereCenterAndRadius.w + GroupMaxSampleRadius), TileConeVertex, TileConeAxis, TileConeAngleCos, TileConeAngleSin, ConeAxisDepthRanges))
diff --git a/Engine/Shaders/DistanceFieldSurfaceCacheLightingCompute.usf b/Engine/Shaders/DistanceFieldSurfaceCacheLightingCompute.usf
--- a/Engine/Shaders/DistanceFieldSurfaceCacheLightingCompute.usf
+++ b/Engine/Shaders/DistanceFieldSurfaceCacheLightingCompute.usf
@@ -243,24 +243,24 @@ void DistanceFieldAOBuildTileListMain(
 	{
 		uint StartIndex;
 		uint EndIndex;
-		float StepScale = GetPhaseParameters(0, StartIndex, EndIndex);
-		SmallGroupMaxSampleRadius = GetStepOffset(EndIndex, StepScale) * 2 * RadiusScale;
+		GetPhaseParameters(0, StartIndex, EndIndex);
+		SmallGroupMaxSampleRadius = GetStepOffset(EndIndex) * 2 * RadiusScale;
 	}
 
 	float MediumGroupMaxSampleRadius;
 	{
 		uint StartIndex;
 		uint EndIndex;
-		float StepScale = GetPhaseParameters(1, StartIndex, EndIndex);
-		MediumGroupMaxSampleRadius = GetStepOffset(EndIndex, StepScale) * 2 * RadiusScale;
+		GetPhaseParameters(1, StartIndex, EndIndex);
+		MediumGroupMaxSampleRadius = GetStepOffset(EndIndex) * 2 * RadiusScale;
 	}
 
 	float LargeGroupMaxSampleRadius;
 	{
 		uint StartIndex;
 		uint EndIndex;
-		float StepScale = GetPhaseParameters(2, StartIndex, EndIndex);
-		LargeGroupMaxSampleRadius = GetStepOffset(EndIndex, StepScale) * 2 * RadiusScale;
+		GetPhaseParameters(2, StartIndex, EndIndex);
+		LargeGroupMaxSampleRadius = GetStepOffset(EndIndex) * 2 * RadiusScale;
 	}
 
 	// Compute per-tile lists of affecting objects through bounds culling
@@ -459,11 +459,23 @@ Buffer<uint> SavedStartIndex;
 RWBuffer<float> RWOccluderRadius;
 RWBuffer<float4> RWIrradianceCacheBentNormal;
 
-#define FINAL_GATHER_THREADGROUP_SIZE 64
-#define SIMULTANEOUSLY_TRACED_OBJECTS 4
+// Whether to use cone tracing or cone stepping
+#define USE_CONE_TRACING 1
+
+// Have to disable surface caching dependencies for this to work
+#define VISUALIZE_ONE_CONE 0
+
+#define FINAL_GATHER_THREADGROUP_SIZE 256
+
+#if USE_CONE_TRACING
+	// Enforce one thread per cone direction
+	#define SIMULTANEOUSLY_TRACED_OBJECTS (FINAL_GATHER_THREADGROUP_SIZE / NUM_CONE_DIRECTIONS)
+#else
+	#define SIMULTANEOUSLY_TRACED_OBJECTS 4
+#endif
 #define THREADS_PER_OBJECT (FINAL_GATHER_THREADGROUP_SIZE / SIMULTANEOUSLY_TRACED_OBJECTS)
 
-groupshared uint SharedConeVisibility[NUM_SAMPLES];
+groupshared uint SharedConeVisibility[NUM_CONE_DIRECTIONS];
 groupshared uint SharedMinOcclusionDistance;
 
 /** Computes ambient occlusion for a surface cache record by cone stepping through the nearby object distance fields. */
@@ -478,7 +490,7 @@ void FinalGatherCS(
 
 	if (ThreadIndex == 0)
 	{
-		for (uint ConeIndex = 0; ConeIndex < NUM_SAMPLES; ConeIndex++) 
+		for (uint ConeIndex = 0; ConeIndex < NUM_CONE_DIRECTIONS; ConeIndex++) 
 		{
 			SharedConeVisibility[ConeIndex] = asuint(1.0f);
 		}
@@ -507,12 +519,28 @@ void FinalGatherCS(
 
 		FindBestAxisVectors2(WorldNormal, TangentX, TangentY);
 
+		float3 DebugConeDirection = normalize(float3(.3f, .5f, .4f));
+
+#if USE_CONE_TRACING
+		uint ConeIndex = ThreadIndex % THREADS_PER_OBJECT;
+
+		float3 ConeDirection = AOSamples2.SampleDirections[ConeIndex].xyz;
+		float3 RotatedConeDirection = ConeDirection.x * TangentX + ConeDirection.y * TangentY + ConeDirection.z * WorldNormal;
+
+#if VISUALIZE_ONE_CONE
+		RotatedConeDirection = DebugConeDirection;
+#endif
+
+		float MaxWorldStepOffset = GetStepOffset(NUM_CONE_STEPS);
+#endif
 		LOOP 
 		for (uint ListObjectIndex = 0; ListObjectIndex < NumObjectsAffectingTile; ListObjectIndex += SIMULTANEOUSLY_TRACED_OBJECTS)
 		{
 			uint EffectiveListObjectIndex = ListObjectIndex + ObjectOffsetIndex;
 
-			if (EffectiveListObjectIndex < NumObjectsAffectingTile)
+			if (EffectiveListObjectIndex < NumObjectsAffectingTile 
+				// Ignore extra threads
+				&& ObjectOffsetIndex < SIMULTANEOUSLY_TRACED_OBJECTS)
 			{
 #if COHERENT_OBJECT_LIST_WRITES
 
@@ -561,10 +589,52 @@ void FinalGatherCS(
 						StartStepIndex = 5;
 					}
 
+#if USE_CONE_TRACING
+					float WorldStepOffset = GetStepOffset(StartStepIndex);
+
+					LOOP
+					for (uint StepIndex = StartStepIndex; StepIndex < NUM_CONE_STEPS && WorldStepOffset < MaxWorldStepOffset; StepIndex++)
+					{
+						float3 WorldSamplePosition = WorldShadingPosition + RotatedConeDirection * WorldStepOffset;
+						float3 StepSamplePosition = mul(float4(WorldSamplePosition, 1), WorldToVolume).xyz;
+						float3 ClampedSamplePosition = clamp(StepSamplePosition, -LocalPositionExtent, LocalPositionExtent);
+						float DistanceToClamped = length(StepSamplePosition - ClampedSamplePosition);
+
+						float3 StepVolumeUV = DistanceFieldVolumePositionToUV(ClampedSamplePosition, UVScaleAndVolumeScale.xyz, UVAdd);
+						float DistanceToOccluder = (Texture3DSampleLevel(DistanceFieldTexture, DistanceFieldSampler, StepVolumeUV, 0).x + DistanceToClamped) * UVScaleAndVolumeScale.w;
+						float SphereRadius = WorldStepOffset * TanConeHalfAngle;
+						//@todo - have to bias away from surface further for this to work
+						float ShadingSphereRadius = SphereRadius * 1.0f;
+
+						// Derive visibility from 1d intersection
+						float Visibility = saturate(DistanceToOccluder / ShadingSphereRadius);
+
+						// Don't allow small objects to fully occlude a cone step
+						Visibility = max(Visibility, 1 - saturate(ObjectOccluderRadius / SphereRadius));
+						float OccluderDistanceFraction = (WorldStepOffset + DistanceToOccluder) / AOMaxDistance;
+						// Fade out occlusion based on distance to occluder
+						Visibility = max(Visibility, saturate((OccluderDistanceFraction - 0) / 4));
+						Visibility = max(Visibility, saturate((OccluderDistanceFraction - .7f) / .4f));
+
+						InterlockedMin(SharedConeVisibility[ConeIndex], asuint(Visibility));
+
+						if (DistanceToOccluder < .9f * SphereRadius)
+						{
+							// Assuming occluder is straight forward along the cone
+							float WorldDistanceToOccluder = WorldStepOffset + DistanceToOccluder;
+							InterlockedMin(SharedMinOcclusionDistance, asuint(WorldDistanceToOccluder));
+						}
+
+						float MinStepSize = .6f * (GetStepOffset(StepIndex + 1) - GetStepOffset(StepIndex));
+						WorldStepOffset += max(DistanceToOccluder, MinStepSize);
+					}
+
+#else // Cone stepping
+
 					uint NumConeSteps = NUM_CONE_STEPS - StartStepIndex;
 
 					LOOP
-					for (uint SampleIndex = ThreadIndex % THREADS_PER_OBJECT; SampleIndex < NUM_SAMPLES * NumConeSteps; SampleIndex += THREADS_PER_OBJECT)
+					for (uint SampleIndex = ThreadIndex % THREADS_PER_OBJECT; SampleIndex < NUM_CONE_DIRECTIONS * NumConeSteps; SampleIndex += THREADS_PER_OBJECT)
 					{
 						// Avoid uint divide
 						uint ConeIndex = (uint)(SampleIndex / (float)NumConeSteps);
@@ -573,6 +643,10 @@ void FinalGatherCS(
 						float3 ConeDirection = AOSamples2.SampleDirections[ConeIndex].xyz;
 						float3 RotatedConeDirection = ConeDirection.x * TangentX + ConeDirection.y * TangentY + ConeDirection.z * WorldNormal;
 
+						#if VISUALIZE_ONE_CONE
+							RotatedConeDirection = DebugConeDirection;
+						#endif
+
 						float WorldStepOffset = AOStepScale * exp2(AOStepExponentScale * StepIndex);
 						float3 WorldSamplePosition = WorldShadingPosition + RotatedConeDirection * WorldStepOffset;
 						float3 StepSamplePosition = mul(float4(WorldSamplePosition, 1), WorldToVolume).xyz;
@@ -583,6 +657,7 @@ void FinalGatherCS(
 						float DistanceToOccluder = (Texture3DSampleLevel(DistanceFieldTexture, DistanceFieldSampler, StepVolumeUV, 0).x + DistanceToClamped) * UVScaleAndVolumeScale.w;
 
 						float SphereRadius = WorldStepOffset * TanConeHalfAngle;
+						float ShadingSphereRadius = SphereRadius * 1.0f;
 
 						#define CIRCLE_LINE_INTERSECTION 0
 
@@ -603,7 +678,7 @@ void FinalGatherCS(
 							// Allows use of negative distances in the small steps to create full penumbra range, reduces over-occlusion
 							float MinOcclusionFraction = 0;//-.6f * (1 - saturate(4 * StepIndex / (float)NUM_CONE_STEPS));
 							// Derive visibility from 1d intersection
-							float Visibility = saturate((DistanceToOccluder / SphereRadius - MinOcclusionFraction) / (1 - MinOcclusionFraction));
+							float Visibility = saturate((DistanceToOccluder / ShadingSphereRadius - MinOcclusionFraction) / (1 - MinOcclusionFraction));
 						#endif
 
 						// Don't allow small objects to fully occlude a cone step
@@ -618,6 +693,7 @@ void FinalGatherCS(
 							InterlockedMin(SharedMinOcclusionDistance, asuint(WorldDistanceToOccluder));
 						}
 					}
+#endif
 				}
 			}
 		}
@@ -629,7 +705,7 @@ void FinalGatherCS(
 	{
 		float3 UnoccludedDirection = 0;
 		float Visibility = 0;
-		float NumSamplesToDisplay = NUM_SAMPLES;
+		float NumSamplesToDisplay = VISUALIZE_ONE_CONE ? 1 : NUM_CONE_DIRECTIONS;
 
 		for (uint ConeIndex = 0; ConeIndex < NumSamplesToDisplay; ConeIndex++)
 		{
diff --git a/Engine/Source/Runtime/Engine/Private/DirectionalLightComponent.cpp b/Engine/Source/Runtime/Engine/Private/DirectionalLightComponent.cpp
--- a/Engine/Source/Runtime/Engine/Private/DirectionalLightComponent.cpp
+++ b/Engine/Source/Runtime/Engine/Private/DirectionalLightComponent.cpp
@@ -187,8 +187,7 @@ public:
 	/** Returns the number of view dependent shadows this light will create. */
 	virtual int32 GetNumViewDependentWholeSceneShadows(const FSceneView& View) const 
 	{ 
-		int32 NumCascades = WholeSceneDynamicShadowRadius > 0.0f ? DynamicShadowCascades : 0;
-		NumCascades = FMath::Min<int32>(NumCascades, View.MaxShadowCascades);
+		int32 NumCascades = GetNumShadowMappedCascades();
 
 		if (ShouldCreateRayTracedCascade())
 		{
@@ -298,6 +297,12 @@ public:
 
 private:
 
+	int32 GetNumShadowMappedCascades() const
+	{
+		int32 NumCascades = WholeSceneDynamicShadowRadius > 0.0f ? DynamicShadowCascades : 0;
+		return FMath::Min<int32>(NumCascades, View.MaxShadowCascades);
+	}
+
 	float GetCSMMaxDistance() const
 	{
 		static const auto CVar = IConsoleManager::Get().FindTConsoleVariableDataFloat(TEXT("r.Shadow.DistanceScale"));
@@ -322,7 +327,11 @@ private:
 
 	bool ShouldCreateRayTracedCascade() const
 	{
-		return GetDistanceFieldShadowDistance() > GetCSMMaxDistance() && DoesPlatformSupportDistanceFieldShadowing(GRHIShaderPlatform);
+		const int32 NumCascades = GetNumShadowMappedCascades();
+		const float RaytracedShadowDistance = GetDistanceFieldShadowDistance();
+		const bool bCreateWithCSM = NumCascades > 0 && RaytracedShadowDistance > GetCSMMaxDistance();
+		const bool bCreateWithoutCSM = NumCascades == 0 && RaytracedShadowDistance > 0;
+		return DoesPlatformSupportDistanceFieldShadowing(GRHIShaderPlatform) && (bCreateWithCSM || bCreateWithoutCSM);
 	}
 
 	float ComputeShadowDistance() const
diff --git a/Engine/Source/Runtime/Renderer/Private/DistanceFieldLightingShared.h b/Engine/Source/Runtime/Renderer/Private/DistanceFieldLightingShared.h
--- a/Engine/Source/Runtime/Renderer/Private/DistanceFieldLightingShared.h
+++ b/Engine/Source/Runtime/Renderer/Private/DistanceFieldLightingShared.h
@@ -8,9 +8,6 @@
 
 #include "DistanceFieldAtlas.h"
 
-/** Maximum object / tile interactions to allocate memory for. */
-static const int32 GMaxTileShadingJobs = 500000;
-
 /** Tile sized used for most AO compute shaders. */
 const int32 GDistanceFieldAOTileSizeX = 8;
 const int32 GDistanceFieldAOTileSizeY = 8;
diff --git a/Engine/Source/Runtime/Renderer/Private/DistanceFieldSurfaceCacheLighting.cpp b/Engine/Source/Runtime/Renderer/Private/DistanceFieldSurfaceCacheLighting.cpp
--- a/Engine/Source/Runtime/Renderer/Private/DistanceFieldSurfaceCacheLighting.cpp
+++ b/Engine/Source/Runtime/Renderer/Private/DistanceFieldSurfaceCacheLighting.cpp
@@ -303,7 +303,7 @@ public:
 		FGlobalShader::ModifyCompilationEnvironment(Platform,OutEnvironment);
 		OutEnvironment.SetDefine(TEXT("THREADGROUP_SIZEX"), GDistanceFieldAOTileSizeX);
 		OutEnvironment.SetDefine(TEXT("THREADGROUP_SIZEY"), GDistanceFieldAOTileSizeY);
-		OutEnvironment.SetDefine(TEXT("NUM_SAMPLES"), NumConeSampleDirections);
+		OutEnvironment.SetDefine(TEXT("NUM_CONE_DIRECTIONS"), NumConeSampleDirections);
 		OutEnvironment.SetDefine(TEXT("DOWNSAMPLE_FACTOR"), GAODownsampleFactor);
 		OutEnvironment.SetDefine(TEXT("MAX_OBJECTS_PER_TILE"), GMaxNumObjectsPerTile);
 		
@@ -547,7 +547,7 @@ public:
 		FGlobalShader::ModifyCompilationEnvironment(Platform,OutEnvironment);
 		OutEnvironment.SetDefine(TEXT("THREADGROUP_SIZEX"), GDistanceFieldAOTileSizeX);
 		OutEnvironment.SetDefine(TEXT("THREADGROUP_SIZEY"), GDistanceFieldAOTileSizeY);
-		OutEnvironment.SetDefine(TEXT("NUM_SAMPLES"), NumConeSampleDirections);
+		OutEnvironment.SetDefine(TEXT("NUM_CONE_DIRECTIONS"), NumConeSampleDirections);
 		OutEnvironment.SetDefine(TEXT("DOWNSAMPLE_FACTOR"), GAODownsampleFactor);
 		OutEnvironment.SetDefine(TEXT("MAX_OBJECTS_PER_TILE"), GMaxNumObjectsPerTile);
 		// To reduce shader compile time of compute shaders with shared memory, doesn't have an impact on generated code with current compiler (June 2010 DX SDK)
@@ -1149,7 +1149,7 @@ public:
 		FGlobalShader::ModifyCompilationEnvironment(Platform,OutEnvironment);
 		OutEnvironment.SetDefine(TEXT("THREADGROUP_SIZEX"), GDistanceFieldAOTileSizeX);
 		OutEnvironment.SetDefine(TEXT("THREADGROUP_SIZEY"), GDistanceFieldAOTileSizeY);
-		OutEnvironment.SetDefine(TEXT("NUM_SAMPLES"), NumConeSampleDirections);
+		OutEnvironment.SetDefine(TEXT("NUM_CONE_DIRECTIONS"), NumConeSampleDirections);
 		OutEnvironment.SetDefine(TEXT("DOWNSAMPLE_FACTOR"), GAODownsampleFactor);
 
 		// To reduce shader compile time of compute shaders with shared memory, doesn't have an impact on generated code with current compiler (June 2010 DX SDK)
@@ -1369,7 +1369,7 @@ public:
 	static void ModifyCompilationEnvironment(EShaderPlatform Platform, FShaderCompilerEnvironment& OutEnvironment)
 	{
 		FGlobalShader::ModifyCompilationEnvironment(Platform,OutEnvironment);
-		OutEnvironment.SetDefine(TEXT("NUM_SAMPLES"), NumConeSampleDirections);
+		OutEnvironment.SetDefine(TEXT("NUM_CONE_DIRECTIONS"), NumConeSampleDirections);
 		OutEnvironment.SetDefine(TEXT("DOWNSAMPLE_FACTOR"), GAODownsampleFactor);
 		OutEnvironment.SetDefine(TEXT("MAX_OBJECTS_PER_TILE"), GMaxNumObjectsPerTile);
 
diff --git a/Engine/Source/Runtime/Renderer/Private/DistanceFieldSurfaceCacheLighting.h b/Engine/Source/Runtime/Renderer/Private/DistanceFieldSurfaceCacheLighting.h
--- a/Engine/Source/Runtime/Renderer/Private/DistanceFieldSurfaceCacheLighting.h
+++ b/Engine/Source/Runtime/Renderer/Private/DistanceFieldSurfaceCacheLighting.h
@@ -7,6 +7,7 @@
 #pragma once
 
 const static int32 GAOMaxSupportedLevel = 4;
+//@todo - derive from worst case
 const static uint32 GMaxIrradianceCacheSamples = 100000;
 static const int32 GMaxNumObjectsPerTile = 512;
 
