fixed "ApplyDBufferData adding .5 to spec brightness?"
was affecting specular brightness if DBuffer is enabled, this is no longer the case

diff --git a/Engine/Shaders/BasePassPixelShader.usf b/Engine/Shaders/BasePassPixelShader.usf
--- a/Engine/Shaders/BasePassPixelShader.usf
+++ b/Engine/Shaders/BasePassPixelShader.usf
@@ -619,9 +619,9 @@ void Main(
 		FDBufferData DBufferData = GetDBufferData(ScreenUV);
 
 		// the material can disable the DBuffer effects for beter performance or control
-		if((MATERIALDECALRESPONSEMASK & 0x1) == 0) { DBufferData.Color = 0; DBufferData.ColorOpacity = 1; }
-		if((MATERIALDECALRESPONSEMASK & 0x2) == 0) { DBufferData.WorldNormal = 0; DBufferData.NormalOpacity = 1; }
-		if((MATERIALDECALRESPONSEMASK & 0x4) == 0) { DBufferData.Roughness = 0; DBufferData.RoughnessOpacity = 1; }
+		if((MATERIALDECALRESPONSEMASK & 0x1) == 0) { DBufferData.PreMulColor = 0; DBufferData.ColorOpacity = 1; }
+		if((MATERIALDECALRESPONSEMASK & 0x2) == 0) { DBufferData.PreMulWorldNormal = 0; DBufferData.NormalOpacity = 1; }
+		if((MATERIALDECALRESPONSEMASK & 0x4) == 0) { DBufferData.PreMulRoughness = 0; DBufferData.RoughnessOpacity = 1; }
 
 #if DIFFUSE_SPEC_INPUTS
 		ApplyDBufferData(DBufferData, MaterialParameters.WorldNormal, SubsurfaceColor, Roughness, DiffuseColor, SpecularColor);
diff --git a/Engine/Shaders/DeferredShadingCommon.usf b/Engine/Shaders/DeferredShadingCommon.usf
--- a/Engine/Shaders/DeferredShadingCommon.usf
+++ b/Engine/Shaders/DeferredShadingCommon.usf
@@ -310,17 +310,17 @@ struct FGBufferData
 struct FDBufferData
 {
 	// 0..1, premultiplied with ColorOpacity
-	float3 Color;
+	float3 PreMulColor;
 	// 0:opaque ..1:see through
 	float ColorOpacity;
 
 	// -1..1, premultiplied with NormalOpacity
-	float3 WorldNormal;
+	float3 PreMulWorldNormal;
 	// 0:opaque ..1:see through
 	float NormalOpacity;
 
 	// 0..1, premultiplied with RoughnessOpacity
-	float Roughness;
+	float PreMulRoughness;
 	// 0:opaque ..1:see through
 	float RoughnessOpacity;
 };
@@ -469,15 +469,15 @@ FDBufferData DecodeDBufferData(
 	FDBufferData ret;
 
 	// UNORM 4 channel
-	ret.Color = DBufferA.rgb;
+	ret.PreMulColor = DBufferA.rgb;
 	ret.ColorOpacity = DBufferA.a;
 
 	// UNORM 4 channel, 128/255 represents 0
-	ret.WorldNormal = DBufferB.rgb * 2 - (256.0 / 255.0);
+	ret.PreMulWorldNormal = DBufferB.rgb * 2 - (256.0 / 255.0);
 	ret.NormalOpacity = DBufferB.a;
 
 	// UNORM 2 channel
-	ret.Roughness = DBufferC.r;
+	ret.PreMulRoughness = DBufferC.r;
 	ret.RoughnessOpacity = DBufferC.g;
 
 	return ret;
@@ -537,18 +537,18 @@ void ApplyDBufferData(FDBufferData DBufferData, inout float3 WorldNormal, inout
 #endif
 					  )
 {
-	WorldNormal = WorldNormal * DBufferData.NormalOpacity + DBufferData.WorldNormal;
-	Roughness = Roughness * DBufferData.RoughnessOpacity + DBufferData.Roughness;
+	WorldNormal = WorldNormal * DBufferData.NormalOpacity + DBufferData.PreMulWorldNormal;
+	Roughness = Roughness * DBufferData.RoughnessOpacity + DBufferData.PreMulRoughness;
 	
 	SubsurfaceColor *= DBufferData.ColorOpacity;
 
 #if DIFFUSE_SPEC_INPUTS
-	DiffuseColor = DiffuseColor * DBufferData.ColorOpacity + DBufferData.Color;
-	SpecularColor = SpecularColor * DBufferData.ColorOpacity + 0.04f;		// most non metal materials have a specular of 4%
+	DiffuseColor = DiffuseColor * DBufferData.ColorOpacity + DBufferData.PreMulColor;
+	SpecularColor = lerp(0.04f, SpecularColor, DBufferData.ColorOpacity);		// most non metal materials have a specular of 4%
 #else
-	BaseColor = BaseColor * DBufferData.ColorOpacity + DBufferData.Color;
+	BaseColor = BaseColor * DBufferData.ColorOpacity + DBufferData.PreMulColor;
 	Metallic = Metallic * DBufferData.ColorOpacity + 0;			// decals are always no metallic
-	Specular = Specular * DBufferData.ColorOpacity + 0.5f;		// most non metal materials have a specular of 4% which is 0.5 in this scale
+	Specular = lerp(0.5f, Specular, DBufferData.ColorOpacity);		// most non metal materials have a specular of 4% which is 0.5 in this scale
 #endif
 }
 
