Distance Field AO now works on flat ground planes whose X and Y scale was larger than Z

diff --git a/Engine/Shaders/DistanceFieldSurfaceCacheLightingCompute.usf b/Engine/Shaders/DistanceFieldSurfaceCacheLightingCompute.usf
--- a/Engine/Shaders/DistanceFieldSurfaceCacheLightingCompute.usf
+++ b/Engine/Shaders/DistanceFieldSurfaceCacheLightingCompute.usf
@@ -531,6 +531,8 @@ void FinalGatherCS(
 		RotatedConeDirection = DebugConeDirection;
 #endif
 
+		float MinVisibility = 1;
+		float MinWorldDistanceToOccluder = AOMaxDistance;
 		float MaxWorldStepOffset = GetStepOffset(NUM_CONE_STEPS);
 #endif
 		LOOP 
@@ -609,22 +611,26 @@ void FinalGatherCS(
 						// Derive visibility from 1d intersection
 						float Visibility = saturate(DistanceToOccluder / ShadingSphereRadius);
 
-						// Don't allow small objects to fully occlude a cone step
-						Visibility = max(Visibility, 1 - saturate(ObjectOccluderRadius / SphereRadius));
-						float OccluderDistanceFraction = (WorldStepOffset + DistanceToOccluder) / AOMaxDistance;
-						// Fade out occlusion based on distance to occluder
-						Visibility = max(Visibility, saturate((OccluderDistanceFraction - 0) / 4));
-						Visibility = max(Visibility, saturate((OccluderDistanceFraction - .7f) / .4f));
+						BRANCH
+						if (Visibility < MinVisibility)
+						{
+							// Don't allow small objects to fully occlude a cone step
+							Visibility = max(Visibility, 1 - saturate(ObjectOccluderRadius / SphereRadius));
+							float OccluderDistanceFraction = (WorldStepOffset + DistanceToOccluder) / AOMaxDistance;
+							// Fade out occlusion based on distance to occluder
+							Visibility = max(Visibility, saturate((OccluderDistanceFraction - 0) / 4));
+							Visibility = max(Visibility, saturate((OccluderDistanceFraction - .7f) / .4f));
 
-						InterlockedMin(SharedConeVisibility[ConeIndex], asuint(Visibility));
+							MinVisibility = min(MinVisibility, Visibility);
 
-						if (DistanceToOccluder < .9f * SphereRadius)
-						{
-							// Assuming occluder is straight forward along the cone
-							float WorldDistanceToOccluder = WorldStepOffset + DistanceToOccluder;
-							InterlockedMin(SharedMinOcclusionDistance, asuint(WorldDistanceToOccluder));
+							if (DistanceToOccluder < .9f * SphereRadius)
+							{
+								// Assuming occluder is straight forward along the cone
+								float WorldDistanceToOccluder = WorldStepOffset + DistanceToOccluder;
+								MinWorldDistanceToOccluder = min(MinWorldDistanceToOccluder, WorldDistanceToOccluder);
+							}
 						}
-
+						
 						float MinStepSize = .6f * (GetStepOffset(StepIndex + 1) - GetStepOffset(StepIndex));
 						WorldStepOffset += max(DistanceToOccluder, MinStepSize);
 					}
@@ -697,6 +703,11 @@ void FinalGatherCS(
 				}
 			}
 		}
+
+#if USE_CONE_TRACING
+		InterlockedMin(SharedConeVisibility[ConeIndex], asuint(MinVisibility));
+		InterlockedMin(SharedMinOcclusionDistance, asuint(MinWorldDistanceToOccluder));
+#endif
 	}
 
 	GroupMemoryBarrierWithGroupSync();
diff --git a/Engine/Source/Developer/MeshUtilities/Private/MeshUtilities.cpp b/Engine/Source/Developer/MeshUtilities/Private/MeshUtilities.cpp
--- a/Engine/Source/Developer/MeshUtilities/Private/MeshUtilities.cpp
+++ b/Engine/Source/Developer/MeshUtilities/Private/MeshUtilities.cpp
@@ -354,11 +354,28 @@ void FMeshUtilities::GenerateSignedDistanceFieldVolumeData(
 		FIndexArrayView Indices = LODModel.IndexBuffer.GetArrayView();
 		TArray<FkDOPBuildCollisionTriangle<uint32> > BuildTriangles;
 
+		FVector BoundsSize = Bounds.GetBox().GetExtent() * 2;
+		float MaxDimension = FMath::Max(FMath::Max(BoundsSize.X, BoundsSize.Y), BoundsSize.Z);
+
+		// Consider the mesh a plane if it is very flat
+		const bool bMeshWasPlane = BoundsSize.Z * 100 < MaxDimension 
+			// And it lies mostly on the origin
+			&& Bounds.Origin.Z - Bounds.BoxExtent.Z < KINDA_SMALL_NUMBER 
+			&& Bounds.Origin.Z + Bounds.BoxExtent.Z > -KINDA_SMALL_NUMBER;
+
 		for (int32 i = 0; i < Indices.Num(); i += 3)
 		{
-			const FVector V0 = PositionVertexBuffer.VertexPosition(Indices[i + 0]);
-			const FVector V1 = PositionVertexBuffer.VertexPosition(Indices[i + 1]);
-			const FVector V2 = PositionVertexBuffer.VertexPosition(Indices[i + 2]);
+			FVector V0 = PositionVertexBuffer.VertexPosition(Indices[i + 0]);
+			FVector V1 = PositionVertexBuffer.VertexPosition(Indices[i + 1]);
+			FVector V2 = PositionVertexBuffer.VertexPosition(Indices[i + 2]);
+
+			if (bMeshWasPlane)
+			{
+				// Flatten out the mesh into an actual plane, this will allow us to manipulate the component's Z scale at runtime without artifacts
+				V0.Z = 0;
+				V1.Z = 0;
+				V2.Z = 0;
+			}
 
 			const FVector LocalNormal = ((V1 - V2) ^ (V0 - V2)).SafeNormal();
 
@@ -468,6 +485,7 @@ void FMeshUtilities::GenerateSignedDistanceFieldVolumeData(
 
 			OutData.bMeshWasClosed = !bNegativeAtBorder;
 			OutData.bBuiltAsIfTwoSided = bGenerateAsIfTwoSided;
+			OutData.bMeshWasPlane = bMeshWasPlane;
 
 			UE_LOG(LogMeshUtilities,Log,TEXT("Finished distance field build in %.1fs - %ux%ux%u distance field, %u triangles"), 
 				(float)(FPlatformTime::Seconds() - StartTime),
diff --git a/Engine/Source/Runtime/Engine/Private/StaticMeshRender.cpp b/Engine/Source/Runtime/Engine/Private/StaticMeshRender.cpp
--- a/Engine/Source/Runtime/Engine/Private/StaticMeshRender.cpp
+++ b/Engine/Source/Runtime/Engine/Private/StaticMeshRender.cpp
@@ -1208,7 +1208,7 @@ void FStaticMeshSceneProxy::GetLightRelevance(const FLightSceneProxy* LightScene
 	}
 }
 
-void FStaticMeshSceneProxy::GetDistancefieldAtlasData(FBox& LocalVolumeBounds, FIntVector& OutBlockMin, FIntVector& OutBlockSize, bool& bOutBuiltAsIfTwoSided) const
+void FStaticMeshSceneProxy::GetDistancefieldAtlasData(FBox& LocalVolumeBounds, FIntVector& OutBlockMin, FIntVector& OutBlockSize, bool& bOutBuiltAsIfTwoSided, bool& bMeshWasPlane) const
 {
 	if (DistanceFieldData)
 	{
@@ -1216,6 +1216,7 @@ void FStaticMeshSceneProxy::GetDistancefieldAtlasData(FBox& LocalVolumeBounds, F
 		OutBlockMin = DistanceFieldData->VolumeTexture.GetAllocationMin();
 		OutBlockSize = DistanceFieldData->VolumeTexture.GetAllocationSize();
 		bOutBuiltAsIfTwoSided = DistanceFieldData->bBuiltAsIfTwoSided;
+		bMeshWasPlane = DistanceFieldData->bMeshWasPlane;
 	}
 	else
 	{
@@ -1223,6 +1224,7 @@ void FStaticMeshSceneProxy::GetDistancefieldAtlasData(FBox& LocalVolumeBounds, F
 		OutBlockMin = FIntVector(-1, -1, -1);
 		OutBlockSize = FIntVector(0, 0, 0);
 		bOutBuiltAsIfTwoSided = false;
+		bMeshWasPlane = false;
 	}
 }
 
diff --git a/Engine/Source/Runtime/Engine/Public/DistanceFieldAtlas.h b/Engine/Source/Runtime/Engine/Public/DistanceFieldAtlas.h
--- a/Engine/Source/Runtime/Engine/Public/DistanceFieldAtlas.h
+++ b/Engine/Source/Runtime/Engine/Public/DistanceFieldAtlas.h
@@ -10,7 +10,7 @@
 #include "GCObject.h"
 
 // DDC key for distance field data, must be changed when modifying the generation code or data format
-#define DISTANCEFIELD_DERIVEDDATA_VER TEXT("7768798763B445A9533C94441EA899D")
+#define DISTANCEFIELD_DERIVEDDATA_VER TEXT("7768798763B445A9543C94442EA899D")
 
 /** Represents a distance field volume texture for a single UStaticMesh. */
 class FDistanceFieldVolumeTexture
@@ -111,6 +111,9 @@ public:
 	/** Whether the distance field was built assuming that every triangle is a frontface. */
 	bool bBuiltAsIfTwoSided;
 
+	/** Whether the mesh was a plane with very little extent in Z. */
+	bool bMeshWasPlane;
+
 	FDistanceFieldVolumeTexture VolumeTexture;
 
 	FDistanceFieldVolumeData() :
@@ -118,6 +121,7 @@ public:
 		LocalBoundingBox(0),
 		bMeshWasClosed(true),
 		bBuiltAsIfTwoSided(false),
+		bMeshWasPlane(false),
 		VolumeTexture(*this)
 	{}
 
@@ -141,7 +145,7 @@ public:
 	friend FArchive& operator<<(FArchive& Ar,FDistanceFieldVolumeData& Data)
 	{
 		// Note: this is derived data, no need for versioning (bump the DDC guid)
-		Ar << Data.DistanceFieldVolume << Data.Size << Data.LocalBoundingBox << Data.bMeshWasClosed << Data.bBuiltAsIfTwoSided;
+		Ar << Data.DistanceFieldVolume << Data.Size << Data.LocalBoundingBox << Data.bMeshWasClosed << Data.bBuiltAsIfTwoSided << Data.bMeshWasPlane;
 		return Ar;
 	}
 };
diff --git a/Engine/Source/Runtime/Engine/Public/PrimitiveSceneProxy.h b/Engine/Source/Runtime/Engine/Public/PrimitiveSceneProxy.h
--- a/Engine/Source/Runtime/Engine/Public/PrimitiveSceneProxy.h
+++ b/Engine/Source/Runtime/Engine/Public/PrimitiveSceneProxy.h
@@ -224,12 +224,13 @@ public:
 		bShadowMapped = false;
 	}
 
-	virtual void GetDistancefieldAtlasData(FBox& LocalVolumeBounds, FIntVector& OutBlockMin, FIntVector& OutBlockSize, bool& bOutBuiltAsIfTwoSided) const 
+	virtual void GetDistancefieldAtlasData(FBox& LocalVolumeBounds, FIntVector& OutBlockMin, FIntVector& OutBlockSize, bool& bOutBuiltAsIfTwoSided, bool& bMeshWasPlane) const 
 	{
 		LocalVolumeBounds = FBox(0);
 		OutBlockMin = FIntVector(-1, -1, -1);
 		OutBlockSize = FIntVector(0, 0, 0);
 		bOutBuiltAsIfTwoSided = false;
+		bMeshWasPlane = false;
 	}
 
 	/**
diff --git a/Engine/Source/Runtime/Engine/Public/StaticMeshResources.h b/Engine/Source/Runtime/Engine/Public/StaticMeshResources.h
--- a/Engine/Source/Runtime/Engine/Public/StaticMeshResources.h
+++ b/Engine/Source/Runtime/Engine/Public/StaticMeshResources.h
@@ -773,7 +773,7 @@ public:
 	virtual FPrimitiveViewRelevance GetViewRelevance(const FSceneView* View) override;
 	virtual bool CanBeOccluded() const override;
 	virtual void GetLightRelevance(const FLightSceneProxy* LightSceneProxy, bool& bDynamic, bool& bRelevant, bool& bLightMapped, bool& bShadowMapped) const override;
-	virtual void GetDistancefieldAtlasData(FBox& LocalVolumeBounds, FIntVector& OutBlockMin, FIntVector& OutBlockSize, bool& bOutBuiltAsIfTwoSided) const override;
+	virtual void GetDistancefieldAtlasData(FBox& LocalVolumeBounds, FIntVector& OutBlockMin, FIntVector& OutBlockSize, bool& bOutBuiltAsIfTwoSided, bool& bMeshWasPlane) const override;
 	virtual bool HasDistanceFieldRepresentation() const override;
 	virtual uint32 GetMemoryFootprint( void ) const override { return( sizeof( *this ) + GetAllocatedSize() ); }
 	uint32 GetAllocatedSize( void ) const { return( FPrimitiveSceneProxy::GetAllocatedSize() + LODs.GetAllocatedSize() ); }
diff --git a/Engine/Source/Runtime/Renderer/Private/DistanceFieldShadowing.cpp b/Engine/Source/Runtime/Renderer/Private/DistanceFieldShadowing.cpp
--- a/Engine/Source/Runtime/Renderer/Private/DistanceFieldShadowing.cpp
+++ b/Engine/Source/Runtime/Renderer/Private/DistanceFieldShadowing.cpp
@@ -192,7 +192,8 @@ int32 FProjectedShadowInfo::UpdateShadowCastingObjectBuffers() const
 			FIntVector BlockMin;
 			FIntVector BlockSize;
 			bool bBuiltAsIfTwoSided;
-			PrimitiveSceneInfo->Proxy->GetDistancefieldAtlasData(LocalVolumeBounds, BlockMin, BlockSize, bBuiltAsIfTwoSided);
+			bool bMeshWasPlane;
+			PrimitiveSceneInfo->Proxy->GetDistancefieldAtlasData(LocalVolumeBounds, BlockMin, BlockSize, bBuiltAsIfTwoSided, bMeshWasPlane);
 
 			if (BlockMin.X >= 0 && BlockMin.Y >= 0 && BlockMin.Z >= 0)
 			{
diff --git a/Engine/Source/Runtime/Renderer/Private/DistanceFieldSurfaceCacheLighting.cpp b/Engine/Source/Runtime/Renderer/Private/DistanceFieldSurfaceCacheLighting.cpp
--- a/Engine/Source/Runtime/Renderer/Private/DistanceFieldSurfaceCacheLighting.cpp
+++ b/Engine/Source/Runtime/Renderer/Private/DistanceFieldSurfaceCacheLighting.cpp
@@ -2664,11 +2664,23 @@ void UpdateVisibleObjectBuffers(const FScene* Scene, const FDistanceFieldAOParam
 				FIntVector BlockMin;
 				FIntVector BlockSize;
 				bool bBuiltAsIfTwoSided;
-				PrimitiveSceneInfo->Proxy->GetDistancefieldAtlasData(LocalVolumeBounds, BlockMin, BlockSize, bBuiltAsIfTwoSided);
+				bool bMeshWasPlane;
+				PrimitiveSceneInfo->Proxy->GetDistancefieldAtlasData(LocalVolumeBounds, BlockMin, BlockSize, bBuiltAsIfTwoSided, bMeshWasPlane);
 
 				if (BlockMin.X >= 0 && BlockMin.Y >= 0 && BlockMin.Z >= 0)
 				{
-					const FMatrix LocalToWorld = PrimitiveSceneInfo->Proxy->GetLocalToWorld();
+					FMatrix LocalToWorld = PrimitiveSceneInfo->Proxy->GetLocalToWorld();
+
+					if (bMeshWasPlane)
+					{
+						FVector LocalScales = LocalToWorld.GetScaleVector();
+						FVector AbsLocalScales(FMath::Abs(LocalScales.X), FMath::Abs(LocalScales.Y), FMath::Abs(LocalScales.Z));
+						float MidScale = FMath::Min(AbsLocalScales.X, AbsLocalScales.Y);
+						float ScaleAdjust = FMath::Sign(LocalScales.Z) * MidScale / AbsLocalScales.Z;
+						// The mesh was determined to be a plane flat in Z during the build process, so we can change the Z scale
+						// Helps in cases with modular ground pieces with scales of (10, 10, 1) and some triangles just above Z=0
+						LocalToWorld.SetAxis(2, LocalToWorld.GetScaledAxis(EAxis::Z) * ScaleAdjust);
+					}
 
 					const FMatrix VolumeToWorld = FScaleMatrix(LocalVolumeBounds.GetExtent()) 
 						* FTranslationMatrix(LocalVolumeBounds.GetCenter())
