Dropped another 20% off the fast temporal AA. This introduces a YCoCg transform which like the HdrWeight can be factored into a prior pass.

diff --git a/Engine/Shaders/PostProcessTemporalAA.usf b/Engine/Shaders/PostProcessTemporalAA.usf
--- a/Engine/Shaders/PostProcessTemporalAA.usf
+++ b/Engine/Shaders/PostProcessTemporalAA.usf
@@ -15,13 +15,12 @@
 float4 CameraMotion[5];
 float SampleWeights[9];
 float LowpassWeights[9];
+float PlusWeights[5];
 float2 RandomOffset;
 
 // TODO: This can be removed.
 float ClampFeebackMix;
 
-
-
 // Faster but less accurate luma computation. 
 // Luma includes a scaling by 4.
 float Luma4(float3 Color)
@@ -35,6 +34,11 @@ float HdrWeight4(float3 Color, float Exposure)
 	return rcp(Luma4(Color) * Exposure + 4.0);
 }
 
+float HdrWeightY(float Color, float Exposure) 
+{
+	return rcp(Color * Exposure + 1.0);
+}
+
 float HdrWeightG(float3 Color, float Exposure) 
 {
 	return rcp(Color.g * Exposure + 1.0);
@@ -63,6 +67,11 @@ float HdrWeightInvG(float3 Color, float Exposure)
 	return rcp(Color.g * (-Exposure) + 1.0);
 }
 
+float HdrWeightInvY(float Color, float Exposure) 
+{
+	return rcp(Color * (-Exposure) + 1.0);
+}
+
 float HdrWeightInv4_(float Color, float Exposure) 
 {
 	return 4.0 * rcp(Color * (-Exposure) + 1.0);
@@ -182,7 +191,6 @@ void SSRTemporalAAPS( float4 UVAndScreenPos : TEXCOORD0, float3 InExposureScaleV
 {
 #if ENABLE_TEMPORAL_AA
 	float InExposureScale = InExposureScaleVignette.x;
-	#define AA_NEW 1
 	#define AA_ALPHA 0
 	#define AA_CROSS 0
 	#define AA_DYNAMIC 0
@@ -199,7 +207,6 @@ void DOFTemporalAAPS( float4 UVAndScreenPos : TEXCOORD0, float3 InExposureScaleV
 {
 #if ENABLE_TEMPORAL_AA
 	float InExposureScale = InExposureScaleVignette.x;
-	#define AA_NEW 1
 	#define AA_ALPHA 0
 	#define AA_CROSS 4
 	#define AA_DYNAMIC 0
@@ -216,7 +223,6 @@ void LightShaftTemporalAAPS( float4 UVAndScreenPos : TEXCOORD0, float3 InExposur
 {
 #if ENABLE_TEMPORAL_AA
 	float InExposureScale = InExposureScaleVignette.x;
-	#define AA_NEW 1
 	#define AA_ALPHA 0
 	#define AA_CROSS 0
 	#define AA_DYNAMIC 0
@@ -232,22 +238,14 @@ void MainTemporalAAPS( float4 UVAndScreenPos : TEXCOORD0, float3 InExposureScale
 {
 #if ENABLE_TEMPORAL_AA
 	float InExposureScale = InExposureScaleVignette.x;
-	#if 1
-		#define AA_NEW 1
-		#define AA_FILTERED 1
-		#define AA_BORDER 1
-		#define AA_CROSS 2
-		#define AA_GREEN_AS_LUMA 0
-		#define AA_AABB 1
-		#define AA_LOWPASS 1
-		#define AA_ONE_DYNAMIC_SAMPLE 1
-		#define AA_DEBUG 0
-		#define AA_VELOCITY_WEIGHTING 1
-	#else
-		#define AA_NEW 0
-		#define AA_BORDER 1
-		#define AA_GRAIN 1
-	#endif
+	#define AA_FILTERED 1
+	#define AA_BORDER 1
+	#define AA_CROSS 2
+	#define AA_GREEN_AS_LUMA 0
+	#define AA_AABB 1
+	#define AA_LOWPASS 1
+	#define AA_DEBUG 0
+	#define AA_VELOCITY_WEIGHTING 1
 	#include "PostProcessTemporalCommon.usf"
 #else 
 	OutColor = PostprocessInput0.SampleLevel(PostprocessInput0Sampler, UVAndScreenPos.xy, 0);
@@ -258,26 +256,15 @@ void MainFastTemporalAAPS( float4 UVAndScreenPos : TEXCOORD0, float3 InExposureS
 {
 #if ENABLE_TEMPORAL_AA
 	float InExposureScale = InExposureScaleVignette.x;
-	#if 1
-		#define AA_NEW 1
-		#define AA_FILTERED 1
-		#define AA_BORDER 1
-		#define AA_CROSS 2
-		#define AA_GREEN_AS_LUMA 1
-		#define AA_AABB 1
-		#define AA_LOWPASS 0
-		#define AA_ONE_DYNAMIC_SAMPLE 1
-		#define AA_DEBUG 0
-		#define AA_VELOCITY_WEIGHTING 0
-	#else
-		#define AA_NEW 0
-		#define AA_BORDER 1
-		#define AA_GRAIN 1
-		#define AA_HDR 0
-		#define AA_HDR_HISTORY 0
-		#define AA_LOWPASS 0
-		#define AA_ONE_DYNAMIC_SAMPLE 1
-	#endif
+	#define AA_FILTERED 1
+	#define AA_BORDER 1
+	#define AA_CROSS 2
+	#define AA_GREEN_AS_LUMA 1
+	#define AA_AABB 1
+	#define AA_LOWPASS 0
+	#define AA_DEBUG 0
+	#define AA_VELOCITY_WEIGHTING 0
+	#define AA_YCOCG 1
 	#include "PostProcessTemporalCommon.usf"
 #else 
 	OutColor = PostprocessInput0.SampleLevel(PostprocessInput0Sampler, UVAndScreenPos.xy, 0);
diff --git a/Engine/Shaders/PostProcessTemporalCommon.usf b/Engine/Shaders/PostProcessTemporalCommon.usf
--- a/Engine/Shaders/PostProcessTemporalCommon.usf
+++ b/Engine/Shaders/PostProcessTemporalCommon.usf
@@ -105,38 +105,29 @@ which can be configured via defines for various temporal AA passes.
 	#define AA_FORCE_ALPHA_CLAMP 0
 #endif
 
-// Optimized path to avoid a second dynamic motion sample.
-#ifndef AA_ONE_DYNAMIC_SAMPLE
-	#define AA_ONE_DYNAMIC_SAMPLE 0
-#endif
-
-// Use green as luma.
-#ifndef AA_GREEN_AS_LUMA
-	#define AA_GREEN_AS_LUMA 0
-#endif
-
-// Try the new AA algorithm.
-#ifndef AA_NEW 
-	#define AA_NEW 0
-#endif
-
 // Use velocity weighting.
 #ifndef AA_VELOCITY_WEIGHTING
 	#define AA_VELOCITY_WEIGHTING 0
 #endif
 
+// Use YCoCg path.
+#ifndef AA_YCOCG
+	#define AA_YCOCG 0
+#endif
+
+// Use green as luma.
+#ifndef AA_GREEN_AS_LUMA
+	#define AA_GREEN_AS_LUMA AA_YCOCG
+#endif
 
 
 
-#if AA_NEW
 
 	// FIND MOTION OF PIXEL AND NEAREST IN NEIGHBORHOOD
     // ------------------------------------------------
-	float3 PosP; // Position of this pixel.
-	PosP.xy = UVAndScreenPos.zw * float2(0.5, -0.5) + 0.5; // View position [0 to 1] flipped in Y.
-	PosP.z = SceneDepthTexture.SampleLevel(SceneDepthTextureSampler, UVAndScreenPos.xy, 0).r;
-	float3 PosN; // Position of closest pixel in neighborhood.
-	PosN = PosP;
+	float3 PosN; // Position of this pixel, possibly later nearest pixel in neighborhood.
+	PosN.xy = UVAndScreenPos.zw * float2(0.5, -0.5) + 0.5; // View position [0 to 1] flipped in Y.
+	PosN.z = SceneDepthTexture.SampleLevel(SceneDepthTextureSampler, UVAndScreenPos.xy, 0).r;
 	// Screen position of minimum depth.
 	float2 VelocityOffset = float2(0.0, 0.0);
 	#if AA_CROSS
@@ -171,7 +162,7 @@ which can be configured via defines for various temporal AA passes.
 			DepthOffset.x = DepthOffsetXx; 
 		}
 		float DepthsXYZW = max(DepthsXY, DepthsZW);
-		if(DepthsXYZW > PosP.z) 
+		if(DepthsXYZW > PosN.z) 
 		{
 			// This is offset for reading from velocity texture.
 			// This supports half or fractional resolution velocity textures.
@@ -182,50 +173,30 @@ which can be configured via defines for various temporal AA passes.
 			PosN.z = DepthsXYZW;
 		}
 	#endif
-	// Camera motion for pixel (in ScreenPos space).
-	float ScaleM = 1.0 / (dot(PosP, CameraMotion[0].xyz) + CameraMotion[0].w);
-	float2 BackP;
-	BackP.x = -2.0 * ((PosP.x * ((CameraMotion[1].x * PosP.y) + (CameraMotion[1].y * PosP.z) + CameraMotion[1].z)) + (CameraMotion[1].w * PosP.y) + (CameraMotion[2].x * PosP.x * PosP.x) + (CameraMotion[2].y * PosP.z) + CameraMotion[2].z) * ScaleM;
-	BackP.y =  2.0 * ((PosP.y * ((CameraMotion[3].x * PosP.x) + (CameraMotion[3].y * PosP.z) + CameraMotion[3].z)) + (CameraMotion[3].w * PosP.x) + (CameraMotion[4].x * PosP.y * PosP.y) + (CameraMotion[4].y * PosP.z) + CameraMotion[4].z) * ScaleM;
-	// Convert BackP back into pixels.
-	float2 BackTemp = BackP * ViewportSize.xy;
-	#if AA_CROSS
-		// Camera motion for neighborhood.
-		ScaleM = 1.0 / (dot(PosN, CameraMotion[0].xyz) + CameraMotion[0].w);
-		float2 BackN;
-		BackN.x = -2.0 * ((PosN.x * ((CameraMotion[1].x * PosN.y) + (CameraMotion[1].y * PosN.z) + CameraMotion[1].z)) + (CameraMotion[1].w * PosN.y) + (CameraMotion[2].x * PosN.x * PosN.x) + (CameraMotion[2].y * PosN.z) + CameraMotion[2].z) * ScaleM;
-		BackN.y =  2.0 * ((PosN.y * ((CameraMotion[3].x * PosN.x) + (CameraMotion[3].y * PosN.z) + CameraMotion[3].z)) + (CameraMotion[3].w * PosN.x) + (CameraMotion[4].x * PosN.y * PosN.y) + (CameraMotion[4].y * PosN.z) + CameraMotion[4].z) * ScaleM;
-		BackTemp = BackN * ViewportSize.xy;
-	#endif
+	// Camera motion for pixel or nearest pixel (in ScreenPos space).
+	float ScaleM = 1.0 / (dot(PosN, CameraMotion[0].xyz) + CameraMotion[0].w);
+	float2 BackN;
+	BackN.x = -2.0 * ((PosN.x * ((CameraMotion[1].x * PosN.y) + (CameraMotion[1].y * PosN.z) + CameraMotion[1].z)) + (CameraMotion[1].w * PosN.y) + (CameraMotion[2].x * PosN.x * PosN.x) + (CameraMotion[2].y * PosN.z) + CameraMotion[2].z) * ScaleM;
+	BackN.y =  2.0 * ((PosN.y * ((CameraMotion[3].x * PosN.x) + (CameraMotion[3].y * PosN.z) + CameraMotion[3].z)) + (CameraMotion[3].w * PosN.x) + (CameraMotion[4].x * PosN.y * PosN.y) + (CameraMotion[4].y * PosN.z) + CameraMotion[4].z) * ScaleM;
+	float2 BackTemp = BackN * ViewportSize.xy;
 	#if AA_DYNAMIC
-		// Dynamic motion of pixel.
-		#if (AA_CROSS == 0) || (AA_ONE_DYNAMIC_SAMPLE == 0)
-			float2 VelocityP = PostprocessInput3.SampleLevel(PostprocessInput3Sampler, UVAndScreenPos.xy, 0).xy;
-			bool DynamicP = VelocityP.x > 0.0;
-			if(DynamicP)
-			{
-				BackP = DecodeVelocityFromTexture(VelocityP);
-			}
-		#endif
-		BackTemp = BackP * ViewportSize.xy;
+		float2 VelocityN;
 		#if AA_CROSS
-			// Dynamic motion of nearest in neighborhood.
-			float2 VelocityN = PostprocessInput3.SampleLevel(PostprocessInput3Sampler, UVAndScreenPos.xy + VelocityOffset, 0).xy;
-			bool DynamicN = VelocityN.x > 0.0;
-			if(DynamicN)
-			{
-				BackN = DecodeVelocityFromTexture(VelocityN);
-			}
-			BackTemp = BackN * ViewportSize.xy;
+			VelocityN = PostprocessInput3.SampleLevel(PostprocessInput3Sampler, UVAndScreenPos.xy + VelocityOffset, 0).xy;
+		#else
+			VelocityN = PostprocessInput3.SampleLevel(PostprocessInput3Sampler, UVAndScreenPos.xy, 0).xy;
 		#endif
+		bool DynamicN = VelocityN.x > 0.0;
+		if(DynamicN)
+		{
+			BackN = DecodeVelocityFromTexture(VelocityN);
+		}
+		BackTemp = BackN * ViewportSize.xy;
 	#endif
 	// Save the amount of pixel offset of just camera motion, used later as the amount of blur introduced by history.
 	float HistoryBlurAmp = 2.0;
 	float HistoryBlur = saturate(abs(BackTemp.x) * HistoryBlurAmp + abs(BackTemp.y) * HistoryBlurAmp);
 	float Velocity = sqrt(dot(BackTemp, BackTemp));
-	#if (AA_CROSS == 0)
-		float2 BackN = BackP;
-	#endif
 	// Easier to do off screen check before conversion.
 	// BackN is in units of 2pixels/viewportWidthInPixels
 	// This converts back projection vector to [-1 to 1] offset in viewport.
@@ -246,94 +217,148 @@ which can be configured via defines for various temporal AA passes.
 	// 012
 	// 345
 	// 678
-	float4 Neighbor0 = PostprocessInput0.SampleLevel(PostprocessInput0Sampler, UVAndScreenPos.xy, 0, int2(-1, -1));
-	float4 Neighbor1 = PostprocessInput0.SampleLevel(PostprocessInput0Sampler, UVAndScreenPos.xy, 0, int2( 0, -1));
-	float4 Neighbor2 = PostprocessInput0.SampleLevel(PostprocessInput0Sampler, UVAndScreenPos.xy, 0, int2( 1, -1));
-	float4 Neighbor3 = PostprocessInput0.SampleLevel(PostprocessInput0Sampler, UVAndScreenPos.xy, 0, int2(-1,  0));
-	float4 Neighbor4 = PostprocessInput0.SampleLevel(PostprocessInput0Sampler, UVAndScreenPos.xy, 0);
-	float4 Neighbor5 = PostprocessInput0.SampleLevel(PostprocessInput0Sampler, UVAndScreenPos.xy, 0, int2( 1,  0));
-	float4 Neighbor6 = PostprocessInput0.SampleLevel(PostprocessInput0Sampler, UVAndScreenPos.xy, 0, int2(-1,  1));
-	float4 Neighbor7 = PostprocessInput0.SampleLevel(PostprocessInput0Sampler, UVAndScreenPos.xy, 0, int2( 0,  1));
-	float4 Neighbor8 = PostprocessInput0.SampleLevel(PostprocessInput0Sampler, UVAndScreenPos.xy, 0, int2( 1,  1));
-	#if AA_GREEN_AS_LUMA
-		Neighbor0.rgb *= HdrWeightG(Neighbor0.rgb, InExposureScale);
-		Neighbor1.rgb *= HdrWeightG(Neighbor1.rgb, InExposureScale);
-		Neighbor2.rgb *= HdrWeightG(Neighbor2.rgb, InExposureScale);
-		Neighbor3.rgb *= HdrWeightG(Neighbor3.rgb, InExposureScale);
-		Neighbor4.rgb *= HdrWeightG(Neighbor4.rgb, InExposureScale);
-		Neighbor5.rgb *= HdrWeightG(Neighbor5.rgb, InExposureScale);
-		Neighbor6.rgb *= HdrWeightG(Neighbor6.rgb, InExposureScale);
-		Neighbor7.rgb *= HdrWeightG(Neighbor7.rgb, InExposureScale);
-		Neighbor8.rgb *= HdrWeightG(Neighbor8.rgb, InExposureScale);
-	#else
-		Neighbor0.rgb *= HdrWeight4(Neighbor0.rgb, InExposureScale);
-		Neighbor1.rgb *= HdrWeight4(Neighbor1.rgb, InExposureScale);
-		Neighbor2.rgb *= HdrWeight4(Neighbor2.rgb, InExposureScale);
-		Neighbor3.rgb *= HdrWeight4(Neighbor3.rgb, InExposureScale);
-		Neighbor4.rgb *= HdrWeight4(Neighbor4.rgb, InExposureScale);
-		Neighbor5.rgb *= HdrWeight4(Neighbor5.rgb, InExposureScale);
-		Neighbor6.rgb *= HdrWeight4(Neighbor6.rgb, InExposureScale);
-		Neighbor7.rgb *= HdrWeight4(Neighbor7.rgb, InExposureScale);
-		Neighbor8.rgb *= HdrWeight4(Neighbor8.rgb, InExposureScale);
-	#endif
-	#if AA_FILTERED
-		float4 Filtered = 
-			Neighbor0 * SampleWeights[0] +
-			Neighbor1 * SampleWeights[1] +
-			Neighbor2 * SampleWeights[2] +
-			Neighbor3 * SampleWeights[3] +
-			Neighbor4 * SampleWeights[4] +
-			Neighbor5 * SampleWeights[5] +
-			Neighbor6 * SampleWeights[6] +
-			Neighbor7 * SampleWeights[7] +
-			Neighbor8 * SampleWeights[8];
-		#if AA_LOWPASS
-			float4 FilteredLow = 
-				Neighbor0 * LowpassWeights[0] +
-				Neighbor1 * LowpassWeights[1] +
-				Neighbor2 * LowpassWeights[2] +
-				Neighbor3 * LowpassWeights[3] +
-				Neighbor4 * LowpassWeights[4] +
-				Neighbor5 * LowpassWeights[5] +
-				Neighbor6 * LowpassWeights[6] +
-				Neighbor7 * LowpassWeights[7] +
-				Neighbor8 * LowpassWeights[8];
-		#else
-			float4 FilteredLow = Filtered;
+	#if AA_YCOCG	
+		// Special case, only using 5 taps.
+		float4 Neighbor1 = PostprocessInput0.SampleLevel(PostprocessInput0Sampler, UVAndScreenPos.xy, 0, int2( 0, -1));
+		float4 Neighbor3 = PostprocessInput0.SampleLevel(PostprocessInput0Sampler, UVAndScreenPos.xy, 0, int2(-1,  0));
+		float4 Neighbor4 = PostprocessInput0.SampleLevel(PostprocessInput0Sampler, UVAndScreenPos.xy, 0);
+		float4 Neighbor5 = PostprocessInput0.SampleLevel(PostprocessInput0Sampler, UVAndScreenPos.xy, 0, int2( 1,  0));
+		float4 Neighbor7 = PostprocessInput0.SampleLevel(PostprocessInput0Sampler, UVAndScreenPos.xy, 0, int2( 0,  1));
+		#if 1
+			Neighbor1.rgb *= HdrWeightG(Neighbor1.rgb, InExposureScale);
+			Neighbor3.rgb *= HdrWeightG(Neighbor3.rgb, InExposureScale);
+			Neighbor4.rgb *= HdrWeightG(Neighbor4.rgb, InExposureScale);
+			Neighbor5.rgb *= HdrWeightG(Neighbor5.rgb, InExposureScale);
+			Neighbor7.rgb *= HdrWeightG(Neighbor7.rgb, InExposureScale);
 		#endif
-		#if AA_BORDER
-			// Use unfiltered for 1 pixel border.
-			float2 TestPos = abs(UVAndScreenPos.zw);
-			// Add 1 pixel and check if off screen.
-			TestPos += ViewportSize.zw * 2.0;
-			bool FilteredOffScreen = max(TestPos.x, TestPos.y) >= 1.0;
-			if(FilteredOffScreen) 
-			{
-				Filtered = Neighbor4;
-				FilteredLow = Neighbor4;
-			}
+		Neighbor1.rgb = RGBToYCoCg(Neighbor1.rgb);
+		Neighbor3.rgb = RGBToYCoCg(Neighbor3.rgb);
+		Neighbor4.rgb = RGBToYCoCg(Neighbor4.rgb);
+		Neighbor5.rgb = RGBToYCoCg(Neighbor5.rgb);
+		Neighbor7.rgb = RGBToYCoCg(Neighbor7.rgb);
+		#if 0
+			Neighbor1.xyz *= HdrWeightY(Neighbor1.x, InExposureScale);
+			Neighbor3.xyz *= HdrWeightY(Neighbor3.x, InExposureScale);
+			Neighbor4.xyz *= HdrWeightY(Neighbor4.x, InExposureScale);
+			Neighbor5.xyz *= HdrWeightY(Neighbor5.x, InExposureScale);
+			Neighbor7.xyz *= HdrWeightY(Neighbor7.x, InExposureScale);
 		#endif
-	#else
-		// Unfiltered.
-		float4 Filtered = Neighbor4;
-		float4 FilteredLow = Neighbor4;
-	#endif
-	#if AA_ROUND
-		float4 NeighborMin2 = min(min(Neighbor0, Neighbor2), min(Neighbor6, Neighbor8));		
-		float4 NeighborMax2 = max(max(Neighbor0, Neighbor2), max(Neighbor6, Neighbor8));		
+		#if AA_FILTERED
+			float4 Filtered = 
+				Neighbor1 * PlusWeights[0] +
+				Neighbor3 * PlusWeights[1] +
+				Neighbor4 * PlusWeights[2] +
+				Neighbor5 * PlusWeights[3] +
+				Neighbor7 * PlusWeights[4];
+			#if AA_BORDER
+				// Use unfiltered for 1 pixel border.
+				float2 TestPos = abs(UVAndScreenPos.zw);
+				// Add 1 pixel and check if off screen.
+				TestPos += ViewportSize.zw * 2.0;
+				bool FilteredOffScreen = max(TestPos.x, TestPos.y) >= 1.0;
+				if(FilteredOffScreen) 
+				{
+					Filtered = Neighbor4;
+				}
+			#endif
+		#else
+			// Unfiltered.
+			float4 Filtered = Neighbor4;
+		#endif
+		float4 FilteredLow = Filtered;
+		// Neighborhood seems to only need the "+" pattern.
 		float4 NeighborMin = min(min(min(Neighbor1, Neighbor3), min(Neighbor4, Neighbor5)), Neighbor7);		
 		float4 NeighborMax = max(max(max(Neighbor1, Neighbor3), max(Neighbor4, Neighbor5)), Neighbor7);		
-		NeighborMin2 = min(NeighborMin2, NeighborMin);
-		NeighborMax2 = max(NeighborMax2, NeighborMax);
-		NeighborMin = NeighborMin * 0.5 + NeighborMin2 * 0.5;
-		NeighborMax = NeighborMax * 0.5 + NeighborMax2 * 0.5;
 	#else
-		float4 NeighborMin = min(min(
-			min(min(Neighbor0, Neighbor1), min(Neighbor2, Neighbor3)),
-			min(min(Neighbor4, Neighbor5), min(Neighbor6, Neighbor7))), Neighbor8);
-		float4 NeighborMax = max(max(
-			max(max(Neighbor0, Neighbor1), max(Neighbor2, Neighbor3)),
-			max(max(Neighbor4, Neighbor5), max(Neighbor6, Neighbor7))), Neighbor8);
+		float4 Neighbor0 = PostprocessInput0.SampleLevel(PostprocessInput0Sampler, UVAndScreenPos.xy, 0, int2(-1, -1));
+		float4 Neighbor1 = PostprocessInput0.SampleLevel(PostprocessInput0Sampler, UVAndScreenPos.xy, 0, int2( 0, -1));
+		float4 Neighbor2 = PostprocessInput0.SampleLevel(PostprocessInput0Sampler, UVAndScreenPos.xy, 0, int2( 1, -1));
+		float4 Neighbor3 = PostprocessInput0.SampleLevel(PostprocessInput0Sampler, UVAndScreenPos.xy, 0, int2(-1,  0));
+		float4 Neighbor4 = PostprocessInput0.SampleLevel(PostprocessInput0Sampler, UVAndScreenPos.xy, 0);
+		float4 Neighbor5 = PostprocessInput0.SampleLevel(PostprocessInput0Sampler, UVAndScreenPos.xy, 0, int2( 1,  0));
+		float4 Neighbor6 = PostprocessInput0.SampleLevel(PostprocessInput0Sampler, UVAndScreenPos.xy, 0, int2(-1,  1));
+		float4 Neighbor7 = PostprocessInput0.SampleLevel(PostprocessInput0Sampler, UVAndScreenPos.xy, 0, int2( 0,  1));
+		float4 Neighbor8 = PostprocessInput0.SampleLevel(PostprocessInput0Sampler, UVAndScreenPos.xy, 0, int2( 1,  1));
+		#if AA_GREEN_AS_LUMA
+			Neighbor0.rgb *= HdrWeightG(Neighbor0.rgb, InExposureScale);
+			Neighbor1.rgb *= HdrWeightG(Neighbor1.rgb, InExposureScale);
+			Neighbor2.rgb *= HdrWeightG(Neighbor2.rgb, InExposureScale);
+			Neighbor3.rgb *= HdrWeightG(Neighbor3.rgb, InExposureScale);
+			Neighbor4.rgb *= HdrWeightG(Neighbor4.rgb, InExposureScale);
+			Neighbor5.rgb *= HdrWeightG(Neighbor5.rgb, InExposureScale);
+			Neighbor6.rgb *= HdrWeightG(Neighbor6.rgb, InExposureScale);
+			Neighbor7.rgb *= HdrWeightG(Neighbor7.rgb, InExposureScale);
+			Neighbor8.rgb *= HdrWeightG(Neighbor8.rgb, InExposureScale);
+		#else
+			Neighbor0.rgb *= HdrWeight4(Neighbor0.rgb, InExposureScale);
+			Neighbor1.rgb *= HdrWeight4(Neighbor1.rgb, InExposureScale);
+			Neighbor2.rgb *= HdrWeight4(Neighbor2.rgb, InExposureScale);
+			Neighbor3.rgb *= HdrWeight4(Neighbor3.rgb, InExposureScale);
+			Neighbor4.rgb *= HdrWeight4(Neighbor4.rgb, InExposureScale);
+			Neighbor5.rgb *= HdrWeight4(Neighbor5.rgb, InExposureScale);
+			Neighbor6.rgb *= HdrWeight4(Neighbor6.rgb, InExposureScale);
+			Neighbor7.rgb *= HdrWeight4(Neighbor7.rgb, InExposureScale);
+			Neighbor8.rgb *= HdrWeight4(Neighbor8.rgb, InExposureScale);
+		#endif
+		#if AA_FILTERED
+			float4 Filtered = 
+				Neighbor0 * SampleWeights[0] +
+				Neighbor1 * SampleWeights[1] +
+				Neighbor2 * SampleWeights[2] +
+				Neighbor3 * SampleWeights[3] +
+				Neighbor4 * SampleWeights[4] +
+				Neighbor5 * SampleWeights[5] +
+				Neighbor6 * SampleWeights[6] +
+				Neighbor7 * SampleWeights[7] +
+				Neighbor8 * SampleWeights[8];
+			#if AA_LOWPASS
+				float4 FilteredLow = 
+					Neighbor0 * LowpassWeights[0] +
+					Neighbor1 * LowpassWeights[1] +
+					Neighbor2 * LowpassWeights[2] +
+					Neighbor3 * LowpassWeights[3] +
+					Neighbor4 * LowpassWeights[4] +
+					Neighbor5 * LowpassWeights[5] +
+					Neighbor6 * LowpassWeights[6] +
+					Neighbor7 * LowpassWeights[7] +
+					Neighbor8 * LowpassWeights[8];
+			#else
+				float4 FilteredLow = Filtered;
+			#endif
+			#if AA_BORDER
+				// Use unfiltered for 1 pixel border.
+				float2 TestPos = abs(UVAndScreenPos.zw);
+				// Add 1 pixel and check if off screen.
+				TestPos += ViewportSize.zw * 2.0;
+				bool FilteredOffScreen = max(TestPos.x, TestPos.y) >= 1.0;
+				if(FilteredOffScreen) 
+				{
+					Filtered = Neighbor4;
+					FilteredLow = Neighbor4;
+				}
+			#endif
+		#else
+			// Unfiltered.
+			float4 Filtered = Neighbor4;
+			float4 FilteredLow = Neighbor4;
+		#endif
+		#if AA_ROUND
+			float4 NeighborMin2 = min(min(Neighbor0, Neighbor2), min(Neighbor6, Neighbor8));		
+			float4 NeighborMax2 = max(max(Neighbor0, Neighbor2), max(Neighbor6, Neighbor8));		
+			float4 NeighborMin = min(min(min(Neighbor1, Neighbor3), min(Neighbor4, Neighbor5)), Neighbor7);		
+			float4 NeighborMax = max(max(max(Neighbor1, Neighbor3), max(Neighbor4, Neighbor5)), Neighbor7);		
+			NeighborMin2 = min(NeighborMin2, NeighborMin);
+			NeighborMax2 = max(NeighborMax2, NeighborMax);
+			NeighborMin = NeighborMin * 0.5 + NeighborMin2 * 0.5;
+			NeighborMax = NeighborMax * 0.5 + NeighborMax2 * 0.5;
+		#else
+			float4 NeighborMin = min(min(
+				min(min(Neighbor0, Neighbor1), min(Neighbor2, Neighbor3)),
+				min(min(Neighbor4, Neighbor5), min(Neighbor6, Neighbor7))), Neighbor8);
+			float4 NeighborMax = max(max(
+				max(max(Neighbor0, Neighbor1), max(Neighbor2, Neighbor3)),
+				max(max(Neighbor4, Neighbor5), max(Neighbor6, Neighbor7))), Neighbor8);
+		#endif
 	#endif
 
 	// FETCH HISTORY
@@ -347,10 +372,20 @@ which can be configured via defines for various temporal AA passes.
 		FilteredLow.rg = float2(0.0, 0.0);
 		float DebugDiffCurrent = Filtered.b;
 	#endif
-	#if AA_GREEN_AS_LUMA
-		OutColor.rgb *= HdrWeightG(OutColor.rgb, InExposureScale);
+	#if AA_YCOCG
+		#if 1
+			OutColor.rgb *= HdrWeightG(OutColor.rgb, InExposureScale);
+		#endif
+		OutColor.rgb = RGBToYCoCg(OutColor.rgb);
+		#if 0
+			OutColor.xyz *= HdrWeightY(OutColor.x, InExposureScale);
+		#endif
 	#else
-		OutColor.rgb *= HdrWeight4(OutColor.rgb, InExposureScale);
+		#if AA_GREEN_AS_LUMA
+			OutColor.rgb *= HdrWeightG(OutColor.rgb, InExposureScale);
+		#else
+			OutColor.rgb *= HdrWeight4(OutColor.rgb, InExposureScale);
+		#endif
 	#endif
 	#if AA_DEBUG
 		OutColor.rg = float2(0.0, 0.0);
@@ -360,27 +395,40 @@ which can be configured via defines for various temporal AA passes.
 	// FIND LUMA OF CLAMPED HISTORY
 	// ----------------------------
 	// Save off luma of history before the clamp.
-	#if AA_GREEN_AS_LUMA
-		float LumaMin = NeighborMin.g;
-		float LumaMax = NeighborMax.g;
-		float LumaHistory = OutColor.g;
+	#if AA_YCOCG
+		float LumaMin = NeighborMin.x;
+		float LumaMax = NeighborMax.x;
+		float LumaHistory = OutColor.x;
 	#else
-		float LumaMin = Luma4(NeighborMin.rgb);
-		float LumaMax = Luma4(NeighborMax.rgb);
-		float LumaHistory = Luma4(OutColor.rgb);
+		#if AA_GREEN_AS_LUMA
+			float LumaMin = NeighborMin.g;
+			float LumaMax = NeighborMax.g;
+			float LumaHistory = OutColor.g;
+		#else
+			float LumaMin = Luma4(NeighborMin.rgb);
+			float LumaMax = Luma4(NeighborMax.rgb);
+			float LumaHistory = Luma4(OutColor.rgb);
+		#endif
 	#endif
 	float LumaContrast = LumaMax - LumaMin;
-	#if AA_AABB
-		// Clamp history, this uses color AABB intersection for tighter fit.
-		// Clamping works with the low pass (if available) to reduce flicker.
-		float ClampBlend = HistoryClamp(OutColor.rgb, FilteredLow.rgb, NeighborMin.rgb, NeighborMax.rgb);
-		#if AA_ALPHA
-			OutColor.rgb = lerp(OutColor.rgb, FilteredLow.rgb, ClampBlend);
-		#else
-			OutColor.rgba = lerp(OutColor.rgba, FilteredLow.rgba, ClampBlend);
+	#if AA_YCOCG
+		OutColor.rgb = clamp(OutColor.rgb, NeighborMin.rgb, NeighborMax.rgb);
+		#if (AA_ALPHA == 0)
+			OutColor.a = clamp(OutColor.a, NeighborMin.a, NeighborMax.a);
 		#endif
 	#else
-		OutColor = clamp(OutColor, NeighborMin, NeighborMax);
+		#if AA_AABB
+			// Clamp history, this uses color AABB intersection for tighter fit.
+			// Clamping works with the low pass (if available) to reduce flicker.
+			float ClampBlend = HistoryClamp(OutColor.rgb, FilteredLow.rgb, NeighborMin.rgb, NeighborMax.rgb);
+			#if AA_ALPHA
+				OutColor.rgb = lerp(OutColor.rgb, FilteredLow.rgb, ClampBlend);
+			#else
+				OutColor.rgba = lerp(OutColor.rgba, FilteredLow.rgba, ClampBlend);
+			#endif
+		#else
+			OutColor = clamp(OutColor, NeighborMin, NeighborMax);
+		#endif
 	#endif
 	#if AA_DEBUG
 		OutColor.rg = float2(0.0, 0.0);
@@ -392,17 +440,21 @@ which can be configured via defines for various temporal AA passes.
 		// Blend in non-filtered based on the amount of sub-pixel motion.
 		float AddAliasing = saturate(HistoryBlur) * 0.5;
 		float LumaContrastFactor = 32.0;
-		#if AA_GREEN_AS_LUMA
+		#if AA_GREEN_AS_LUMA || AA_YCOCG
 			// GREEN_AS_LUMA is 1/4 as bright.
 			LumaContrastFactor *= 4.0;
 		#endif
 		AddAliasing = saturate(AddAliasing + rcp(1.0 + LumaContrast * LumaContrastFactor));
 		Filtered.rgb = lerp(Filtered.rgb, Neighbor4.rgb, AddAliasing);
 	#endif
-	#if AA_GREEN_AS_LUMA
-		float LumaFiltered = Filtered.g;
+	#if AA_YCOCG
+		float LumaFiltered = Filtered.x;
 	#else
-		float LumaFiltered = Luma4(Filtered.rgb);
+		#if AA_GREEN_AS_LUMA
+			float LumaFiltered = Filtered.g;
+		#else
+			float LumaFiltered = Luma4(Filtered.rgb);
+		#endif
 	#endif
 		
 	// COMPUTE BLEND AMOUNT 
@@ -458,500 +510,23 @@ which can be configured via defines for various temporal AA passes.
 				OutColor.a = clamp(OutColor.a, NeighborMin.a, NeighborMax.a);
 			#endif
 		#endif
-		// Convert back into linear.
-		#if AA_GREEN_AS_LUMA
-			OutColor.rgb *= HdrWeightInvG(OutColor.rgb, InExposureScale);
-		#else
-			OutColor.rgb *= HdrWeightInv4(OutColor.rgb, InExposureScale);
-		#endif
 	#else
 		OutColor = lerp(OutColor, Filtered, FixedLerp);
-		#if AA_GREEN_AS_LUMA
-			OutColor.rgb *= HdrWeightInvG(OutColor.rgb, InExposureScale);
-		#else
-			OutColor.rgb *= HdrWeightInv4(OutColor.rgb, InExposureScale);
-		#endif
-	#endif
-	#if AA_NAN
-		// Transform NaNs to black, transform negative colors to black.
-		OutColor.rgb = -min(-OutColor.rgb, 0.0);
-	#endif
-	#if AA_DEBUG
-		OutColor.g = abs(DebugDiffPrior - DebugDiffCurrent);
-		OutColor.r = OutColor.a;
-	#endif
-
-#else
-
-	// FIND MOTION OF PIXEL AND NEAREST IN NEIGHBORHOOD
-    // ------------------------------------------------
-	float3 PosP; // Position of this pixel.
-	PosP.xy = UVAndScreenPos.zw * float2(0.5, -0.5) + 0.5; // View position [0 to 1] flipped in Y.
-	PosP.z = SceneDepthTexture.SampleLevel(SceneDepthTextureSampler, UVAndScreenPos.xy, 0).r;
-	float3 PosN; // Position of closest pixel in neighborhood.
-	PosN = PosP;
-	// Screen position of minimum depth.
-	float2 VelocityOffset = float2(0.0, 0.0);
-	#if AA_CROSS
-		// For motion vector, use camera/dynamic motion from min depth pixel in pattern around pixel.
-		// This enables better quality outline on foreground against different motion background.
-		// Larger 2 pixel distance "x" works best (because AA diolates surface).
-		float4 Depths;
-		#if FEATURE_LEVEL >= FEATURE_LEVEL_SM5
-			Depths = SceneDepthTexture.GatherRed(SceneDepthTextureSampler, UVAndScreenPos.xy, int2(-AA_CROSS, -AA_CROSS), int2(AA_CROSS, -AA_CROSS), int2(-AA_CROSS, AA_CROSS), int2(AA_CROSS, AA_CROSS));
-		#else
-			Depths.x = SceneDepthTexture.SampleLevel(SceneDepthTextureSampler, UVAndScreenPos.xy, 0, int2(-AA_CROSS, -AA_CROSS)).r;
-			Depths.y = SceneDepthTexture.SampleLevel(SceneDepthTextureSampler, UVAndScreenPos.xy, 0, int2( AA_CROSS, -AA_CROSS)).r;
-			Depths.z = SceneDepthTexture.SampleLevel(SceneDepthTextureSampler, UVAndScreenPos.xy, 0, int2(-AA_CROSS,  AA_CROSS)).r;
-			Depths.w = SceneDepthTexture.SampleLevel(SceneDepthTextureSampler, UVAndScreenPos.xy, 0, int2( AA_CROSS,  AA_CROSS)).r;
-		#endif
-		float2 DepthOffset = float2(AA_CROSS, AA_CROSS);
-		float DepthOffsetXx = float(AA_CROSS);
-		// Nearest depth is the largest depth (depth surface 0=far, 1=near).
-		if(Depths.x > Depths.y) 
-		{
-			DepthOffsetXx = -AA_CROSS;
-		}
-		if(Depths.z > Depths.w) 
-		{
-			DepthOffset.x = -AA_CROSS;
-		}
-		float DepthsXY = max(Depths.x, Depths.y);
-		float DepthsZW = max(Depths.z, Depths.w);
-		if(DepthsXY > DepthsZW) 
-		{
-			DepthOffset.y = -AA_CROSS;
-			DepthOffset.x = DepthOffsetXx; 
-		}
-		float DepthsXYZW = max(DepthsXY, DepthsZW);
-		if(DepthsXYZW > PosP.z) 
-		{
-			// This is offset for reading from velocity texture.
-			// This supports half or fractional resolution velocity textures.
-			// With the assumption that UV position scales between velocity and color.
-			VelocityOffset = DepthOffset * PostprocessInput0Size.zw;
-			// This is [0 to 1] flipped in Y.
-			PosN.xy = (UVAndScreenPos.zw + DepthOffset * ViewportSize.zw * 2.0) * float2(0.5, -0.5) + 0.5;
-			PosN.z = DepthsXYZW;
-		}
 	#endif
-	// Camera motion for pixel (in ScreenPos space).
-	float ScaleM = 1.0 / (dot(PosP, CameraMotion[0].xyz) + CameraMotion[0].w);
-	float2 BackP;
-	BackP.x = -2.0 * ((PosP.x * ((CameraMotion[1].x * PosP.y) + (CameraMotion[1].y * PosP.z) + CameraMotion[1].z)) + (CameraMotion[1].w * PosP.y) + (CameraMotion[2].x * PosP.x * PosP.x) + (CameraMotion[2].y * PosP.z) + CameraMotion[2].z) * ScaleM;
-	BackP.y =  2.0 * ((PosP.y * ((CameraMotion[3].x * PosP.x) + (CameraMotion[3].y * PosP.z) + CameraMotion[3].z)) + (CameraMotion[3].w * PosP.x) + (CameraMotion[4].x * PosP.y * PosP.y) + (CameraMotion[4].y * PosP.z) + CameraMotion[4].z) * ScaleM;
-	// Convert BackP back into pixels.
-	float2 BackTemp = BackP * ViewportSize.xy;
-	// At this point DoublePixelMotionCamera is kept squared until later.
-	float DoublePixelMotionCamera = dot(BackTemp, BackTemp);
-	#if AA_CROSS
-		// Camera motion for neighborhood.
-		ScaleM = 1.0 / (dot(PosN, CameraMotion[0].xyz) + CameraMotion[0].w);
-		float2 BackN;
-		BackN.x = -2.0 * ((PosN.x * ((CameraMotion[1].x * PosN.y) + (CameraMotion[1].y * PosN.z) + CameraMotion[1].z)) + (CameraMotion[1].w * PosN.y) + (CameraMotion[2].x * PosN.x * PosN.x) + (CameraMotion[2].y * PosN.z) + CameraMotion[2].z) * ScaleM;
-		BackN.y =  2.0 * ((PosN.y * ((CameraMotion[3].x * PosN.x) + (CameraMotion[3].y * PosN.z) + CameraMotion[3].z)) + (CameraMotion[3].w * PosN.x) + (CameraMotion[4].x * PosN.y * PosN.y) + (CameraMotion[4].y * PosN.z) + CameraMotion[4].z) * ScaleM;
-		BackTemp = BackN * ViewportSize.xy;
-		DoublePixelMotionCamera = max(DoublePixelMotionCamera, dot(BackTemp, BackTemp));
-	#endif
-	DoublePixelMotionCamera = sqrt(DoublePixelMotionCamera);
-	#if AA_DYNAMIC
-		// Dynamic motion of pixel.
-		#if (AA_CROSS == 0) || (AA_ONE_DYNAMIC_SAMPLE == 0)
-			float2 VelocityP = PostprocessInput3.SampleLevel(PostprocessInput3Sampler, UVAndScreenPos.xy, 0).xy;
-			bool DynamicP = VelocityP.x > 0.0;
-			if(DynamicP)
-			{
-				BackP = DecodeVelocityFromTexture(VelocityP);
-			}
-		#endif
-		BackTemp = BackP * ViewportSize.xy;
-		// At this point DoublePixelMotionDynamic is kept squared until later.
-		float DoublePixelMotionDynamic = dot(BackTemp, BackTemp);
-		#if AA_CROSS
-			// Dynamic motion of nearest in neighborhood.
-			float2 VelocityN = PostprocessInput3.SampleLevel(PostprocessInput3Sampler, UVAndScreenPos.xy + VelocityOffset, 0).xy;
-			bool DynamicN = VelocityN.x > 0.0;
-			if(DynamicN)
-			{
-				BackN = DecodeVelocityFromTexture(VelocityN);
-			}
-			#if (AA_ONE_DYNAMIC_SAMPLE == 0)
-				bool DynamicMotion = DynamicN || DynamicP;
-			#else
-				bool DynamicMotion = DynamicN;
-			#endif
-			BackTemp = BackN * ViewportSize.xy;
-			DoublePixelMotionDynamic = max(DoublePixelMotionDynamic, dot(BackTemp, BackTemp));
-		#else
-			bool DynamicMotion = DynamicP;
+	#if AA_YCOCG
+		#if 0
+			OutColor.xyz *= HdrWeightInvY(OutColor.x, InExposureScale);
 		#endif
-		DoublePixelMotionDynamic = sqrt(DoublePixelMotionDynamic);
-	#else
-		bool DynamicMotion = false;
-		float DoublePixelMotionDynamic = 0.0;
-	#endif
-	// Find motion edges.
-	#if AA_CROSS
-		float2 BackPN = (BackP - BackN) * ViewportSize.xy;	
-		float DoublePixelMotionDiff = sqrt(dot(BackPN, BackPN));
-	#else
-		float2 BackN = BackP;
-		float DoublePixelMotionDiff = 0.0;
-	#endif
-	// Easier to do off screen check before conversion.
-	// BackN is in units of 2pixels/viewportWidthInPixels
-	// This converts back projection vector to [-1 to 1] offset in viewport.
-	BackN = UVAndScreenPos.zw - BackN;
-	bool OffScreen = max(abs(BackN.x), abs(BackN.y)) >= 1.0;
-	// Also clamp to be on screen (fixes problem with DOF).
-	// The .z and .w is the 1/width and 1/height.
-	// This clamps to be a pixel inside the viewport.
-	BackN.x = clamp(BackN.x, -1.0 + ViewportSize.z, 1.0 - ViewportSize.z);
-	BackN.y = clamp(BackN.y, -1.0 + ViewportSize.w, 1.0 - ViewportSize.w);
-	// Convert from [-1 to 1] to view rectangle which is somewhere in [0 to 1].
-	// The extra +0.5 factor is because ScreenPosToPixel.zw is incorrectly computed
-	// as the upper left of the pixel instead of the center of the pixel.
-	BackN = (BackN * ScreenPosToPixel.xy + ScreenPosToPixel.zw + 0.5) * PostprocessInput0Size.zw;
-
-	// FILTER PIXEL (RESAMPLE TO REMOVE JITTER OFFSET) AND GET NEIGHBORHOOD
-    // --------------------------------------------------------------------
-	// 012
-	// 345
-	// 678
-	float4 Neighbor0 = PostprocessInput0.SampleLevel(PostprocessInput0Sampler, UVAndScreenPos.xy, 0, int2(-1, -1));
-	float4 Neighbor1 = PostprocessInput0.SampleLevel(PostprocessInput0Sampler, UVAndScreenPos.xy, 0, int2( 0, -1));
-	float4 Neighbor2 = PostprocessInput0.SampleLevel(PostprocessInput0Sampler, UVAndScreenPos.xy, 0, int2( 1, -1));
-	float4 Neighbor3 = PostprocessInput0.SampleLevel(PostprocessInput0Sampler, UVAndScreenPos.xy, 0, int2(-1,  0));
-	float4 Neighbor4 = PostprocessInput0.SampleLevel(PostprocessInput0Sampler, UVAndScreenPos.xy, 0);
-	float4 Neighbor5 = PostprocessInput0.SampleLevel(PostprocessInput0Sampler, UVAndScreenPos.xy, 0, int2( 1,  0));
-	float4 Neighbor6 = PostprocessInput0.SampleLevel(PostprocessInput0Sampler, UVAndScreenPos.xy, 0, int2(-1,  1));
-	float4 Neighbor7 = PostprocessInput0.SampleLevel(PostprocessInput0Sampler, UVAndScreenPos.xy, 0, int2( 0,  1));
-	float4 Neighbor8 = PostprocessInput0.SampleLevel(PostprocessInput0Sampler, UVAndScreenPos.xy, 0, int2( 1,  1));
-	#if AA_FILTERED
-		#if AA_HDR
-			float SampleHdrWeight0 = HdrWeight(Neighbor0.rgb, InExposureScale);
-			float SampleHdrWeight1 = HdrWeight(Neighbor1.rgb, InExposureScale);
-			float SampleHdrWeight2 = HdrWeight(Neighbor2.rgb, InExposureScale);
-			float SampleHdrWeight3 = HdrWeight(Neighbor3.rgb, InExposureScale);
-			float SampleHdrWeight4 = HdrWeight(Neighbor4.rgb, InExposureScale);
-			float SampleHdrWeight5 = HdrWeight(Neighbor5.rgb, InExposureScale);
-			float SampleHdrWeight6 = HdrWeight(Neighbor6.rgb, InExposureScale);
-			float SampleHdrWeight7 = HdrWeight(Neighbor7.rgb, InExposureScale);
-			float SampleHdrWeight8 = HdrWeight(Neighbor8.rgb, InExposureScale);
-			float4 Filtered = (
-				Neighbor0 * (SampleWeights[0] * SampleHdrWeight0) +
-				Neighbor1 * (SampleWeights[1] * SampleHdrWeight1) +
-				Neighbor2 * (SampleWeights[2] * SampleHdrWeight2) +
-				Neighbor3 * (SampleWeights[3] * SampleHdrWeight3) +
-				Neighbor4 * (SampleWeights[4] * SampleHdrWeight4) +
-				Neighbor5 * (SampleWeights[5] * SampleHdrWeight5) +
-				Neighbor6 * (SampleWeights[6] * SampleHdrWeight6) +
-				Neighbor7 * (SampleWeights[7] * SampleHdrWeight7) +
-				Neighbor8 * (SampleWeights[8] * SampleHdrWeight8)) * rcp(
-					SampleWeights[0] * SampleHdrWeight0 +
-					SampleWeights[1] * SampleHdrWeight1 +
-					SampleWeights[2] * SampleHdrWeight2 +
-					SampleWeights[3] * SampleHdrWeight3 +
-					SampleWeights[4] * SampleHdrWeight4 +
-					SampleWeights[5] * SampleHdrWeight5 +
-					SampleWeights[6] * SampleHdrWeight6 +
-					SampleWeights[7] * SampleHdrWeight7 +
-					SampleWeights[8] * SampleHdrWeight8);
-			#if AA_LOWPASS
-				float4 FilteredLow = (
-					Neighbor0 * (LowpassWeights[0] * SampleHdrWeight0) +
-					Neighbor1 * (LowpassWeights[1] * SampleHdrWeight1) +
-					Neighbor2 * (LowpassWeights[2] * SampleHdrWeight2) +
-					Neighbor3 * (LowpassWeights[3] * SampleHdrWeight3) +
-					Neighbor4 * (LowpassWeights[4] * SampleHdrWeight4) +
-					Neighbor5 * (LowpassWeights[5] * SampleHdrWeight5) +
-					Neighbor6 * (LowpassWeights[6] * SampleHdrWeight6) +
-					Neighbor7 * (LowpassWeights[7] * SampleHdrWeight7) +
-					Neighbor8 * (LowpassWeights[8] * SampleHdrWeight8)) * rcp(
-						LowpassWeights[0] * SampleHdrWeight0 +
-						LowpassWeights[1] * SampleHdrWeight1 +
-						LowpassWeights[2] * SampleHdrWeight2 +
-						LowpassWeights[3] * SampleHdrWeight3 +
-						LowpassWeights[4] * SampleHdrWeight4 +
-						LowpassWeights[5] * SampleHdrWeight5 +
-						LowpassWeights[6] * SampleHdrWeight6 +
-						LowpassWeights[7] * SampleHdrWeight7 +
-						LowpassWeights[8] * SampleHdrWeight8);
-			#else
-				float4 FilteredLow = Filtered;
-			#endif
-		#else
-			float4 Filtered = 
-				Neighbor0 * SampleWeights[0] +
-				Neighbor1 * SampleWeights[1] +
-				Neighbor2 * SampleWeights[2] +
-				Neighbor3 * SampleWeights[3] +
-				Neighbor4 * SampleWeights[4] +
-				Neighbor5 * SampleWeights[5] +
-				Neighbor6 * SampleWeights[6] +
-				Neighbor7 * SampleWeights[7] +
-				Neighbor8 * SampleWeights[8];
-			#if AA_LOWPASS
-				float4 FilteredLow = 
-					Neighbor0 * LowpassWeights[0] +
-					Neighbor1 * LowpassWeights[1] +
-					Neighbor2 * LowpassWeights[2] +
-					Neighbor3 * LowpassWeights[3] +
-					Neighbor4 * LowpassWeights[4] +
-					Neighbor5 * LowpassWeights[5] +
-					Neighbor6 * LowpassWeights[6] +
-					Neighbor7 * LowpassWeights[7] +
-					Neighbor8 * LowpassWeights[8];
-			#else
-				float4 FilteredLow = Filtered;
-			#endif
-		#endif
-		#if AA_BORDER
-			// Use unfiltered for 1 pixel border.
-			float2 TestPos = abs(UVAndScreenPos.zw);
-			// Add 1 pixel and check if off screen.
-			TestPos += ViewportSize.zw * 2.0;
-			bool FilteredOffScreen = max(TestPos.x, TestPos.y) >= 1.0;
-			if(FilteredOffScreen) 
-			{
-				Filtered = Neighbor4;
-				FilteredLow = Filtered;
-			}
-		#endif
-	#else
-		// Unfiltered.
-		float4 Filtered = Neighbor4;
-		float4 FilteredLow = Filtered;
-	#endif
-	#if AA_ROUND
-		float4 NeighborMin2 = min(min(Neighbor0, Neighbor2), min(Neighbor6, Neighbor8));		
-		float4 NeighborMax2 = max(max(Neighbor0, Neighbor2), max(Neighbor6, Neighbor8));		
-		float4 NeighborMin = min(min(min(Neighbor1, Neighbor3), min(Neighbor4, Neighbor5)), Neighbor7);		
-		float4 NeighborMax = max(max(max(Neighbor1, Neighbor3), max(Neighbor4, Neighbor5)), Neighbor7);		
-		NeighborMin2 = min(NeighborMin2, NeighborMin);
-		NeighborMax2 = max(NeighborMax2, NeighborMax);
-		#if AA_HDR
-			NeighborMin = HdrLerp(NeighborMin, NeighborMin2, 0.5, InExposureScale);
-			NeighborMax = HdrLerp(NeighborMax, NeighborMax2, 0.5, InExposureScale);
-		#else
-			NeighborMin = lerp(NeighborMin, NeighborMin2, 0.5);
-			NeighborMax = lerp(NeighborMax, NeighborMax2, 0.5);
-		#endif
-	#else
-		float4 NeighborMin = min(min(
-			min(min(Neighbor0, Neighbor1), min(Neighbor2, Neighbor3)),
-			min(min(Neighbor4, Neighbor5), min(Neighbor6, Neighbor7))), Neighbor8);
-		float4 NeighborMax = max(max(
-			max(max(Neighbor0, Neighbor1), max(Neighbor2, Neighbor3)),
-			max(max(Neighbor4, Neighbor5), max(Neighbor6, Neighbor7))), Neighbor8);
-	#endif
-	#if AA_DEBUG
-		NeighborMin.rg = float2(0.0, 0.0);
-		NeighborMax.rg = float2(0.0, 0.0);
-		Filtered.rg = float2(0.0, 0.0);
-		FilteredLow.rg = float2(0.0, 0.0);
-		float DebugDiffCurrent = Filtered.b;
-	#endif
-
-	// FETCH HISTORY AND MANUALLY INTERPOLATE WITH WEIGHT FILTER FOR PROPER HDR
-	// ------------------------------------------------------------------------
-	#if AA_HDR_HISTORY
-		// a c
-		//  EF
-		// gHI
-		float4 OutColorE = PostprocessInput2.SampleLevel(PostprocessInput2Sampler, BackN.xy, 0, int2(0, 0));
-		float4 OutColorF = PostprocessInput2.SampleLevel(PostprocessInput2Sampler, BackN.xy, 0, int2(1, 0));
-		float4 OutColorH = PostprocessInput2.SampleLevel(PostprocessInput2Sampler, BackN.xy, 0, int2(0, 1));
-		float4 OutColorI = PostprocessInput2.SampleLevel(PostprocessInput2Sampler, BackN.xy, 0, int2(1, 1));
-		float WeightE = HdrWeight(OutColorE.rgb, InExposureScale);
-		float WeightF = HdrWeight(OutColorF.rgb, InExposureScale);
-		float WeightH = HdrWeight(OutColorH.rgb, InExposureScale);
-		float WeightI = HdrWeight(OutColorI.rgb, InExposureScale);
-		float2 Subpix = frac(BackN.xy * PostprocessInput1Size.zw);
-		WeightE *= (1.0 - Subpix.x) * (1.0 - Subpix.y);
-		WeightF *= (      Subpix.x) * (1.0 - Subpix.y);
-		WeightH *= (1.0 - Subpix.x) * (      Subpix.y);
-		WeightI *= (      Subpix.x) * (      Subpix.y);
-		float RcpWeightEFHI = rcp(WeightE + WeightF + WeightH + WeightI);
-		WeightE *= RcpWeightEFHI;
-		WeightF *= RcpWeightEFHI;
-		WeightH *= RcpWeightEFHI;
-		WeightI *= RcpWeightEFHI;
-		OutColor = (OutColorE * WeightE) + (OutColorF * WeightF) + (OutColorH * WeightH) + (OutColorI * WeightI);
-	#else
-		OutColor = PostprocessInput1.SampleLevel(PostprocessInput1Sampler, BackN.xy, 0);
-	#endif
-	#if AA_DILATE
-		// Grab alphas from cross pattern and take maximum, dilate feedback control.
-		float4 Alphas;
-		#if FEATURE_LEVEL >= FEATURE_LEVEL_SM5
-			// Samples acgI
-			Alphas = PostprocessInput2.GatherAlpha(PostprocessInput2Sampler, BackN.xy, int2(-AA_DILATE, -AA_DILATE), int2(AA_DILATE, -AA_DILATE), int2(-AA_DILATE, AA_DILATE), int2(AA_DILATE, AA_DILATE));
-		#else
-			Alphas.x = PostprocessInput2.SampleLevel(PostprocessInput2Sampler, BackN.xy, 0, int2(-AA_DILATE, -AA_DILATE)).w;
-			Alphas.y = PostprocessInput2.SampleLevel(PostprocessInput2Sampler, BackN.xy, 0, int2( AA_DILATE, -AA_DILATE)).w;
-			Alphas.z = PostprocessInput2.SampleLevel(PostprocessInput2Sampler, BackN.xy, 0, int2(-AA_DILATE,  AA_DILATE)).w;
-			#if (AA_HDR_HISTORY == 1) && (AA_DILATE == 1)
-				Alphas.w = OutColorI.w;
-			#else
-				Alphas.w = PostprocessInput2.SampleLevel(PostprocessInput2Sampler, BackN.xy, 0, int2(AA_DILATE, AA_DILATE)).w;
-			#endif
-		#endif
-		OutColor.w = max(OutColor.w, max(max(Alphas.x, Alphas.y), max(Alphas.z, Alphas.w)));
-	#endif
-	#if AA_DEBUG
-		OutColor.rg = float2(0.0, 0.0);
-		float DebugDiffPrior = OutColor.b;
-	#endif
-
-	// FIND LUMA OF CLAMPED HISTORY
-	// ----------------------------
-	// Save luma converted into a perceptual space.
-	float LumaHistory = PerceptualLuma(OutColor.rgb, InExposureScale);
-	#if AA_AABB
-		// Clamp history, this uses color AABB intersection for tighter fit.
-		// Clamping works with the low pass (if available) to reduce flicker.
-		float ClampBlend = HistoryClamp(OutColor.rgb, FilteredLow.rgb, NeighborMin.rgb, NeighborMax.rgb);
-		float4 Clamped = lerp(OutColor, FilteredLow, ClampBlend);
-	#else
-		float4 Clamped = clamp(OutColor, NeighborMin, NeighborMax);
-	#endif
-
-	// FIND PERCEPTUAL LUMAS
-	// ---------------------
-	float LumaClamped = PerceptualLuma(Clamped.rgb, InExposureScale);
-	float LumaFiltered = PerceptualLuma(Filtered.rgb, InExposureScale);
-	float LumaFilteredLow = PerceptualLuma(FilteredLow.rgb, InExposureScale);
-	float LumaMin = PerceptualLuma(NeighborMin.rgb, InExposureScale);
-	float LumaMax = PerceptualLuma(NeighborMax.rgb, InExposureScale);
-
-	// CONVERT MOTION AMOUNTS TO ALLOWED HISTORY TO FILTERED CHANGE AMOUNTS
-    // --------------------------------------------------------------------	
-	float3 Change;
-	// Small camera motion.
-	Change.x = (1.0/8.0) * saturate(DoublePixelMotionCamera * 128.0);
-	// Motion edges.
-	Change.y = (1.0/8.0) * saturate(DoublePixelMotionDiff * 8.0);
-	// Dynamic motion.
-	Change.z = 0.0;
-	if(DynamicMotion) 
-	{
-		Change.z = (1.0/2.0) * saturate(DoublePixelMotionDynamic * (1.0/8.0));
-	}
-	
-	// FEEDBACK CHANGE AMOUNT THEN CONVERT TO LUMA LIMITS FOR THIS FRAME
-    // -----------------------------------------------------------------	
-    float ChangeLimit = 1.0/2.0;
-   	#if AA_ALPHA 
-		#if (AA_DYNAMIC == 0) && (AA_CROSS == 0)
-			ChangeLimit = 0.0;
-		#else
-			// Feedback half of change from last frame.
-			float ChangePrior = OutColor.a * (1.0/2.0);
-		#endif
-		// Feedback of 2.0 provides 1 extra frame of fast reconvergence.
-		// Feedback of 4.0 provides 2 extra frames of fast reconvergence.
-		// This enables fast reconverge after dynamic motion.
-		float FeedbackMotionEdge = 4.0;
-		float FeedbackDynamic = 4.0;
-		// Change limit on this frame.
-		#if (AA_DYNAMIC == 1) && (AA_CROSS == 0)
-			OutColor.a = max(ChangePrior, Change.z * FeedbackDynamic);
-			ChangeLimit = min(ChangeLimit, max(ChangePrior, Change.z)); 
-		#endif 
-		#if (AA_DYNAMIC == 0) && (AA_CROSS != 0)
-			OutColor.a = max(ChangePrior, Change.y * FeedbackMotionEdge);
-			ChangeLimit = min(ChangeLimit, max(ChangePrior, Change.y)); 
-		#endif
-		#if (AA_DYNAMIC == 1) && (AA_CROSS != 0)
-			OutColor.a = max(ChangePrior, max(Change.y * FeedbackMotionEdge, Change.z * FeedbackDynamic));
-			ChangeLimit = min(ChangeLimit, max(ChangePrior, max(Change.y, Change.z)));
-		#endif
-	#else
-		#if (AA_DYNAMIC == 0) && (AA_CROSS == 0)
-			ChangeLimit = 0.0; 
-		#endif 
-		#if (AA_DYNAMIC == 1) && (AA_CROSS == 0)
-			ChangeLimit = min(ChangeLimit, Change.z); 
-		#endif 
-		#if (AA_DYNAMIC == 0) && (AA_CROSS != 0)
-			ChangeLimit = min(ChangeLimit, Change.y); 
-		#endif
-		#if (AA_DYNAMIC == 1) && (AA_CROSS != 0)
-			ChangeLimit = min(1.0/2.0, max(Change.y, Change.z));
-		#endif
-	#endif
-	#if AA_GRAIN
-		// Introduce noise to the change limit.
-		// This reduces the convergence to a non-AA look on some edges.
- 		float2 GrainUV = (UVAndScreenPos.zw * float2(0.5, -0.5) + 0.5) + RandomOffset.xy;
-		float Grain = frac(sin(GrainUV.x + GrainUV.y * 543.31) *  493013.0);
-		ChangeLimit += Grain * (1.0/16.0);
-	#endif
-	// Change limit is related to change in luma between history and current frame.
-    float ChangeLimitLuma = ChangeLimit * abs(LumaHistory - LumaFiltered);
-	// Remove the blur of small motion.
-	float EdgeFactor = min(abs(LumaMax - LumaFiltered), abs(LumaMin - LumaFiltered));
-	// Add in a constant factor to improve convergence speed with a small amount of noise.
-	Change.x = max(Change.x, 1.0/32.0);
-	#if (AA_DYNAMIC == 0) && (AA_CROSS == 0)
-		ChangeLimitLuma = EdgeFactor * Change.x;
-	#else
-		ChangeLimitLuma = max(ChangeLimitLuma, EdgeFactor * Change.x);
-	#endif
-	// Allow a minimal amount of change on still and motion.
-	ChangeLimitLuma = max(ChangeLimitLuma, 1.0/512.0);
-	// Target is filtered value clamped to +/- change limit.
-	float LumaTarget = clamp(LumaFiltered, LumaHistory - ChangeLimitLuma, LumaHistory + ChangeLimitLuma);
-	// Now perform the neighborhood clamp.
-	LumaTarget = clamp(LumaTarget, min(LumaFiltered, LumaClamped), max(LumaFiltered, LumaClamped));
-	// Offscreen feedback resets.
-	#if AA_LERP 
-		float FixedLerp = 1.0/float(AA_LERP);
-	#endif
-	if(OffScreen) 
-	{
-		OutColor = Filtered;
-		#if AA_ALPHA
-			OutColor.a = 1.0;
-		#endif
-		#if AA_LERP
-			FixedLerp = 1.0;
-		#endif
-	}
-	
-	// DO FINAL BLEND BETWEEN HISTORY AND FILTERED COLOR
-	// -------------------------------------------------
-	#if (AA_LERP == 0)
-		// Switch luma back to linear.
-		LumaTarget = LinearLuma(LumaTarget);
-		LumaHistory = LinearLuma(LumaHistory);
-		LumaFiltered = LinearLuma(LumaFiltered);
-		float LumaDiff = LumaHistory - LumaFiltered;
-		float RcpLumaDiff = rcp(LumaDiff);
-		if(abs(LumaDiff) < (1.0/1024.0)) 
-		{
-			// If diff between frames is under a set limit, use new frame 100%.
-			RcpLumaDiff = 0.0;
-		}
-		float BlendFinal = saturate((LumaTarget - LumaFiltered) * RcpLumaDiff);
-		#if AA_ALPHA
-			// Blend in linear to hit luma target.
-			OutColor.rgb = lerp(Filtered.rgb, OutColor.rgb, BlendFinal);
-			#if RESPONSIVE
-				OutColor.a = max(OutColor.a, 1.0);
-			#endif
-		#else
-			OutColor = lerp(Filtered, OutColor, BlendFinal);
-			#if AA_FORCE_ALPHA_CLAMP
-				OutColor.a = clamp(OutColor.a, NeighborMin.a, NeighborMax.a);
-			#endif
+		OutColor.rgb = YCoCgToRGB(OutColor.rgb);
+		#if 1
+			OutColor.rgb *= HdrWeightInvG(OutColor.rgb, InExposureScale);
 		#endif
 	#else
-		#if AA_HDR
-			OutColor = HdrLerp(Clamped, Filtered, FixedLerp, InExposureScale);
+		// Convert back into linear.
+		#if AA_GREEN_AS_LUMA
+			OutColor.rgb *= HdrWeightInvG(OutColor.rgb, InExposureScale);
 		#else
-			OutColor = lerp(Clamped, Filtered, FixedLerp);
+			OutColor.rgb *= HdrWeightInv4(OutColor.rgb, InExposureScale);
 		#endif
 	#endif
 	#if AA_NAN
@@ -963,7 +538,6 @@ which can be configured via defines for various temporal AA passes.
 		OutColor.r = OutColor.a;
 	#endif
 
-#endif
 
 
 #undef AA_AABB
@@ -982,7 +556,6 @@ which can be configured via defines for various temporal AA passes.
 #undef AA_NAN
 #undef AA_BORDER
 #undef AA_FORCE_ALPHA_CLAMP
-#undef AA_ONE_DYNAMIC_SAMPLE
 #undef AA_GREEN_AS_LUMA
-#undef AA_NEW
 #undef AA_VELOCITY_WEIGHTING
+#undef AA_YCOCG
diff --git a/Engine/Source/Runtime/Renderer/Private/PostProcess/PostProcessTemporalAA.cpp b/Engine/Source/Runtime/Renderer/Private/PostProcess/PostProcessTemporalAA.cpp
--- a/Engine/Source/Runtime/Renderer/Private/PostProcess/PostProcessTemporalAA.cpp
+++ b/Engine/Source/Runtime/Renderer/Private/PostProcess/PostProcessTemporalAA.cpp
@@ -64,6 +64,7 @@ public:
 	FDeferredPixelShaderParameters DeferredParameters;
 	FShaderParameter SampleWeights;
 	FShaderParameter LowpassWeights;
+	FShaderParameter PlusWeights;
 	FCameraMotionParameters CameraMotionParams;
 	FShaderParameter RandomOffset;
 
@@ -75,6 +76,7 @@ public:
 		DeferredParameters.Bind(Initializer.ParameterMap);
 		SampleWeights.Bind(Initializer.ParameterMap, TEXT("SampleWeights"));
 		LowpassWeights.Bind(Initializer.ParameterMap, TEXT("LowpassWeights"));
+		PlusWeights.Bind(Initializer.ParameterMap, TEXT("PlusWeights"));
 		CameraMotionParams.Bind(Initializer.ParameterMap);
 		RandomOffset.Bind(Initializer.ParameterMap, TEXT("RandomOffset"));
 	}
@@ -83,7 +85,7 @@ public:
 	virtual bool Serialize(FArchive& Ar)
 	{
 		bool bShaderHasOutdatedParameters = FGlobalShader::Serialize(Ar);
-		Ar << PostprocessParameter << DeferredParameters << SampleWeights << LowpassWeights << CameraMotionParams << RandomOffset;
+		Ar << PostprocessParameter << DeferredParameters << SampleWeights << LowpassWeights << PlusWeights << CameraMotionParams << RandomOffset;
 		return bShaderHasOutdatedParameters;
 	}
 
@@ -135,8 +137,10 @@ public:
 
 			float Weights[9];
 			float WeightsLow[9];
+			float WeightsPlus[5];
 			float TotalWeight = 0.0f;
 			float TotalWeightLow = 0.0f;
+			float TotalWeightPlus = 0.0f;
 			for( int32 i = 0; i < 9; i++ )
 			{
 				// Exponential fit to Blackman-Harris 3.3
@@ -157,12 +161,24 @@ public:
 				WeightsLow[i] = FMath::Exp( -2.29f * ( PixelOffsetX * PixelOffsetX + PixelOffsetY * PixelOffsetY ) );
 				TotalWeightLow += WeightsLow[i];
 			}
+
+			WeightsPlus[0] = Weights[1];
+			WeightsPlus[1] = Weights[3];
+			WeightsPlus[2] = Weights[4];
+			WeightsPlus[3] = Weights[5];
+			WeightsPlus[4] = Weights[7];
+			TotalWeightPlus = Weights[1] + Weights[3] + Weights[4] + Weights[5] + Weights[7];
 			
 			for( int32 i = 0; i < 9; i++ )
 			{
 				SetShaderValue(Context.RHICmdList, ShaderRHI, SampleWeights, Weights[i] / TotalWeight, i );
 				SetShaderValue(Context.RHICmdList, ShaderRHI, LowpassWeights, WeightsLow[i] / TotalWeightLow, i );
 			}
+
+			for( int32 i = 0; i < 5; i++ )
+			{
+				SetShaderValue(Context.RHICmdList, ShaderRHI, PlusWeights, WeightsPlus[i] / TotalWeightPlus, i );
+			}
 			
 			FVector2D RandomOffsetValue;
 			TemporalRandom(&RandomOffsetValue, Context.View.FrameNumber);
