Fix for partiicle SubUVs

#ttp 346626 - CRITICAL: Regression: CONTENT EXAMPLES: Jittery fire and smoke FX on example 3.3

diff --git a/Engine/Shaders/ParticleGPUSpriteVertexFactory.usf b/Engine/Shaders/ParticleGPUSpriteVertexFactory.usf
--- a/Engine/Shaders/ParticleGPUSpriteVertexFactory.usf
+++ b/Engine/Shaders/ParticleGPUSpriteVertexFactory.usf
@@ -504,6 +504,9 @@ FVertexFactoryInterpolantsVSToPS VertexFactoryGetInterpolantsVSToPS( FVertexFact
 	{
 		SetUV(Interpolants, CoordinateIndex, CustomizedUVs[CoordinateIndex]);
 	}
+#elif USE_PARTICLE_SUBUVS
+	SetUV(Interpolants, 0, Intermediates.TexCoord.xy);
+	SetUV(Interpolants, 1, Intermediates.TexCoord.zw);
 #endif
 
 	Interpolants.Misc.x = Intermediates.SubImageLerp;
diff --git a/Engine/Shaders/ParticleSpriteVertexFactory.usf b/Engine/Shaders/ParticleSpriteVertexFactory.usf
--- a/Engine/Shaders/ParticleSpriteVertexFactory.usf
+++ b/Engine/Shaders/ParticleSpriteVertexFactory.usf
@@ -471,6 +471,9 @@ FVertexFactoryInterpolantsVSToPS VertexFactoryGetInterpolantsVSToPS(FVertexFacto
 	{
 		SetUV(Interpolants, CoordinateIndex, CustomizedUVs[CoordinateIndex]);
 	}
+#elif LIGHTMAP_UV_ACCESS || USE_PARTICLE_SUBUVS
+	SetUV(Interpolants, 0, Intermediates.TexCoord.xy);
+	SetUV(Interpolants, 1, Intermediates.TexCoord.zw);
 #endif
 
 #if NEEDS_TANGENT_BASIS
