Visualize Mesh DistanceFields mode has proper occlusion now by tracing the scene twice, makes it much easier to understand what's going on
Visualize DistanceField AO mode now shows occlusion only instead of bent normal

diff --git a/Engine/Shaders/DistanceFieldSurfaceCacheLighting.usf b/Engine/Shaders/DistanceFieldSurfaceCacheLighting.usf
--- a/Engine/Shaders/DistanceFieldSurfaceCacheLighting.usf
+++ b/Engine/Shaders/DistanceFieldSurfaceCacheLighting.usf
@@ -570,7 +570,11 @@ void AOUpsamplePS(in float4 UVAndScreenPos : TEXCOORD0, out float4 OutColor : SV
 	float3 BentNormal = Texture2DSampleLevel(BentNormalAOTexture, BentNormalAOSampler, DistanceFieldUVs, 0).xyz;
 #endif
 
+#if OUTPUT_BENT_NORMAL
 	OutColor = float4(BentNormal, 1);
+#else
+	OutColor = float4(length(BentNormal).xxx, 1);
+#endif
 }
 
 RWTexture2D<float4> RWVisualizeMeshDistanceFields;
@@ -587,70 +591,17 @@ groupshared uint IntersectingObjectIndices[MAX_INTERSECTING_OBJECTS];
 
 groupshared uint NumIntersectingObjects;
 
-[numthreads(THREADGROUP_SIZEX, THREADGROUP_SIZEY, 1)]
-void VisualizeMeshDistanceFieldCS(
-	uint3 GroupId : SV_GroupID,
-	uint3 DispatchThreadId : SV_DispatchThreadID,
-    uint3 GroupThreadId : SV_GroupThreadID) 
+void RayTraceThroughDistanceFields(float3 WorldRayStart, float3 WorldRayEnd, float MaxRayTime, out float MinRayTime, out float TotalStepsTaken)
 {
-	uint ThreadIndex = GroupThreadId.y * THREADGROUP_SIZEX + GroupThreadId.x;
-
-	float2 ScreenUV = float2((DispatchThreadId.xy * DOWNSAMPLE_FACTOR + View.ViewRectMin.xy + .5f) * View.ViewSizeAndSceneTexelSize.zw);
-	float2 ScreenPosition = (ScreenUV.xy - View.ScreenPositionScaleBias.wz) / View.ScreenPositionScaleBias.xy;
-
-	float SceneDepth = CalcSceneDepth(ScreenUV);
-	float4 HomogeneousWorldPosition = mul(float4(ScreenPosition * SceneDepth, SceneDepth, 1), View.ScreenToWorld);
-	float3 OpaqueWorldPosition = HomogeneousWorldPosition.xyz / HomogeneousWorldPosition.w;
-
-	float TraceDistance = 10000;
-	float3 WorldRayStart = View.ViewOrigin.xyz;
-	float3 WorldRayEnd = View.ViewOrigin.xyz + normalize(OpaqueWorldPosition - View.ViewOrigin.xyz) * TraceDistance;
-	float3 WorldRayDirection = WorldRayEnd - WorldRayStart;
-	float3 UnitWorldRayDirection = normalize(WorldRayDirection);
-
-	if (ThreadIndex == 0)
-	{
-		NumIntersectingObjects = 0;
-	}
-
-	GroupMemoryBarrierWithGroupSync();
+	MinRayTime = MaxRayTime;
+	TotalStepsTaken = 0;
 
 	LOOP
-	for (uint ObjectIndex = ThreadIndex; ObjectIndex < NumObjects; ObjectIndex += THREADGROUP_TOTALSIZE)
-	{
-		float4 SphereCenterAndRadius = LoadObjectPositionAndRadius(ObjectIndex);
-
-		//@todo - make independent of current pixel
-		BRANCH
-		if (RayHitSphere(WorldRayStart, UnitWorldRayDirection, SphereCenterAndRadius.xyz, SphereCenterAndRadius.w))
-		{
-			uint ListIndex;
-			InterlockedAdd(NumIntersectingObjects, 1U, ListIndex);
-			IntersectingObjectIndices[ListIndex] = ObjectIndex; 
-		}
-	}
-
-	GroupMemoryBarrierWithGroupSync();
-
-	uint LocalNumIntersectingObjects = NumIntersectingObjects;
-
-	float3 Result = 0;
-	float MinRayTime = TraceDistance;
-	float MinDistanceToSurface = 10000;
-	float TotalStepsTaken = 0;
-
-	LOOP
-	for (uint ListObjectIndex = 0; ListObjectIndex < LocalNumIntersectingObjects; ListObjectIndex++)
+	for (uint ListObjectIndex = 0; ListObjectIndex < NumIntersectingObjects; ListObjectIndex++)
 	{
 		uint ObjectIndex = IntersectingObjectIndices[ListObjectIndex];
 		float4 SphereCenterAndRadius = LoadObjectPositionAndRadius(ObjectIndex);
 
-#define TRACE_THROUGH_WORLD 0
-#define VISUALIZE_NUM_STEPS 1
-#define VISUALIZE_MIN_DISTANCE 0
-
-#if TRACE_THROUGH_WORLD || VISUALIZE_NUM_STEPS
-
 		float3 LocalPositionExtent = LoadObjectLocalPositionExtent(ObjectIndex);
 		float4x4 WorldToVolume = LoadObjectWorldToVolume(ObjectIndex);
 		float4 UVScaleAndVolumeScale = LoadObjectUVScale(ObjectIndex);
@@ -705,59 +656,73 @@ void VisualizeMeshDistanceFieldCS(
 
 			TotalStepsTaken += StepIndex;
 		}
-	
-#elif VISUALIZE_MIN_DISTANCE
-		float3 LocalPositionExtent = LoadObjectLocalPositionExtent(ObjectIndex);
-		float4x4 WorldToVolume = LoadObjectWorldToVolume(ObjectIndex);
-		float4 UVScaleAndVolumeScale = LoadObjectUVScale(ObjectIndex);
-		float3 UVAdd = LoadObjectUVAdd(ObjectIndex);
+	}
+}
 
-		float3 SampleVolumePosition = mul(float4(OpaqueWorldPosition, 1), WorldToVolume).xyz;
-		float3 ClampedSamplePosition = clamp(SampleVolumePosition, -LocalPositionExtent, LocalPositionExtent);
-		float DistanceToClamped = length(SampleVolumePosition - ClampedSamplePosition);
+[numthreads(THREADGROUP_SIZEX, THREADGROUP_SIZEY, 1)]
+void VisualizeMeshDistanceFieldCS(
+	uint3 GroupId : SV_GroupID,
+	uint3 DispatchThreadId : SV_DispatchThreadID,
+    uint3 GroupThreadId : SV_GroupThreadID) 
+{
+	uint ThreadIndex = GroupThreadId.y * THREADGROUP_SIZEX + GroupThreadId.x;
 
-		if (DistanceToClamped < .01f)
-		{
-			float3 VolumeUV = DistanceFieldVolumePositionToUV(ClampedSamplePosition, UVScaleAndVolumeScale.xyz, UVAdd);
-			float DistanceToOccluder = Texture3DSampleLevel(DistanceFieldTexture, DistanceFieldSampler, VolumeUV, 0).x + DistanceToClamped;
-			float WorldSpaceDistanceFromSurface = DistanceToOccluder * UVScaleAndVolumeScale.w;
+	float2 ScreenUV = float2((DispatchThreadId.xy * DOWNSAMPLE_FACTOR + View.ViewRectMin.xy + .5f) * View.ViewSizeAndSceneTexelSize.zw);
+	float2 ScreenPosition = (ScreenUV.xy - View.ScreenPositionScaleBias.wz) / View.ScreenPositionScaleBias.xy;
 
-			MinDistanceToSurface = min(MinDistanceToSurface, WorldSpaceDistanceFromSurface);
-		}
-#else
-		//Result = frac(OpaqueWorldPosition / 200);
-		Result = RayHitSphere(WorldRayStart, UnitWorldRayDirection, SphereCenterAndRadius.xyz, SphereCenterAndRadius.w);
-#endif
-	}
+	float SceneDepth = CalcSceneDepth(ScreenUV);
+	float4 HomogeneousWorldPosition = mul(float4(ScreenPosition * SceneDepth, SceneDepth, 1), View.ScreenToWorld);
+	float3 OpaqueWorldPosition = HomogeneousWorldPosition.xyz / HomogeneousWorldPosition.w;
 
-#if TRACE_THROUGH_WORLD
-	float3 WorldIntersectionPosition = WorldRayStart + MinRayTime * UnitWorldRayDirection;
-	Result = frac(WorldIntersectionPosition / 200) * .25f;
+	float TraceDistance = 10000;
+	float3 WorldRayStart = View.ViewOrigin.xyz;
+	float3 WorldRayEnd = WorldRayStart + normalize(OpaqueWorldPosition - View.ViewOrigin.xyz) * TraceDistance;
+	float3 WorldRayDirection = WorldRayEnd - WorldRayStart;
+	float3 UnitWorldRayDirection = normalize(WorldRayDirection);
 
-	if (MinRayTime >= TraceDistance)
+	if (ThreadIndex == 0)
 	{
-		Result = 0;
+		NumIntersectingObjects = 0;
 	}
-#elif VISUALIZE_NUM_STEPS
 
-	Result = saturate(TotalStepsTaken / 200.0f);
+	GroupMemoryBarrierWithGroupSync();
 
-	if (MinRayTime < TraceDistance)
+	LOOP
+	for (uint ObjectIndex = ThreadIndex; ObjectIndex < NumObjects; ObjectIndex += THREADGROUP_TOTALSIZE)
 	{
-		Result += .1f;
+		float4 SphereCenterAndRadius = LoadObjectPositionAndRadius(ObjectIndex);
+
+		//@todo - make independent of current pixel
+		BRANCH
+		if (RayHitSphere(WorldRayStart, UnitWorldRayDirection, SphereCenterAndRadius.xyz, SphereCenterAndRadius.w))
+		{
+			uint ListIndex;
+			InterlockedAdd(NumIntersectingObjects, 1U, ListIndex);
+			IntersectingObjectIndices[ListIndex] = ObjectIndex; 
+		}
 	}
 
-#elif VISUALIZE_MIN_DISTANCE
+	GroupMemoryBarrierWithGroupSync();
+
+	float MinRayTime;
+	float TotalStepsTaken;
 
-	Result = saturate(MinDistanceToSurface / 300.0f + .5f);
+	// Trace once to find the distance to first intersection
+	RayTraceThroughDistanceFields(WorldRayStart, WorldRayEnd, TraceDistance, MinRayTime, TotalStepsTaken);
 
-	if (MinDistanceToSurface < 0)
+	float TempMinRayTime;
+	// Recompute the ray end point
+	WorldRayEnd = WorldRayStart + UnitWorldRayDirection * MinRayTime;
+	// Trace a second time to only accumulate steps taken before the first intersection, improves visualization
+	RayTraceThroughDistanceFields(WorldRayStart, WorldRayEnd, MinRayTime, TempMinRayTime, TotalStepsTaken);
+
+	float3 Result = saturate(TotalStepsTaken / 200.0f);
+
+	if (MinRayTime < TraceDistance)
 	{
-		Result = 0;
+		Result += .1f;
 	}
 
-#endif
-
 	RWVisualizeMeshDistanceFields[DispatchThreadId.xy] = float4(Result, 0);
 }
 
diff --git a/Engine/Source/Runtime/Renderer/Private/DistanceFieldSurfaceCacheLighting.cpp b/Engine/Source/Runtime/Renderer/Private/DistanceFieldSurfaceCacheLighting.cpp
--- a/Engine/Source/Runtime/Renderer/Private/DistanceFieldSurfaceCacheLighting.cpp
+++ b/Engine/Source/Runtime/Renderer/Private/DistanceFieldSurfaceCacheLighting.cpp
@@ -2190,9 +2190,10 @@ void PostProcessBentNormalAO(
 	}
 }
 
-class FDistanceFieldAOUpsamplePS : public FGlobalShader
+template<bool bOutputBentNormal>
+class TDistanceFieldAOUpsamplePS : public FGlobalShader
 {
-	DECLARE_SHADER_TYPE(FDistanceFieldAOUpsamplePS, Global);
+	DECLARE_SHADER_TYPE(TDistanceFieldAOUpsamplePS, Global);
 public:
 
 	static bool ShouldCache(EShaderPlatform Platform)
@@ -2203,13 +2204,14 @@ public:
 	static void ModifyCompilationEnvironment(EShaderPlatform Platform, FShaderCompilerEnvironment& OutEnvironment)
 	{
 		OutEnvironment.SetDefine(TEXT("DOWNSAMPLE_FACTOR"), GAODownsampleFactor);
+		OutEnvironment.SetDefine(TEXT("OUTPUT_BENT_NORMAL"), bOutputBentNormal ? TEXT("1") : TEXT("0"));
 	}
 
 	/** Default constructor. */
-	FDistanceFieldAOUpsamplePS() {}
+	TDistanceFieldAOUpsamplePS() {}
 
 	/** Initialization constructor. */
-	FDistanceFieldAOUpsamplePS(const ShaderMetaType::CompiledShaderInitializerType& Initializer)
+	TDistanceFieldAOUpsamplePS(const ShaderMetaType::CompiledShaderInitializerType& Initializer)
 		: FGlobalShader(Initializer)
 	{
 		DeferredParameters.Bind(Initializer.ParameterMap);
@@ -2243,9 +2245,10 @@ private:
 	FShaderResourceParameter BentNormalAOSampler;
 };
 
-IMPLEMENT_SHADER_TYPE(,FDistanceFieldAOUpsamplePS,TEXT("DistanceFieldSurfaceCacheLighting"),TEXT("AOUpsamplePS"),SF_Pixel);
+IMPLEMENT_SHADER_TYPE(template<>,TDistanceFieldAOUpsamplePS<true>,TEXT("DistanceFieldSurfaceCacheLighting"),TEXT("AOUpsamplePS"),SF_Pixel);
+IMPLEMENT_SHADER_TYPE(template<>,TDistanceFieldAOUpsamplePS<false>,TEXT("DistanceFieldSurfaceCacheLighting"),TEXT("AOUpsamplePS"),SF_Pixel);
 
-void UpsampleBentNormalAO(FRHICommandList& RHICmdList, const TArray<FViewInfo>& Views, TRefCountPtr<IPooledRenderTarget>& DistanceFieldAOBentNormal)
+void UpsampleBentNormalAO(FRHICommandList& RHICmdList, const TArray<FViewInfo>& Views, TRefCountPtr<IPooledRenderTarget>& DistanceFieldAOBentNormal, bool bOutputBentNormal)
 {
 	for (int32 ViewIndex = 0; ViewIndex < Views.Num(); ViewIndex++)
 	{
@@ -2259,13 +2262,21 @@ void UpsampleBentNormalAO(FRHICommandList& RHICmdList, const TArray<FViewInfo>&
 		RHICmdList.SetBlendState(TStaticBlendState<>::GetRHI());
 
 		TShaderMapRef<FPostProcessVS> VertexShader(View.ShaderMap);
-		TShaderMapRef<FDistanceFieldAOUpsamplePS> PixelShader(View.ShaderMap);
 
-		static FGlobalBoundShaderState BoundShaderState;
-		
-		SetGlobalBoundShaderState(RHICmdList, View.GetFeatureLevel(), BoundShaderState, GFilterVertexDeclaration.VertexDeclarationRHI, *VertexShader, *PixelShader);
-
-		PixelShader->SetParameters(RHICmdList, View, DistanceFieldAOBentNormal);
+		if (bOutputBentNormal)
+		{
+			TShaderMapRef<TDistanceFieldAOUpsamplePS<true> > PixelShader(View.ShaderMap);
+			static FGlobalBoundShaderState BoundShaderState;
+			SetGlobalBoundShaderState(RHICmdList, View.GetFeatureLevel(), BoundShaderState, GFilterVertexDeclaration.VertexDeclarationRHI, *VertexShader, *PixelShader);
+			PixelShader->SetParameters(RHICmdList, View, DistanceFieldAOBentNormal);
+		}
+		else
+		{
+			TShaderMapRef<TDistanceFieldAOUpsamplePS<false> > PixelShader(View.ShaderMap);
+			static FGlobalBoundShaderState BoundShaderState;
+			SetGlobalBoundShaderState(RHICmdList, View.GetFeatureLevel(), BoundShaderState, GFilterVertexDeclaration.VertexDeclarationRHI, *VertexShader, *PixelShader);
+			PixelShader->SetParameters(RHICmdList, View, DistanceFieldAOBentNormal);
+		}
 
 		DrawRectangle( 
 			RHICmdList,
@@ -3017,7 +3028,7 @@ bool FDeferredShadingSceneRenderer::RenderDistanceFieldAOSurfaceCache(FRHIComman
 				}
 
 				// Upsample to full resolution, write to output
-				UpsampleBentNormalAO(RHICmdList, Views, AOOutput);
+				UpsampleBentNormalAO(RHICmdList, Views, AOOutput, !bApplyToSceneColor);
 
 				return true;
 			}
