Fixed DFAO distance fading to unoccluded
DF shadowing from point lights now clamps cone angles to 10 degrees to avoid artifacts

diff --git a/Engine/Shaders/DistanceFieldShadowing.usf b/Engine/Shaders/DistanceFieldShadowing.usf
--- a/Engine/Shaders/DistanceFieldShadowing.usf
+++ b/Engine/Shaders/DistanceFieldShadowing.usf
@@ -408,7 +408,9 @@ void DistanceFieldShadowingCS(
 				float3 WorldRayStart = OpaqueWorldPosition + LightVector / LightVectorLength * RayStartOffset;
 				float3 WorldRayEnd = LightPositionAndInvRadius.xyz;
 				float MaxRayTime = LightVectorLength;
-				float TanLightAngle = LightSourceRadius / LightVectorLength;
+				float MaxAngle = tan(10 * PI / 180.0f);
+				// Comparing tangents instead of angles, but tangent is always increasing in this range
+				float TanLightAngle = min(LightSourceRadius / LightVectorLength, MaxAngle);
 
 			#else
 
diff --git a/Engine/Shaders/DistanceFieldSurfaceCacheLighting.usf b/Engine/Shaders/DistanceFieldSurfaceCacheLighting.usf
--- a/Engine/Shaders/DistanceFieldSurfaceCacheLighting.usf
+++ b/Engine/Shaders/DistanceFieldSurfaceCacheLighting.usf
@@ -417,21 +417,16 @@ void AOCombinePS2(
 		OutColor.a = -SceneDepth;
 	}
 
-	BRANCH
+	float BentNormalLength = length(OutColor.rgb);
+	// Fade to unoccluded in the distance
+	float FadeAlpha = saturate(.001f * (.95f * AOMaxViewDistance - SceneDepth));
+	OutColor.rgb = normalize(OutColor.rgb) * lerp(1, BentNormalLength, FadeAlpha);
+
+	FLATTEN
 	if (SceneDepth > AOMaxViewDistance)
 	{
-		float BentNormalLength = length(OutColor.rgb);
-		// Fade to unoccluded in the distance
-		float FadeAlpha = saturate(.001f * (.95f * AOMaxViewDistance - SceneDepth));
-		OutColor.rgb = lerp(DistanceFieldNormalValue.xyz, OutColor.rgb, FadeAlpha);
-		OutColor.rgb = normalize(OutColor.rgb) * lerp(1, BentNormalLength, FadeAlpha);
-
-		FLATTEN
-		if (SceneDepth > AOMaxViewDistance)
-		{
-			// Mark as valid for the gap filling pass, so we don't get overwritten
-			OutColor.a = abs(OutColor.a);
-		}
+		// Mark as valid for the gap filling pass, so we don't get overwritten
+		OutColor.a = abs(OutColor.a);
 	}
 
 #endif
