Fixing removal of scaling when computing tangent to world matrix.

diff --git a/Engine/Shaders/MeshParticleVertexFactory.usf b/Engine/Shaders/MeshParticleVertexFactory.usf
--- a/Engine/Shaders/MeshParticleVertexFactory.usf
+++ b/Engine/Shaders/MeshParticleVertexFactory.usf
@@ -374,7 +374,7 @@ void CalcTangentToWorld(FVertexFactoryInput Input, FVertexFactoryIntermediates I
 
 	// Normalize to remove scaling. Incorrect but faster than computing the inverse-transpose.
 	TangentToWorld0 = normalize(TangentToWorld[0]);
-	TangentToWorld2 = normalize(float4(TangentToWorld[2], TangentBias(Input.TangentZ.w) * Primitive.LocalToWorldDeterminantSign));
+	TangentToWorld2 = float4(normalize(TangentToWorld[2]), TangentBias(Input.TangentZ.w) * Primitive.LocalToWorldDeterminantSign);
 }
 
 FVertexFactoryInterpolantsVSToPS VertexFactoryGetInterpolantsVSToPS(FVertexFactoryInput Input, FVertexFactoryIntermediates Intermediates, FMaterialVertexParameters VertexParameters)
