Proper fix for sub uvs.
Separated them entirely from the uvs.

diff --git a/Engine/Shaders/ParticleGPUSpriteVertexFactory.usf b/Engine/Shaders/ParticleGPUSpriteVertexFactory.usf
--- a/Engine/Shaders/ParticleGPUSpriteVertexFactory.usf
+++ b/Engine/Shaders/ParticleGPUSpriteVertexFactory.usf
@@ -52,8 +52,10 @@ struct FVertexFactoryInterpolantsVSToPS
 {
 #if NUM_MATERIAL_TEXCOORDS
 	float4	TexCoords[(NUM_MATERIAL_TEXCOORDS + 1) / 2]	: TEXCOORD0;
-#elif USE_PARTICLE_SUBUVS
-	float4	TexCoords[1]	: TEXCOORD0;
+#endif
+
+#if USE_PARTICLE_SUBUVS
+	float4 ParticleSubUVs : PARTICLE_SUBUVS;
 #endif
 
 	/** X: SubImageLerp Y:RelativeTime Z:MotionBlurFade */
@@ -302,7 +304,7 @@ FVertexFactoryIntermediates GetVertexFactoryIntermediates( FVertexFactoryInput I
 	return Intermediates;
 }
 
-#if NUM_MATERIAL_TEXCOORDS || USE_PARTICLE_SUBUVS
+#if NUM_MATERIAL_TEXCOORDS
 float2 GetUV(FVertexFactoryInterpolantsVSToPS Interpolants, int UVIndex)
 {
 	float4 UVVector = Interpolants.TexCoords[UVIndex / 2];
@@ -379,8 +381,8 @@ FMaterialPixelParameters GetMaterialPixelParameters( FVertexFactoryInterpolantsV
 #endif
 
 #if USE_PARTICLE_SUBUVS
-	Result.Particle.SubUVCoords[0] = GetUV(Interpolants, 0);
-	Result.Particle.SubUVCoords[1] = GetUV(Interpolants, 1);
+	Result.Particle.SubUVCoords[0] = Interpolants.ParticleSubUVs.xy;
+	Result.Particle.SubUVCoords[1] = Interpolants.ParticleSubUVs.zw;
 	Result.Particle.SubUVLerp = Interpolants.Misc.xx;
 #endif
 
@@ -435,6 +437,11 @@ FMaterialVertexParameters GetMaterialVertexParameters( FVertexFactoryInput Input
 #endif	// >= 1
 #endif	//NUM_MATERIAL_TEXCOORDS
 
+#if USE_PARTICLE_SUBUVS
+		Result.Particle.SubUVCoords[0] = Intermediates.TexCoord.xy;
+		Result.Particle.SubUVCoords[1] = Intermediates.TexCoord.zw;
+#endif
+
 #if USE_PARTICLE_SIZE
 	Result.Particle.Size = Intermediates.ParticleSize;
 #endif
@@ -504,9 +511,11 @@ FVertexFactoryInterpolantsVSToPS VertexFactoryGetInterpolantsVSToPS( FVertexFact
 	{
 		SetUV(Interpolants, CoordinateIndex, CustomizedUVs[CoordinateIndex]);
 	}
-#elif USE_PARTICLE_SUBUVS
-	SetUV(Interpolants, 0, Intermediates.TexCoord.xy);
-	SetUV(Interpolants, 1, Intermediates.TexCoord.zw);
+#endif
+
+#if USE_PARTICLE_SUBUVS
+	Interpolants.ParticleSubUVs.xy = VertexParameters.Particle.SubUVCoords[0];
+	Interpolants.ParticleSubUVs.zw = VertexParameters.Particle.SubUVCoords[1];
 #endif
 
 	Interpolants.Misc.x = Intermediates.SubImageLerp;
diff --git a/Engine/Shaders/ParticleSpriteVertexFactory.usf b/Engine/Shaders/ParticleSpriteVertexFactory.usf
--- a/Engine/Shaders/ParticleSpriteVertexFactory.usf
+++ b/Engine/Shaders/ParticleSpriteVertexFactory.usf
@@ -51,8 +51,10 @@ struct FVertexFactoryInterpolantsVSToPS
 
 #if NUM_MATERIAL_TEXCOORDS
 	float4	TexCoords[(NUM_MATERIAL_TEXCOORDS + 1) / 2]	: TEXCOORD3;
-#elif LIGHTMAP_UV_ACCESS || USE_PARTICLE_SUBUVS
-	float4	TexCoords[1]	: TEXCOORD3;	
+#endif
+
+#if USE_PARTICLE_SUBUVS
+	float4 ParticleSubUVs : PARTICLE_SUBUVS;
 #endif
 
 #if USE_PARTICLE_POSITION
@@ -111,7 +113,7 @@ struct FVertexFactoryIntermediates
 	float2 ParticleSize;
 };
 
-#if NUM_MATERIAL_TEXCOORDS || LIGHTMAP_UV_ACCESS || USE_PARTICLE_SUBUVS
+#if NUM_MATERIAL_TEXCOORDS
 bool UVIndexUseZW(int UVIndex)
 {
 	#if COMPILER_GLSL_ES2
@@ -196,20 +198,21 @@ FMaterialPixelParameters GetMaterialPixelParameters(FVertexFactoryInterpolantsVS
 
 	Result.Particle.MotionBlurFade = 1.0f;
 
-#if LIGHTMAP_UV_ACCESS
-	Result.LightmapUVs = GetUV(Interpolants,0);
-#endif
 
 #if USE_PARTICLE_LIGHTING_OFFSET
 	Result.LightingPositionOffset = Interpolants.LightingPositionOffset;
 #endif
 
 #if USE_PARTICLE_SUBUVS
-	Result.Particle.SubUVCoords[0] = GetUV(Interpolants, 0);
-	Result.Particle.SubUVCoords[1] = GetUV(Interpolants, 1);
+	Result.Particle.SubUVCoords[0] = Interpolants.ParticleSubUVs.xy;
+	Result.Particle.SubUVCoords[1] = Interpolants.ParticleSubUVs.zw;
 	Result.Particle.SubUVLerp = SubImageLerp;
 #endif
 
+#if LIGHTMAP_UV_ACCESS
+	Result.LightmapUVs = Interpolants.ParticleSubUVs.xy;
+#endif
+
 #if USE_PARTICLE_SIZE
 	Result.Particle.Size = Interpolants.ParticleSize;
 #endif
@@ -254,6 +257,11 @@ FMaterialVertexParameters GetMaterialVertexParameters(FVertexFactoryInput Input,
 	#endif	// >= 1
 #endif	//NUM_MATERIAL_TEXCOORDS_VERTEX
 
+#if USE_PARTICLE_SUBUVS || LIGHTMAP_UV_ACCESS
+	Result.Particle.SubUVCoords[0] = Intermediates.TexCoord.xy;
+	Result.Particle.SubUVCoords[1] = Intermediates.TexCoord.zw;
+#endif
+
 #if USE_PARTICLE_SIZE
 	Result.Particle.Size = Intermediates.ParticleSize;
 #endif
@@ -471,9 +479,11 @@ FVertexFactoryInterpolantsVSToPS VertexFactoryGetInterpolantsVSToPS(FVertexFacto
 	{
 		SetUV(Interpolants, CoordinateIndex, CustomizedUVs[CoordinateIndex]);
 	}
-#elif LIGHTMAP_UV_ACCESS || USE_PARTICLE_SUBUVS
-	SetUV(Interpolants, 0, Intermediates.TexCoord.xy);
-	SetUV(Interpolants, 1, Intermediates.TexCoord.zw);
+#endif
+
+#if USE_PARTICLE_SUBUVS
+	Interpolants.ParticleSubUVs.xy = VertexParameters.Particle.SubUVCoords[0];
+	Interpolants.ParticleSubUVs.zw = VertexParameters.Particle.SubUVCoords[1];
 #endif
 
 #if NEEDS_TANGENT_BASIS
