Fixed some shader errors when using Lightmap uvs with sprite particles.

Not sure these are still used anymore tbh as the data is just a pass through of the intermediats texcoord.xy. In case I'm wrong I've given them their own interpolant to remove any interference with the customizable uvs or the particle sub uvs.

diff --git a/Engine/Shaders/ParticleSpriteVertexFactory.usf b/Engine/Shaders/ParticleSpriteVertexFactory.usf
--- a/Engine/Shaders/ParticleSpriteVertexFactory.usf
+++ b/Engine/Shaders/ParticleSpriteVertexFactory.usf
@@ -53,6 +53,11 @@ struct FVertexFactoryInterpolantsVSToPS
 	float4	TexCoords[(NUM_MATERIAL_TEXCOORDS + 1) / 2]	: TEXCOORD3;
 #endif
 
+	//Not sure this is actually being used now and it's awkward to slot in now we're supporting custom UVs so I'm just giving this its own interpolant.
+#if LIGHTMAP_UV_ACCESS
+	float2 LightMapUVs : LIGHTMAP_UVS;
+#endif
+
 #if USE_PARTICLE_SUBUVS
 	float4 ParticleSubUVs : PARTICLE_SUBUVS;
 #endif
@@ -210,7 +215,7 @@ FMaterialPixelParameters GetMaterialPixelParameters(FVertexFactoryInterpolantsVS
 #endif
 
 #if LIGHTMAP_UV_ACCESS
-	Result.LightmapUVs = Interpolants.ParticleSubUVs.xy;
+	Result.LightmapUVs = Interpolants.LightMapUVs;
 #endif
 
 #if USE_PARTICLE_SIZE
@@ -257,7 +262,7 @@ FMaterialVertexParameters GetMaterialVertexParameters(FVertexFactoryInput Input,
 	#endif	// >= 1
 #endif	//NUM_MATERIAL_TEXCOORDS_VERTEX
 
-#if USE_PARTICLE_SUBUVS || LIGHTMAP_UV_ACCESS
+#if USE_PARTICLE_SUBUVS
 	Result.Particle.SubUVCoords[0] = Intermediates.TexCoord.xy;
 	Result.Particle.SubUVCoords[1] = Intermediates.TexCoord.zw;
 #endif
@@ -481,6 +486,10 @@ FVertexFactoryInterpolantsVSToPS VertexFactoryGetInterpolantsVSToPS(FVertexFacto
 	}
 #endif
 
+#if LIGHTMAP_UV_ACCESS
+	Interpolants.LightMapUVs = Intermediates.TexCoord.xy;
+#endif
+
 #if USE_PARTICLE_SUBUVS
 	Interpolants.ParticleSubUVs.xy = VertexParameters.Particle.SubUVCoords[0];
 	Interpolants.ParticleSubUVs.zw = VertexParameters.Particle.SubUVCoords[1];
