Material inputs are clamped to valid ranges at the source.  Adds a few ALUs but improves cross-platform compatibility and consistency.

diff --git a/Engine/Shaders/BasePassPixelShader.usf b/Engine/Shaders/BasePassPixelShader.usf
--- a/Engine/Shaders/BasePassPixelShader.usf
+++ b/Engine/Shaders/BasePassPixelShader.usf
@@ -568,21 +568,12 @@ void Main(
 	half3 DiffuseColor = GetMaterialDiffuseColor( MaterialParameters );
 	half3 SpecularColor = GetMaterialSpecularColor( MaterialParameters );
 
-	// prevents negative colors and inconsistent results with lighting in the base pass (should be almost free on PC, on mobile it's more costly so we expect the content creator to fix it)
-	// could be optimized a bit later but it would be more fragile
-	DiffuseColor = saturate(DiffuseColor);
-	SpecularColor = saturate(SpecularColor);
 #else
 	// Store the results in local variables and reuse instead of calling the functions multiple times.
 	half3 BaseColor = GetMaterialBaseColor( MaterialParameters );
 	half  Metallic = GetMaterialMetallic( MaterialParameters );
 	half  Specular = GetMaterialSpecular( MaterialParameters );
 
-	// prevents negative colors and inconsistent results with lighting in the base pass (should be almost free on PC, on mobile it's more costly so we expect the content creator to fix it)
-	// could be optimized a bit later but it would be more fragile
-	BaseColor = saturate(BaseColor);
-	Metallic = saturate(Metallic);
-
 #endif
 
 	float MaterialAO = GetMaterialAmbientOcclusion(MaterialParameters);
diff --git a/Engine/Shaders/MaterialTemplate.usf b/Engine/Shaders/MaterialTemplate.usf
--- a/Engine/Shaders/MaterialTemplate.usf
+++ b/Engine/Shaders/MaterialTemplate.usf
@@ -999,39 +999,69 @@ half3 GetMaterialNormal(FMaterialPixelParameters Parameters)
 	return RetNormal;
 }
 
-half3 GetMaterialEmissive(FMaterialPixelParameters Parameters)
+half3 GetMaterialEmissiveRaw(FMaterialPixelParameters Parameters)
 {
 %s;
 }
 
+half3 GetMaterialEmissive(FMaterialPixelParameters Parameters)
+{
+	return max(GetMaterialEmissiveRaw(Parameters), 0.0f);
+}
+
 half3 GetMaterialEmissiveForCS(FMaterialPixelParameters Parameters)
 {
 %s;
 }
 
+half3 GetMaterialDiffuseColorRaw(FMaterialPixelParameters Parameters)
+{
+%s;
+}
+
 half3 GetMaterialDiffuseColor(FMaterialPixelParameters Parameters)
+{
+	return saturate(GetMaterialDiffuseColorRaw(Parameters));
+}
+
+half3 GetMaterialSpecularColorRaw(FMaterialPixelParameters Parameters)
 {
 %s;
 }
 
 half3 GetMaterialSpecularColor(FMaterialPixelParameters Parameters)
+{
+	return saturate(GetMaterialSpecularColorRaw(Parameters));
+}
+
+half3 GetMaterialBaseColorRaw(FMaterialPixelParameters Parameters)
 {
 %s;
 }
 
 half3 GetMaterialBaseColor(FMaterialPixelParameters Parameters)
+{
+	return saturate(GetMaterialBaseColorRaw(Parameters));
+}
+
+half GetMaterialMetallicRaw(FMaterialPixelParameters Parameters)
 {
 %s;
 }
 
 half GetMaterialMetallic(FMaterialPixelParameters Parameters)
+{
+	return saturate(GetMaterialMetallicRaw(Parameters));
+}
+
+half GetMaterialSpecularRaw(FMaterialPixelParameters Parameters)
 {
 %s;
 }
 
 half GetMaterialSpecular(FMaterialPixelParameters Parameters)
 {
-%s;
+	return saturate(GetMaterialSpecularRaw(Parameters));
 }
 
 half GetMaterialRoughnessRaw(FMaterialPixelParameters Parameters)
@@ -1053,7 +1083,7 @@ half GetMaterialRoughness(FMaterialPixelParameters Parameters)
 	// Increase value from 0.09 to 0.12 to fix missing specular lobe problem on device
 	return clamp( GetMaterialRoughnessRaw(Parameters), 0.12, 1.0 );
 #else
-	half Roughness = GetMaterialRoughnessRaw(Parameters);
+	half Roughness = saturate(GetMaterialRoughnessRaw(Parameters));
 
 	// todo: COMPILE_SHADERS_FOR_DEVELOPMENT is unfinished feature, using XBOXONE_PROFILE as workaround
 	#if COMPILE_SHADERS_FOR_DEVELOPMENT == 1 && !XBOXONE_PROFILE && !ES31_AEP_PROFILE
@@ -1156,24 +1186,46 @@ half GetMaterialTessellationMultiplier(FMaterialTessellationParameters Parameter
 }
 
 // .rgb:SubsurfaceColor, .a:SSProfileId in 0..1 range
+half4 GetMaterialSubsurfaceDataRaw(FMaterialPixelParameters Parameters)
+{
+%s;
+}
+
 half4 GetMaterialSubsurfaceData(FMaterialPixelParameters Parameters)
+{
+	half4 OutSubsurface = GetMaterialSubsurfaceDataRaw(Parameters);
+	OutSubsurface.rgb = saturate(OutSubsurface.rgb);
+	return OutSubsurface;
+}
+
+half GetMaterialClearCoatRaw(FMaterialPixelParameters Parameters)
 {
 %s;
 }
 
 half GetMaterialClearCoat(FMaterialPixelParameters Parameters)
+{
+	return saturate(GetMaterialClearCoatRaw(Parameters));
+}
+
+half GetMaterialClearCoatRoughnessRaw(FMaterialPixelParameters Parameters)
 {
 %s;
 }
 
 half GetMaterialClearCoatRoughness(FMaterialPixelParameters Parameters)
+{
+	return saturate(GetMaterialClearCoatRoughnessRaw(Parameters));
+}
+
+half GetMaterialAmbientOcclusionRaw(FMaterialPixelParameters Parameters)
 {
 %s;
 }
 
 half GetMaterialAmbientOcclusion(FMaterialPixelParameters Parameters)
 {
-%s;
+	return saturate(GetMaterialAmbientOcclusionRaw(Parameters));
 }
 
 half2 GetMaterialRefraction(FMaterialPixelParameters Parameters)
