fixed shader

diff --git a/Engine/Shaders/PostProcessVisualizeHDR.usf b/Engine/Shaders/PostProcessVisualizeHDR.usf
--- a/Engine/Shaders/PostProcessVisualizeHDR.usf
+++ b/Engine/Shaders/PostProcessVisualizeHDR.usf
@@ -130,7 +130,8 @@ void MainPS(float4 UVAndScreenPos : TEXCOORD0, out float4 OutColor : SV_Target0)
 
 	uint Bucket = (uint)(InsetUV.x * HISTOGRAM_SIZE);
 	
-	float HistogramSum = ComputeHistogramSum(PostprocessInput1);
+	Texture2D HistogramTexture = PostprocessInput1;
+	float HistogramSum = ComputeHistogramSum(HistogramTexture);
 
 	if(InsetUV.x < ComputeHistogramPositionFromLuminance(EyeAdaptationParams[0].z))
 	{
