added filtered option to the SceneTextures for PostprocessInput0..6, doesn't work on OpenGL yet

diff --git a/Engine/Shaders/MaterialTemplate.usf b/Engine/Shaders/MaterialTemplate.usf
--- a/Engine/Shaders/MaterialTemplate.usf
+++ b/Engine/Shaders/MaterialTemplate.usf
@@ -871,7 +871,7 @@ float3 GetShadingModelColor(FGBufferData GBuffer)
 }
 
 /** Applies an offset to the scene texture lookup and decodes the HDR linear space color. */
-float4 SceneTextureLookup(float2 UV, int SceneTextureIndex)
+float4 SceneTextureLookup(float2 UV, int SceneTextureIndex, bool bFiltered)
 {
 	FScreenSpaceData ScreenSpaceData = GetScreenSpaceData(UV, false);
 
@@ -892,13 +892,13 @@ float4 SceneTextureLookup(float2 UV, int SceneTextureIndex)
 	else if (SceneTextureIndex == 11)	return ScreenSpaceData.GBuffer.Roughness;
 	else if (SceneTextureIndex == 12)	return ScreenSpaceData.GBuffer.GBufferAO;
 	else if (SceneTextureIndex == 13)	return ScreenSpaceData.GBuffer.CustomDepth;
-	else if (SceneTextureIndex == 14)	return Texture2DSample(PostprocessInput0, PostprocessInput0Sampler, UV);
-	else if (SceneTextureIndex == 15)	return Texture2DSample(PostprocessInput1, PostprocessInput1Sampler, UV);
-	else if (SceneTextureIndex == 16)	return Texture2DSample(PostprocessInput2, PostprocessInput2Sampler, UV);
-	else if (SceneTextureIndex == 17)	return Texture2DSample(PostprocessInput3, PostprocessInput3Sampler, UV);
-	else if (SceneTextureIndex == 18)	return Texture2DSample(PostprocessInput4, PostprocessInput4Sampler, UV);
-	else if (SceneTextureIndex == 19)	return Texture2DSample(PostprocessInput5, PostprocessInput5Sampler, UV);
-	else if (SceneTextureIndex == 20)	return Texture2DSample(PostprocessInput6, PostprocessInput6Sampler, UV);
+	else if (SceneTextureIndex == 14)	return Texture2DSample(PostprocessInput0, bFiltered ? BilinearTextureSampler0 : PostprocessInput0Sampler, UV);
+	else if (SceneTextureIndex == 15)	return Texture2DSample(PostprocessInput1, bFiltered ? BilinearTextureSampler0 : PostprocessInput1Sampler, UV);
+	else if (SceneTextureIndex == 16)	return Texture2DSample(PostprocessInput2, bFiltered ? BilinearTextureSampler0 : PostprocessInput2Sampler, UV);
+	else if (SceneTextureIndex == 17)	return Texture2DSample(PostprocessInput3, bFiltered ? BilinearTextureSampler0 : PostprocessInput3Sampler, UV);
+	else if (SceneTextureIndex == 18)	return Texture2DSample(PostprocessInput4, bFiltered ? BilinearTextureSampler0 : PostprocessInput4Sampler, UV);
+	else if (SceneTextureIndex == 19)	return Texture2DSample(PostprocessInput5, bFiltered ? BilinearTextureSampler0 : PostprocessInput5Sampler, UV);
+	else if (SceneTextureIndex == 20)	return Texture2DSample(PostprocessInput6, bFiltered ? BilinearTextureSampler0 : PostprocessInput6Sampler, UV);
 	else if (SceneTextureIndex == 21)	return ScreenSpaceData.GBuffer.DecalMask;
 	else if (SceneTextureIndex == 22)	return float4(GetShadingModelColor(ScreenSpaceData.GBuffer), 1);
 	else if (SceneTextureIndex == 23)	return ScreenSpaceData.AmbientOcclusion;
@@ -936,19 +936,19 @@ float4 SceneTextureLookup(float2 UV, int SceneTextureIndex)
 		// PPI_CustomDepth
 		case 13: return ScreenSpaceData.GBuffer.CustomDepth;
 		// PPI_PostprocessInput0
-		case 14: return Texture2DSample(PostprocessInput0, PostprocessInput0Sampler, UV);
+		case 14: return Texture2DSample(PostprocessInput0, bFiltered ? BilinearTextureSampler0 : PostprocessInput0Sampler, UV);
 		// PPI_PostprocessInput1
-		case 15: return Texture2DSample(PostprocessInput1, PostprocessInput1Sampler, UV);
+		case 15: return Texture2DSample(PostprocessInput1, bFiltered ? BilinearTextureSampler0 : PostprocessInput1Sampler, UV);
 		// PPI_PostprocessInput2
-		case 16: return Texture2DSample(PostprocessInput2, PostprocessInput2Sampler, UV);
+		case 16: return Texture2DSample(PostprocessInput2, bFiltered ? BilinearTextureSampler0 : PostprocessInput2Sampler, UV);
 		// PPI_PostprocessInput3
-		case 17: return Texture2DSample(PostprocessInput3, PostprocessInput3Sampler, UV);
+		case 17: return Texture2DSample(PostprocessInput3, bFiltered ? BilinearTextureSampler0 : PostprocessInput3Sampler, UV);
 		// PPI_PostprocessInput4
-		case 18: return Texture2DSample(PostprocessInput4, PostprocessInput4Sampler, UV);
+		case 18: return Texture2DSample(PostprocessInput4, bFiltered ? BilinearTextureSampler0 : PostprocessInput4Sampler, UV);
 		// PPI_PostprocessInput5
-		case 19: return Texture2DSample(PostprocessInput5, PostprocessInput5Sampler, UV);
+		case 19: return Texture2DSample(PostprocessInput5, bFiltered ? BilinearTextureSampler0 : PostprocessInput5Sampler, UV);
 		// PPI_PostprocessInput6
-		case 20: return Texture2DSample(PostprocessInput6, PostprocessInput6Sampler, UV);
+		case 20: return Texture2DSample(PostprocessInput6, bFiltered ? BilinearTextureSampler0 : PostprocessInput6Sampler, UV);
 		// PPI_DecalMask
 		case 21: return ScreenSpaceData.GBuffer.DecalMask;
 		// PPI_ShadingModel
diff --git a/Engine/Source/Runtime/Engine/Classes/Materials/MaterialExpressionSceneTexture.h b/Engine/Source/Runtime/Engine/Classes/Materials/MaterialExpressionSceneTexture.h
--- a/Engine/Source/Runtime/Engine/Classes/Materials/MaterialExpressionSceneTexture.h
+++ b/Engine/Source/Runtime/Engine/Classes/Materials/MaterialExpressionSceneTexture.h
@@ -76,6 +76,10 @@ class UMaterialExpressionSceneTexture : public UMaterialExpression
 	UPROPERTY(EditAnywhere, Category=UMaterialExpressionSceneTexture, meta=(DisplayName = "Clamp UVs"))
 	bool bClampUVs;
 
+	/** Whether to use point sampled texture lookup (default) or using [bi-linear] filtered (can be slower, avoid faceted lock with distortions), some SceneTextures cannot be filtered */
+	UPROPERTY(EditAnywhere, Category=UMaterialExpressionSceneTexture, meta=(DisplayName = "Filtered"))
+	bool bFiltered;
+
 	// Begin UMaterialExpression Interface
 	virtual int32 Compile(class FMaterialCompiler* Compiler, int32 OutputIndex, int32 MultiplexIndex) override;
 	virtual void GetCaption(TArray<FString>& OutCaptions) const override;
diff --git a/Engine/Source/Runtime/Engine/Private/Materials/HLSLMaterialTranslator.h b/Engine/Source/Runtime/Engine/Private/Materials/HLSLMaterialTranslator.h
--- a/Engine/Source/Runtime/Engine/Private/Materials/HLSLMaterialTranslator.h
+++ b/Engine/Source/Runtime/Engine/Private/Materials/HLSLMaterialTranslator.h
@@ -2266,7 +2266,7 @@ protected:
 	}
 	
 	// @param SceneTextureId of type ESceneTextureId e.g. PPI_SubsurfaceColor
-	virtual int32 SceneTextureLookup(int32 UV, uint32 InSceneTextureId) override
+	virtual int32 SceneTextureLookup(int32 UV, uint32 InSceneTextureId, bool bFiltered) override
 	{
 		if (ErrorUnlessFeatureLevelSupported(ERHIFeatureLevel::SM4) == INDEX_NONE)
 		{
@@ -2288,8 +2288,8 @@ protected:
 
 		return AddCodeChunk(
 			MCT_Float4,
-			TEXT("SceneTextureLookup(%s, %d)"),
-			*TexCoordCode, (int)SceneTextureId
+			TEXT("SceneTextureLookup(%s, %d, %s)"),
+			*TexCoordCode, (int)SceneTextureId, bFiltered ? TEXT("true") : TEXT("false")
 			);
 	}
 
diff --git a/Engine/Source/Runtime/Engine/Private/Materials/MaterialExpressions.cpp b/Engine/Source/Runtime/Engine/Private/Materials/MaterialExpressions.cpp
--- a/Engine/Source/Runtime/Engine/Private/Materials/MaterialExpressions.cpp
+++ b/Engine/Source/Runtime/Engine/Private/Materials/MaterialExpressions.cpp
@@ -4740,6 +4740,8 @@ UMaterialExpressionSceneTexture::UMaterialExpressionSceneTexture(const class FPo
 
 	//by default slower but reliable results, if the shader never accesses the texels outside it can be disabled.
 	bClampUVs = true;
+	// by default faster, most lookup are read/write the same pixel so this is ralrely needed
+	bFiltered = false;
 
 	Outputs.Reset();
 	Outputs.Add(FExpressionOutput(TEXT("Color"), 1, 1, 1, 1, 1));
@@ -4769,7 +4771,7 @@ int32 UMaterialExpressionSceneTexture::Compile(class FMaterialCompiler* Compiler
 		}
 
 		// Color
-		return Compiler->SceneTextureLookup(UV, SceneTextureId);
+		return Compiler->SceneTextureLookup(UV, SceneTextureId, bFiltered);
 	}
 	else if(OutputIndex == 1 || OutputIndex == 2)
 	{
diff --git a/Engine/Source/Runtime/Engine/Private/SceneView.cpp b/Engine/Source/Runtime/Engine/Private/SceneView.cpp
--- a/Engine/Source/Runtime/Engine/Private/SceneView.cpp
+++ b/Engine/Source/Runtime/Engine/Private/SceneView.cpp
@@ -266,7 +266,6 @@ FSceneView::FSceneView(const FSceneViewInitOptions& InitOptions)
 	, bIsReflectionCapture(false)
 	, bIsLocked(false)
 	, bStaticSceneOnly(false)
-	, RenderingCompositePassContext(0)
 #if WITH_EDITOR
 	, OverrideLODViewOrigin(InitOptions.OverrideLODViewOrigin)
 	, bAllowTranslucentPrimitivesInHitProxy( true )
diff --git a/Engine/Source/Runtime/Engine/Public/MaterialCompiler.h b/Engine/Source/Runtime/Engine/Public/MaterialCompiler.h
--- a/Engine/Source/Runtime/Engine/Public/MaterialCompiler.h
+++ b/Engine/Source/Runtime/Engine/Public/MaterialCompiler.h
@@ -109,7 +109,7 @@ public:
 	virtual int32 SceneDepth(int32 Offset, int32 UV, bool bUseOffset) = 0;
 	virtual int32 SceneColor(int32 Offset, int32 UV, bool bUseOffset) = 0;
 	// @param SceneTextureId of type ESceneTextureId e.g. PPI_SubsurfaceColor
-	virtual int32 SceneTextureLookup(int32 UV, uint32 SceneTextureId) = 0;
+	virtual int32 SceneTextureLookup(int32 UV, uint32 SceneTextureId, bool bFiltered) = 0;
 	// @param SceneTextureId of type ESceneTextureId e.g. PPI_SubsurfaceColor
 	virtual int32 SceneTextureSize(uint32 SceneTextureId, bool bInvert) = 0;
 	// @param SceneTextureId of type ESceneTextureId e.g. PPI_SubsurfaceColor
@@ -262,7 +262,7 @@ public:
 	virtual	int32 PixelDepth() override { return Compiler->PixelDepth();	}
 	virtual int32 SceneDepth(int32 Offset, int32 UV, bool bUseOffset) override { return Compiler->SceneDepth(Offset, UV, bUseOffset); }
 	virtual int32 SceneColor(int32 Offset, int32 UV, bool bUseOffset) override { return Compiler->SceneColor(Offset, UV, bUseOffset); }
-	virtual int32 SceneTextureLookup(int32 UV, uint32 InSceneTextureId) override { return Compiler->SceneTextureLookup(UV, InSceneTextureId); }
+	virtual int32 SceneTextureLookup(int32 UV, uint32 InSceneTextureId, bool bFiltered) override { return Compiler->SceneTextureLookup(UV, InSceneTextureId, bFiltered); }
 	virtual int32 SceneTextureSize(uint32 InSceneTextureId, bool bInvert) override { return Compiler->SceneTextureSize(InSceneTextureId, bInvert); }
 	virtual int32 SceneTextureMax(uint32 InSceneTextureId) override { return Compiler->SceneTextureMax(InSceneTextureId); }
 	virtual int32 SceneTextureMin(uint32 InSceneTextureId) override { return Compiler->SceneTextureMin(InSceneTextureId); }
diff --git a/Engine/Source/Runtime/Engine/Public/SceneView.h b/Engine/Source/Runtime/Engine/Public/SceneView.h
--- a/Engine/Source/Runtime/Engine/Public/SceneView.h
+++ b/Engine/Source/Runtime/Engine/Public/SceneView.h
@@ -333,9 +333,6 @@ public:
 	 */
 	bool bStaticSceneOnly;
 
-	/** 0 if valid (we are rendering a screen postprocess pass )*/
-	struct FRenderingCompositePassContext* RenderingCompositePassContext; 
-
 	/** Aspect ratio constrained view rect. In the editor, when attached to a camera actor and the camera black bar showflag is enabled, the normal viewrect 
 	  * remains as the full viewport, and the black bars are just simulated by drawing black bars. This member stores the effective constrained area within the
 	  * bars.
diff --git a/Engine/Source/Runtime/Renderer/Private/PostProcess/RenderingCompositionGraph.cpp b/Engine/Source/Runtime/Renderer/Private/PostProcess/RenderingCompositionGraph.cpp
--- a/Engine/Source/Runtime/Renderer/Private/PostProcess/RenderingCompositionGraph.cpp
+++ b/Engine/Source/Runtime/Renderer/Private/PostProcess/RenderingCompositionGraph.cpp
@@ -800,14 +800,17 @@ void FPostProcessPassParameters::Set(
 
 			SetTextureParameter(RHICmdList, ShaderRHI, PostprocessInputParameter[Id], PostprocessInputParameterSampler[Id], LocalFilter, SrcTexture);
 
+			if(PostprocessInputSizeParameter[Id].IsBound() || PostProcessInputMinMaxParameter[Id].IsBound())
 			{
 				float Width = InputPooledElement->GetDesc().Extent.X;
 				float Height = InputPooledElement->GetDesc().Extent.Y;
-				FVector4 TextureSize(Width, Height, 1.0f / Width, 1.0f / Height);
+				
+				FVector2D OnePPInputPixelUVSize = FVector2D(1.0f / Width, 1.0f / Height);
+
+				FVector4 TextureSize(Width, Height, OnePPInputPixelUVSize.X, OnePPInputPixelUVSize.Y);
 				SetShaderValue(RHICmdList, ShaderRHI, PostprocessInputSizeParameter[Id], TextureSize);
-					
+
 				//We could use the main scene min max here if it weren't that we need to pull the max in by a pixel on a per input basis.
-				FVector2D OnePPInputPixelUVSize = FVector2D(1.0f / Width, 1.0f / Height);
 				FVector4 PPInputMinMax = BaseSceneTexMinMax;
 				PPInputMinMax.Z -= OnePPInputPixelUVSize.X;
 				PPInputMinMax.W -= OnePPInputPixelUVSize.Y;
diff --git a/Engine/Source/Runtime/Renderer/Private/ShaderBaseClasses.cpp b/Engine/Source/Runtime/Renderer/Private/ShaderBaseClasses.cpp
--- a/Engine/Source/Runtime/Renderer/Private/ShaderBaseClasses.cpp
+++ b/Engine/Source/Runtime/Renderer/Private/ShaderBaseClasses.cpp
@@ -228,12 +228,6 @@ void FMaterialShader::SetParameters(
 		}
 	}
 
-	// if we are in a postprocessing pass
-	if(View.RenderingCompositePassContext)
-	{
-		PostprocessParameter.Set(ShaderRHI, *View.RenderingCompositePassContext, TStaticSamplerState<>::GetRHI());
-	}
-
 	//Use of the eye adaptation texture here is experimental and potentially dangerous as it can introduce a feedback loop. May be removed.
 	if(EyeAdaptation.IsBound())
 	{
