fixed opengl shader compile

diff --git a/Engine/Shaders/PostProcessVisualizeHDR.usf b/Engine/Shaders/PostProcessVisualizeHDR.usf
--- a/Engine/Shaders/PostProcessVisualizeHDR.usf
+++ b/Engine/Shaders/PostProcessVisualizeHDR.usf
@@ -130,7 +130,10 @@ void MainPS(float4 UVAndScreenPos : TEXCOORD0, out float4 OutColor : SV_Target0)
 
 	uint Bucket = (uint)(InsetUV.x * HISTOGRAM_SIZE);
 	
-	Texture2D HistogramTexture = PostprocessInput1;
+	
+//	Texture2D HistogramTexture = InputNew1;
+#define HistogramTexture PostprocessInput1 // WAR (workaround) for HLSLCC not allowing assignment to samplers (yet)
+
 	float HistogramSum = ComputeHistogramSum(HistogramTexture);
 
 	if(InsetUV.x < ComputeHistogramPositionFromLuminance(EyeAdaptationParams[0].z))
