nicer comment

diff --git a/Engine/Shaders/SeparableSSS.usf b/Engine/Shaders/SeparableSSS.usf
--- a/Engine/Shaders/SeparableSSS.usf
+++ b/Engine/Shaders/SeparableSSS.usf
@@ -489,7 +489,6 @@ float4 SSSSBlurPS(
 		// The kernel is symtrical, we want to use that property.
 		// Half the GetKernel() calls (more expensive if done by texture)
 		// Half the weighting computations (saves 3mul per lookup sample)
-		// We don't sum up with (x+=y;x+=z;) but with (x+=(y+z);), order allows for better texture lookup grouping
 		SSSS_UNROLL
 		for (int Side = -1; Side <= 1; Side += 2)
 		{
@@ -509,7 +508,7 @@ float4 SSSSBlurPS(
 			LocalAccum += color.rgb;
 		}
 
-		// Accumulate to final value
+		// Accumulate to final value (left and right sample with the same weight)
 		colorBlurred.rgb += Kernel.rgb * LocalAccum;
 	}
 
