Fix PS4 ICE

diff --git a/Engine/Shaders/MaterialTemplate.usf b/Engine/Shaders/MaterialTemplate.usf
--- a/Engine/Shaders/MaterialTemplate.usf
+++ b/Engine/Shaders/MaterialTemplate.usf
@@ -892,13 +892,13 @@ float4 SceneTextureLookup(float2 UV, int SceneTextureIndex, bool bFiltered)
 	else if (SceneTextureIndex == 11)	return ScreenSpaceData.GBuffer.Roughness;
 	else if (SceneTextureIndex == 12)	return ScreenSpaceData.GBuffer.GBufferAO;
 	else if (SceneTextureIndex == 13)	return ScreenSpaceData.GBuffer.CustomDepth;
-	else if (SceneTextureIndex == 14)	return Texture2DSample(PostprocessInput0, bFiltered ? BilinearTextureSampler0 : PostprocessInput0Sampler, UV);
-	else if (SceneTextureIndex == 15)	return Texture2DSample(PostprocessInput1, bFiltered ? BilinearTextureSampler0 : PostprocessInput1Sampler, UV);
-	else if (SceneTextureIndex == 16)	return Texture2DSample(PostprocessInput2, bFiltered ? BilinearTextureSampler0 : PostprocessInput2Sampler, UV);
-	else if (SceneTextureIndex == 17)	return Texture2DSample(PostprocessInput3, bFiltered ? BilinearTextureSampler0 : PostprocessInput3Sampler, UV);
-	else if (SceneTextureIndex == 18)	return Texture2DSample(PostprocessInput4, bFiltered ? BilinearTextureSampler0 : PostprocessInput4Sampler, UV);
-	else if (SceneTextureIndex == 19)	return Texture2DSample(PostprocessInput5, bFiltered ? BilinearTextureSampler0 : PostprocessInput5Sampler, UV);
-	else if (SceneTextureIndex == 20)	return Texture2DSample(PostprocessInput6, bFiltered ? BilinearTextureSampler0 : PostprocessInput6Sampler, UV);
+	else if (SceneTextureIndex == 14)	if (bFiltered) { return Texture2DSample(PostprocessInput0, BilinearTextureSampler0, UV);} else { return Texture2DSample(PostprocessInput0, PostprocessInput0Sampler, UV);} 
+	else if (SceneTextureIndex == 15)	if (bFiltered) { return Texture2DSample(PostprocessInput1, BilinearTextureSampler0, UV);} else { return Texture2DSample(PostprocessInput1, PostprocessInput1Sampler, UV);} 
+	else if (SceneTextureIndex == 16)	if (bFiltered) { return Texture2DSample(PostprocessInput2, BilinearTextureSampler0, UV);} else { return Texture2DSample(PostprocessInput2, PostprocessInput2Sampler, UV);} 
+	else if (SceneTextureIndex == 17)	if (bFiltered) { return Texture2DSample(PostprocessInput3, BilinearTextureSampler0, UV);} else { return Texture2DSample(PostprocessInput3, PostprocessInput3Sampler, UV);} 
+	else if (SceneTextureIndex == 18)	if (bFiltered) { return Texture2DSample(PostprocessInput4, BilinearTextureSampler0, UV);} else { return Texture2DSample(PostprocessInput4, PostprocessInput4Sampler, UV);} 
+	else if (SceneTextureIndex == 19)	if (bFiltered) { return Texture2DSample(PostprocessInput5, BilinearTextureSampler0, UV);} else { return Texture2DSample(PostprocessInput5, PostprocessInput5Sampler, UV);} 
+	else if (SceneTextureIndex == 20)	if (bFiltered) { return Texture2DSample(PostprocessInput6, BilinearTextureSampler0, UV);} else { return Texture2DSample(PostprocessInput6, PostprocessInput6Sampler, UV);} 
 	else if (SceneTextureIndex == 21)	return ScreenSpaceData.GBuffer.DecalMask;
 	else if (SceneTextureIndex == 22)	return float4(GetShadingModelColor(ScreenSpaceData.GBuffer), 1);
 	else if (SceneTextureIndex == 23)	return ScreenSpaceData.AmbientOcclusion;
