Fixed rare GPU hang from DFAO - DrawIndirect was done with uninitialized memory in very specific circumstances

diff --git a/Engine/Shaders/DistanceFieldSurfaceCacheLightingCompute.usf b/Engine/Shaders/DistanceFieldSurfaceCacheLightingCompute.usf
--- a/Engine/Shaders/DistanceFieldSurfaceCacheLightingCompute.usf
+++ b/Engine/Shaders/DistanceFieldSurfaceCacheLightingCompute.usf
@@ -781,7 +781,8 @@ void SetupCopyIndirectArgumentsCS(
 {
 	uint NumRecords = (1 - TrimFraction) * DrawParameters[1];
 
-	RWDispatchParameters[0] = (NumRecords + COPY_THREADGROUP_SIZE - 1) / COPY_THREADGROUP_SIZE;
+	// Make sure we spawn at least one group so RWScatterDrawParameters gets written to
+	RWDispatchParameters[0] = max((NumRecords + COPY_THREADGROUP_SIZE - 1) / COPY_THREADGROUP_SIZE, 1);
 	RWDispatchParameters[1] = 1;
 	RWDispatchParameters[2] = 1;
 }
diff --git a/Engine/Source/Runtime/Renderer/Private/DistanceFieldSurfaceCacheLighting.cpp b/Engine/Source/Runtime/Renderer/Private/DistanceFieldSurfaceCacheLighting.cpp
--- a/Engine/Source/Runtime/Renderer/Private/DistanceFieldSurfaceCacheLighting.cpp
+++ b/Engine/Source/Runtime/Renderer/Private/DistanceFieldSurfaceCacheLighting.cpp
@@ -3353,6 +3353,8 @@ bool FDeferredShadingSceneRenderer::RenderDistanceFieldAOSurfaceCache(FRHIComman
 					RHICmdList.ClearUAV(SurfaceCacheResources.Level[DepthLevel]->ScatterDrawParameters.UAV, ClearValues);
 				}
 
+				RHICmdList.ClearUAV(SurfaceCacheResources.TempResources->ScatterDrawParameters.UAV, ClearValues);
+
 				SurfaceCacheResources.bClearedResources = true;
 			}
 
