UE4 - Bump up shader version

diff --git a/Engine/Shaders/ShaderVersion.usf b/Engine/Shaders/ShaderVersion.usf
--- a/Engine/Shaders/ShaderVersion.usf
+++ b/Engine/Shaders/ShaderVersion.usf
@@ -2,4 +2,4 @@
 // Each time the console command is executed it generates a new GUID. As this file is included
 // in common.usf (which should be included in any shader) it allows to invalidate the shader DDC.
 // 
-// GUID = 98C7AD7080954AAFA80A0922AE68EC87
\ No newline at end of file
+// GUID = 98C7AD7080954AAFA80A0922AE68EC87a
\ No newline at end of file
