UE4 - Fix some usf comment names for easier searching

diff --git a/Engine/Shaders/ParticleGPUSpriteVertexFactory.usf b/Engine/Shaders/ParticleGPUSpriteVertexFactory.usf
--- a/Engine/Shaders/ParticleGPUSpriteVertexFactory.usf
+++ b/Engine/Shaders/ParticleGPUSpriteVertexFactory.usf
@@ -1,7 +1,7 @@
 // Copyright 1998-2014 Epic Games, Inc. All Rights Reserved.
 
 /*=============================================================================
-	ParticleGPUVertexFactory.usf: Vertex factory for GPU simulated particles.
+	ParticleGPUSpriteVertexFactory.usf: Vertex factory for GPU simulated particles.
 =============================================================================*/
 
 #define PARTICLE_SPRITE_FACTORY 1
diff --git a/Engine/Shaders/VectorFieldVisualizationVertexFactory.usf b/Engine/Shaders/VectorFieldVisualizationVertexFactory.usf
--- a/Engine/Shaders/VectorFieldVisualizationVertexFactory.usf
+++ b/Engine/Shaders/VectorFieldVisualizationVertexFactory.usf
@@ -1,7 +1,7 @@
 // Copyright 1998-2014 Epic Games, Inc. All Rights Reserved.
 
 /*=============================================================================
-	ParticleGPUVertexFactory.usf: Vertex factory for GPU simulated particles.
+	VectorFieldVisualizationVertexFactory.usf: Vertex factory for GPU simulated particles.
 =============================================================================*/
 
 #define PARTICLE_SPRITE_FACTORY 1
