Moved reusable shader function into common spot

diff --git a/Engine/Shaders/PostProcessCommon.usf b/Engine/Shaders/PostProcessCommon.usf
--- a/Engine/Shaders/PostProcessCommon.usf
+++ b/Engine/Shaders/PostProcessCommon.usf
@@ -152,4 +152,30 @@ float2 VignetteSpace(float2 Pos)
 }
 
 
+// gets 4 nearby SceneDepth values for one UV value, useful for depth downsample, uses Gather() if possible
+float4 GatherSceneDepth(float2 UV, float2 InvBufferSize)
+{
+#if FEATURE_LEVEL >= FEATURE_LEVEL_SM5
+	// using Gather: xyzw in counter clockwise order starting with the sample to the lower left of the queried location
+	float4 DeviceZ = SceneDepthTextureNonMS.Gather(BilinearTextureSampler0, UV);
+
+	return float4(
+		ConvertFromDeviceZ(DeviceZ.x),
+		ConvertFromDeviceZ(DeviceZ.y),
+		ConvertFromDeviceZ(DeviceZ.z),
+		ConvertFromDeviceZ(DeviceZ.w)
+		);
+#elif FEATURE_LEVEL == FEATURE_LEVEL_SM4
+	float2 TexelScale = 0.5f * InvBufferSize;
+	// @todo Optimize
+	// Manually sample neighbouring texels in counter clockwise order starting with the sample to the lower left of the queried location.
+	return float4(
+		CalcSceneDepth(UV + (float2(-1, 1) * TexelScale)),
+		CalcSceneDepth(UV + (float2( 1, 1) * TexelScale)),
+		CalcSceneDepth(UV + (float2( 1,-1) * TexelScale)),
+		CalcSceneDepth(UV + (float2(-1,-1) * TexelScale))
+		); 
+#endif // FEATURE_LEVEL
+}
+
 #endif // __POST_PROCESS_COMMON__
diff --git a/Engine/Shaders/PostProcessDOF.usf b/Engine/Shaders/PostProcessDOF.usf
--- a/Engine/Shaders/PostProcessDOF.usf
+++ b/Engine/Shaders/PostProcessDOF.usf
@@ -9,32 +9,6 @@
 #include "DeferredShadingCommon.usf"		// FGBufferData
 #include "DepthOfFieldCommon.usf"
 
-float4 GetSceneDepthQuad(float2 UV)
-{
-#if FEATURE_LEVEL >= FEATURE_LEVEL_SM5
-	// using Gather: xyzw in counter clockwise order starting with the sample to the lower left of the queried location
-	float4 DeviceZ = SceneDepthTextureNonMS.Gather(BilinearTextureSampler0, UV);
-
-	return float4(
-		ConvertFromDeviceZ(DeviceZ.x),
-		ConvertFromDeviceZ(DeviceZ.y),
-		ConvertFromDeviceZ(DeviceZ.z),
-		ConvertFromDeviceZ(DeviceZ.w)
-		);
-#elif FEATURE_LEVEL == FEATURE_LEVEL_SM4
-	float2 TexelScale = 0.5f * PostprocessInput1Size.zw;
-	// @todo Optimize
-	// Manually sample neighbouring texels in counter clockwise order starting with the sample to the lower left of the queried location.
-	return float4(
-		CalcSceneDepth(UV + (float2(-1, 1) * TexelScale)),
-		CalcSceneDepth(UV + (float2( 1, 1) * TexelScale)),
-		CalcSceneDepth(UV + (float2( 1,-1) * TexelScale)),
-		CalcSceneDepth(UV + (float2(-1,-1) * TexelScale))
-		); 
-#endif // FEATURE_LEVEL
-}
-
-
 // todo move to central place
 float ComputeDOFNearFocalMask(float SceneDepth)
 {
@@ -83,7 +57,7 @@ void SetupPS(
 
 	float MaskDistance = View.DepthOfFieldFocalDistance + View.DepthOfFieldFocalRegion * 0.5f;
 
-	float4 DepthQuad = GetSceneDepthQuad(UV);
+	float4 DepthQuad = GatherSceneDepth(UV, PostprocessInput1Size.zw);
 
 #if ENABLE_NEAR_BLUR == 0
 	// We aren't writing out to the second render target, so we'll just make a dummy value here which 
