UE4 - CIS/Build breakage fix

diff --git a/Engine/Shaders/PostProcessCommon.usf b/Engine/Shaders/PostProcessCommon.usf
--- a/Engine/Shaders/PostProcessCommon.usf
+++ b/Engine/Shaders/PostProcessCommon.usf
@@ -175,6 +175,8 @@ float4 GatherSceneDepth(float2 UV, float2 InvBufferSize)
 		CalcSceneDepth(UV + (float2( 1,-1) * TexelScale)),
 		CalcSceneDepth(UV + (float2(-1,-1) * TexelScale))
 		); 
+#else
+	return (float4)CalcSceneDepth(UV));
 #endif // FEATURE_LEVEL
 }
 
