UE4 - Fix the fix!

diff --git a/Engine/Shaders/PostProcessCommon.usf b/Engine/Shaders/PostProcessCommon.usf
--- a/Engine/Shaders/PostProcessCommon.usf
+++ b/Engine/Shaders/PostProcessCommon.usf
@@ -176,7 +176,7 @@ float4 GatherSceneDepth(float2 UV, float2 InvBufferSize)
 		CalcSceneDepth(UV + (float2(-1,-1) * TexelScale))
 		); 
 #else
-	return (float4)CalcSceneDepth(UV));
+	return (float4)CalcSceneDepth(UV);
 #endif // FEATURE_LEVEL
 }
 
