Disabled DFAO interpolation depth testing due to artifacts
Fixed skylight crash when building lighting with no specified cubemap
#PR 491

diff --git a/Engine/Shaders/DistanceFieldSurfaceCacheLighting.usf b/Engine/Shaders/DistanceFieldSurfaceCacheLighting.usf
--- a/Engine/Shaders/DistanceFieldSurfaceCacheLighting.usf
+++ b/Engine/Shaders/DistanceFieldSurfaceCacheLighting.usf
@@ -297,7 +297,8 @@ void IrradianceCacheSplatVS(
 
 	// +1 = behind volume (safe from near plane clipping without depth testing), 
 	// -1 = in front of volume (required for conservative results with depth testing)
-	float BoundingDirection = -1;
+	//@todo - using -1 currently causes artifacts
+	float BoundingDirection = +1;
 	float OffsetFromCenter = BoundingDirection * PositionAndRadius.w;
 	// Distance from the camera position that NormalizedViewToCenter will intersect the near plane
 	float NearPlaneDistance = View.NearPlane / dot(View.ViewForward, NormalizedViewToCenter);
diff --git a/Engine/Source/Runtime/Engine/Private/SkyLightComponent.cpp b/Engine/Source/Runtime/Engine/Private/SkyLightComponent.cpp
--- a/Engine/Source/Runtime/Engine/Private/SkyLightComponent.cpp
+++ b/Engine/Source/Runtime/Engine/Private/SkyLightComponent.cpp
@@ -420,7 +420,7 @@ void USkyLightComponent::CaptureEmissiveIrradianceEnvironmentMap(FSHVectorRGB3&
 {
 	OutIrradianceMap = FSHVectorRGB3();
 
-	if (GetScene())
+	if (GetScene() && (SourceType != SLS_SpecifiedCubemap || Cubemap))
 	{
 		// Capture emissive scene lighting only for the lighting build
 		// This is necessary to avoid a feedback loop with the last lighting build results
diff --git a/Engine/Source/Runtime/Renderer/Private/DistanceFieldSurfaceCacheLighting.cpp b/Engine/Source/Runtime/Renderer/Private/DistanceFieldSurfaceCacheLighting.cpp
--- a/Engine/Source/Runtime/Renderer/Private/DistanceFieldSurfaceCacheLighting.cpp
+++ b/Engine/Source/Runtime/Renderer/Private/DistanceFieldSurfaceCacheLighting.cpp
@@ -203,7 +203,7 @@ FAutoConsoleVariableRef CVarAOInterpolationStencilTesting(
 	ECVF_Cheat | ECVF_RenderThreadSafe
 	);
 
-int32 GAOInterpolationDepthTesting = 1;
+int32 GAOInterpolationDepthTesting = 0;
 FAutoConsoleVariableRef CVarAOInterpolationDepthTesting(
 	TEXT("r.AOInterpolationDepthTesting"),
 	GAOInterpolationDepthTesting,
