fixed shader compiler error

diff --git a/Engine/Shaders/SeparableSSS.usf b/Engine/Shaders/SeparableSSS.usf
--- a/Engine/Shaders/SeparableSSS.usf
+++ b/Engine/Shaders/SeparableSSS.usf
@@ -125,7 +125,7 @@
 #define SSSSSampleLevelZeroPoint(tex, coord) tex.SampleLevel(PointSampler, coord, 0)
 #define SSSSSampleShadowmap(coord)	float4(0,0,0,0)	// todo: not used yet
 #define SSSSSampleMasking(texcoord)	(GetScreenSpaceData(texcoord).GBuffer.ShadingModelID == SHADINGMODELID_SUBSURFACE_PROFILE)	// todo: masking doesn't work with filtering ideally put in alpha to get it filtered
-#define SSSSSampleSceneColor(coord) Texture2DSample(PostprocessInput0, PostprocessInput0Sampler, coord)
+#define SSSSSampleSceneColor(coord) Texture2DSampleLevel(PostprocessInput0, PostprocessInput0Sampler, coord, 0)
 #define SSSSSampleSceneColorPoint(coord) ( Texture2DSample(PostprocessInput0, PostprocessInput0Sampler, coord))		// todo: could be made point
 #define SSSSSampleLevelZeroOffset(tex, coord, offset) tex.SampleLevel(LinearSampler, coord, 0, offset)
 #define SSSSSampleOffset(tex, coord, offset) SSSSSampleLevelZeroOffset(tex, coord, offset)
