fixed SubsurfaceProfile look

diff --git a/Engine/Shaders/PostProcessSubsurface.usf b/Engine/Shaders/PostProcessSubsurface.usf
--- a/Engine/Shaders/PostProcessSubsurface.usf
+++ b/Engine/Shaders/PostProcessSubsurface.usf
@@ -346,10 +346,10 @@ void SubsurfaceRecombinePS(float4 UVAndScreenPos : TEXCOORD0, out float4 OutColo
 		SSSColor = SSSColorWithAlpha.rgb / max(SSSColorWithAlpha.a, 0.00001f);
 
 		// 0..255, which SubSurface profile to pick
-		uint SubsurfaceProfileInt = ExtractSubsurfaceProfileInt(ScreenSpaceData.GBuffer);
+//		uint SubsurfaceProfileInt = ExtractSubsurfaceProfileInt(ScreenSpaceData.GBuffer);
 
 		// add center sample from full resolution input for better quality
-		SSSColor += SceneColor4.rgb * GetKernel(SSSS_N_KERNELWEIGHTOFFSET, SubsurfaceProfileInt).rgb;
+//		SSSColor += SceneColor4.rgb * GetKernel(SSSS_N_KERNELWEIGHTOFFSET, SubsurfaceProfileInt).rgb;
 	}
 
 
diff --git a/Engine/Shaders/SeparableSSS.usf b/Engine/Shaders/SeparableSSS.usf
--- a/Engine/Shaders/SeparableSSS.usf
+++ b/Engine/Shaders/SeparableSSS.usf
@@ -240,9 +240,9 @@ float4 SSSSBlurPS(float2 texcoord, float2 dir, bool initStencil)
 	// center sample is added in Recombine pass
 	colorBlurred.a = colorM.a;
 
-#if SSS_DIRECTION == 0
+//#if SSS_DIRECTION == 0
 	colorBlurred.rgb = colorM * GetKernel(SSSS_N_KERNELWEIGHTOFFSET, SubsurfaceProfileInt).rgb;
-#endif
+//#endif
 
 	// Accumulate the other samples:
     SSSS_UNROLL
