fixed minor artifact in SubsurfaceProfile

diff --git a/Engine/Shaders/SeparableSSS.usf b/Engine/Shaders/SeparableSSS.usf
--- a/Engine/Shaders/SeparableSSS.usf
+++ b/Engine/Shaders/SeparableSSS.usf
@@ -108,16 +108,6 @@
 #define SSSS_FOLLOW_SURFACE 0
 #endif
 
-/**
- * This define allows to specify a different source for the SSS strength
- * (instead of using the alpha channel of the color framebuffer). This is
- * useful when the alpha channel of the mian color buffer is used for something
- * else.
- */
-#ifndef SSSS_STREGTH_SOURCE
-#define SSSS_STREGTH_SOURCE (colorM.a)
-#endif
-
 //-----------------------------------------------------------------------------
 // Porting Functions
 #define SSSSTexture2D Texture2D
@@ -214,8 +204,6 @@ float4 SSSSBlurPS(float2 texcoord, float2 dir, bool initStencil)
 		return 0.0f;
 	}
 
-	float SSSStrength = SSSS_STREGTH_SOURCE(texcoord);
-
     // Initialize the stencil buffer in case it was not already available:
     if (initStencil) // (Checked in compile time, it's optimized away)
         if (SSSStrength == 0.0) discard;
@@ -227,11 +215,15 @@ float4 SSSSBlurPS(float2 texcoord, float2 dir, bool initStencil)
     // Calculate the final step to fetch the surrounding pixels:
     float2 finalStep = scale * dir;
 
-	finalStep *= SSSStrength; // Modulate it using the alpha channel (0..1 range)
+	float SSSStrength = GetSubsurfaceStrength(texcoord);
+
+	// ideally this comes from a half res buffer as well - there are some minor artifacts
+	finalStep *= SSSStrength; // Modulate it using the opacity (0..1 range)
     
 	FGBufferData GBufferData = GetGBufferData(texcoord);
 
 	// 0..255, which SubSurface profile to pick
+	// ideally this comes from a half res buffer as well - there are some minor artifacts
 	uint SubsurfaceProfileInt = ExtractSubsurfaceProfileInt(GBufferData);
 
     // Accumulate the center sample:
