fixed shader compile error

diff --git a/Engine/Shaders/SeparableSSS.usf b/Engine/Shaders/SeparableSSS.usf
--- a/Engine/Shaders/SeparableSSS.usf
+++ b/Engine/Shaders/SeparableSSS.usf
@@ -204,7 +204,9 @@ float4 SSSSBlurPS(float2 texcoord, float2 dir, bool initStencil)
 		return 0.0f;
 	}
 
-    // Initialize the stencil buffer in case it was not already available:
+	float SSSStrength = GetSubsurfaceStrength(texcoord);
+
+	// Initialize the stencil buffer in case it was not already available:
     if (initStencil) // (Checked in compile time, it's optimized away)
         if (SSSStrength == 0.0) discard;
 
@@ -215,8 +217,6 @@ float4 SSSSBlurPS(float2 texcoord, float2 dir, bool initStencil)
     // Calculate the final step to fetch the surrounding pixels:
     float2 finalStep = scale * dir;
 
-	float SSSStrength = GetSubsurfaceStrength(texcoord);
-
 	// ideally this comes from a half res buffer as well - there are some minor artifacts
 	finalStep *= SSSStrength; // Modulate it using the opacity (0..1 range)
     
