Fixed translucent subsurface lighting

diff --git a/Engine/Shaders/BasePassPixelShader.usf b/Engine/Shaders/BasePassPixelShader.usf
--- a/Engine/Shaders/BasePassPixelShader.usf
+++ b/Engine/Shaders/BasePassPixelShader.usf
@@ -217,7 +217,7 @@ float3 GetTranslucencyLighting(FMaterialPixelParameters MaterialParameters, floa
 
 				float InScatterPower = GetMaterialTranslucentBackscatteringExponent();
 				// Setup a pow lobe to approximate anisotropic in-scattering near to the light direction
-				float InScattering = pow(saturate(dot(View.DirectionalLightDirection, MaterialParameters.CameraVector)), InScatterPower);
+				float InScattering = pow(saturate(dot(DirectionalLightDirection, MaterialParameters.CameraVector)), InScatterPower);
 
 				float4 SSData = GetMaterialSubsurfaceData(MaterialParameters);
 				float3 SubsurfaceColor = SSData.rgb;
@@ -225,7 +225,7 @@ float3 GetTranslucencyLighting(FMaterialPixelParameters MaterialParameters, floa
 				BackscatteredLighting = 
 					SubsurfaceColor
 					* InScattering 
-					* View.DirectionalLightColor.rgb
+					* DirectionalLightColor.rgb
 					// Energy normalization, tighter lobes should be brighter
 					* (InScatterPower + 2.0f) / 8.0f 
 					// Mask by shadowing, exaggerated
