optimized histogram shader (2 passes), especially for AMD and high resolutions

diff --git a/Engine/Shaders/PostProcessHistogram.usf b/Engine/Shaders/PostProcessHistogram.usf
--- a/Engine/Shaders/PostProcessHistogram.usf
+++ b/Engine/Shaders/PostProcessHistogram.usf
@@ -16,7 +16,7 @@ RWTexture2D<float4> HistogramRWTexture;
 
 //
 uint2 ThreadGroupCount;
-
+//
 uint2 LeftTopOffset;
 
 // THREADGROUP_SIZEX*THREADGROUP_SIZEY histograms of the size HISTOGRAM_SIZE
diff --git a/Engine/Shaders/PostProcessVisualizeHDR.usf b/Engine/Shaders/PostProcessVisualizeHDR.usf
--- a/Engine/Shaders/PostProcessVisualizeHDR.usf
+++ b/Engine/Shaders/PostProcessVisualizeHDR.usf
@@ -94,6 +94,21 @@ bool InUnitBox(float2 UV)
 	return UV.x >= 0 && UV.y >= 0 && UV.y < 1 && UV.y < 1;
 }
 
+
+// @param x 0=cold..1=hot
+float3 Colorize(float x)
+{
+	x = saturate(x);
+
+	float3 Heat = float3(1.0f, 0.0f, 0.0f);
+	float3 Middle = float3(0.0f, 1.0f, 0.0f);
+	float3 Cold = float3(0.0f, 0.0f, 1.0f);
+
+	float3 ColdHeat = lerp(Cold, Heat, x);
+
+	return lerp(Middle, ColdHeat, abs(0.5f - x) * 2);
+}
+
 //
 void MainPS(float4 UVAndScreenPos : TEXCOORD0, out float4 OutColor : SV_Target0)
 {
@@ -140,7 +155,7 @@ void MainPS(float4 UVAndScreenPos : TEXCOORD0, out float4 OutColor : SV_Target0)
 //		float HistogramMask = HistogramY < Value;
 		float HistogramMask = saturate((Value - HistogramY) * HistogramPixelHeight);
 
-		OutColor = lerp(OutColor, 1, HistogramMask);
+		OutColor = lerp(OutColor, float4(Colorize(HistogramUV.x), 1.0f), HistogramMask);
 	}
 
 
@@ -152,7 +167,13 @@ void MainPS(float4 UVAndScreenPos : TEXCOORD0, out float4 OutColor : SV_Target0)
 	const int2 HistogramLeftTop = int2(64, ViewportRect.w - 128 - 32);
 	const int2 HistogramSize = int2(ViewportRect.z - ViewportRect.x - 64 * 2, 128);
 	const int HistogramOuterBorder = 4;
-	const float3 BorderColor = float3(0.5f, 0.5f, 0.5f);
+
+	// (0, 0) .. (1, 1)
+	float2 InsetPx = PixelPos - HistogramLeftTop;
+	float2 InsetUV = InsetPx / HistogramSize;
+	
+//	const float3 BorderColor = float3(0.5f, 0.5f, 0.5f);
+	const float3 BorderColor = Colorize(InsetUV);
 
 	float BorderDistance = ComputeDistanceToRect(PixelPos, HistogramLeftTop, HistogramSize);
 	
@@ -173,10 +194,6 @@ void MainPS(float4 UVAndScreenPos : TEXCOORD0, out float4 OutColor : SV_Target0)
 
 	// inside the histogram
 
-	// (0, 0) .. (1, 1)
-	float2 InsetPx = PixelPos - HistogramLeftTop;
-	float2 InsetUV = InsetPx / HistogramSize;
-
 	uint Bucket = (uint)(InsetUV.x * HISTOGRAM_SIZE);
 	
 	
@@ -205,7 +222,7 @@ void MainPS(float4 UVAndScreenPos : TEXCOORD0, out float4 OutColor : SV_Target0)
 	if(LocalHistogramValue >= 1 - InsetUV.y)
 	{
 		// histogram bars
-		OutColor.xyz = lerp(OutColor.xyz, float3(0.5f, 0.5f, 0.5f), 0.5f);
+		OutColor.xyz = lerp(OutColor.xyz, Colorize(InsetUV), 0.5f);
 	}
 
 	{
diff --git a/Engine/Source/Runtime/Renderer/Private/PostProcess/PostProcessHistogram.h b/Engine/Source/Runtime/Renderer/Private/PostProcess/PostProcessHistogram.h
--- a/Engine/Source/Runtime/Renderer/Private/PostProcess/PostProcessHistogram.h
+++ b/Engine/Source/Runtime/Renderer/Private/PostProcess/PostProcessHistogram.h
@@ -19,8 +19,10 @@ public:
 
 	// -------------------------------------------
 
-	static const uint32 ThreadGroupSizeX = 4;
-	static const uint32 ThreadGroupSizeY = 4;
+	// changing this number require Histogram.usf to be recompiled
+	static const uint32 ThreadGroupSizeX = 8;
+	// changing this number require Histogram.usf to be recompiled
+	static const uint32 ThreadGroupSizeY = 8;
 
 	static const uint32 HistogramSize = 64;
 
@@ -28,7 +30,10 @@ public:
 	static const uint32 HistogramTexelCount = HistogramSize / 4;
 
 	// one ThreadGroup processes LoopCountX*LoopCountY blocks of size ThreadGroupSizeX*ThreadGroupSizeY
+
+	// changing this number require Histogram.usf to be recompiled
 	static const uint32 LoopCountX = 8;
+	// changing this number require Histogram.usf to be recompiled
 	static const uint32 LoopCountY = 8;
 
 	// -------------------------------------------
