Fixed local vertex factory compile error when subUV is used along with material UVs

diff --git a/Engine/Shaders/LocalVertexFactory.usf b/Engine/Shaders/LocalVertexFactory.usf
--- a/Engine/Shaders/LocalVertexFactory.usf
+++ b/Engine/Shaders/LocalVertexFactory.usf
@@ -499,7 +499,7 @@ FVertexFactoryInterpolantsVSToPS VertexFactoryGetInterpolantsVSToPS(FVertexFacto
 		SetUV(Interpolants, CoordinateIndex, CustomizedUVs[CoordinateIndex]);
 	}
 
-#elif USE_PARTICLE_SUBUVS
+#elif NUM_MATERIAL_TEXCOORDS_VERTEX == 0 && USE_PARTICLE_SUBUVS
 	SetUV(Interpolants, 0, Input.TexCoords[0]);
 #endif	// NUM_MATERIAL_TEXCOORDS
 
