SM5 platforms that support independent textures and samplers now use 1 less sampler for base pass features, and another 2 less samplers with DBuffer decals
* Materials created to use the max samplers of D3D11 will simply fail to compile on OpenGL 4, so the project gets to choose if they care about GL 4

diff --git a/Engine/Shaders/Common.usf b/Engine/Shaders/Common.usf
--- a/Engine/Shaders/Common.usf
+++ b/Engine/Shaders/Common.usf
@@ -150,6 +150,9 @@
 #define FCOLOR_COMPONENT_SWIZZLE
 #endif
 
+// Whether the platform supports independent texture and samplers
+// When enabled, different texture lookups can share samplers to allow more artist samplers in the base pass
+#define SUPPORTS_INDEPENDENT_SAMPLERS (FEATURE_LEVEL_SM5 && !GL4_PROFILE)
 
 // This would need to be a #define in GLSL to ignore the SamplerState, however, it is currently a function call in HLSL
 // for type checking of the parameters - ironically the type checking is really only needed in GLSL!
diff --git a/Engine/Shaders/DeferredShadingCommon.usf b/Engine/Shaders/DeferredShadingCommon.usf
--- a/Engine/Shaders/DeferredShadingCommon.usf
+++ b/Engine/Shaders/DeferredShadingCommon.usf
@@ -608,9 +608,16 @@ uint GetShadingModelId(float2 UV)
 // @param UV - UV space in the DBuffer textures (BufferSize resolution)
 FDBufferData GetDBufferData(float2 UV)
 {
+#if SUPPORTS_INDEPENDENT_SAMPLERS
+	// Share samplers to give artists more
+	float4 DBufferA = Texture2DSampleLevel(DBufferATexture, DBufferATextureSampler, UV, 0);
+	float4 DBufferB = Texture2DSampleLevel(DBufferBTexture, DBufferATextureSampler, UV, 0);
+	float2 DBufferC = Texture2DSampleLevel(DBufferCTexture, DBufferATextureSampler, UV, 0).rg;
+#else
 	float4 DBufferA = Texture2DSampleLevel(DBufferATexture, DBufferATextureSampler, UV, 0);
 	float4 DBufferB = Texture2DSampleLevel(DBufferBTexture, DBufferBTextureSampler, UV, 0);
 	float2 DBufferC = Texture2DSampleLevel(DBufferCTexture, DBufferCTextureSampler, UV, 0).rg;
+#endif
 	
 	return DecodeDBufferData(DBufferA, DBufferB, DBufferC);
 }
diff --git a/Engine/Shaders/LightmapCommon.usf b/Engine/Shaders/LightmapCommon.usf
--- a/Engine/Shaders/LightmapCommon.usf
+++ b/Engine/Shaders/LightmapCommon.usf
@@ -83,7 +83,12 @@ half4 GetLightMapColorHQ( float2 LightmapUV0, float2 LightmapUV1, half3 WorldNor
 
 float4 GetSkyBentNormalAndOcclusion(float2 LightmapUV)
 {
+#if SUPPORTS_INDEPENDENT_SAMPLERS
+	// Share sampler with lightmaps to give artists more available
+	float4 TextureValue = Texture2DSample(SkyOcclusionTexture, LightMapSampler, LightmapUV.xy);
+#else
 	float4 TextureValue = Texture2DSample(SkyOcclusionTexture, SkyOcclusionSampler, LightmapUV.xy);
+#endif
 	// Unpack vector
 	TextureValue.rgb = TextureValue.rgb * 2 - 1;
 	// Undo sqrt which allocated more precision toward 0
