Fixed SUPPORTS_INDEPENDENT_SAMPLERS condition

diff --git a/Engine/Shaders/Common.usf b/Engine/Shaders/Common.usf
--- a/Engine/Shaders/Common.usf
+++ b/Engine/Shaders/Common.usf
@@ -152,7 +152,8 @@
 
 // Whether the platform supports independent texture and samplers
 // When enabled, different texture lookups can share samplers to allow more artist samplers in the base pass
-#define SUPPORTS_INDEPENDENT_SAMPLERS (FEATURE_LEVEL_SM5 && !GL4_PROFILE)
+// Ideally this would just be enabled for all SM5 feature level platforms
+#define SUPPORTS_INDEPENDENT_SAMPLERS (PS4_PROFILE || SM5_PROFILE || SM4_PROFILE)
 
 // This would need to be a #define in GLSL to ignore the SamplerState, however, it is currently a function call in HLSL
 // for type checking of the parameters - ironically the type checking is really only needed in GLSL!
