Fog inscattering exponent is no longer clamped to 1

diff --git a/Engine/Shaders/HeightFogCommon.usf b/Engine/Shaders/HeightFogCommon.usf
--- a/Engine/Shaders/HeightFogCommon.usf
+++ b/Engine/Shaders/HeightFogCommon.usf
@@ -41,7 +41,7 @@ half4 GetExponentialHeightFog(float3 WorldPosition, float3 InCameraPosition)
 	{
 		// Setup a cosine lobe around the light direction to approximate inscattering from the directional light off of the ambient haze;
 		// clamp to avoid issues on some platforms (Could've used ClampedPow(), but we want the dot to be 0..1)
-		half3 DirectionalLightInscattering = DirectionalInscatteringColor.xyz * pow(saturate(dot(CameraToReceiverNormalized, InscatteringLightDirection.xyz)), clamp(DirectionalInscatteringColor.w, 0.000001f, 1.0));
+		half3 DirectionalLightInscattering = DirectionalInscatteringColor.xyz * pow(saturate(dot(CameraToReceiverNormalized, InscatteringLightDirection.xyz)), clamp(DirectionalInscatteringColor.w, 0.000001f, 1000.0));
 
 		// Calculate the line integral of the eye ray through the haze, using a special starting distance to limit the inscattering to the distance
 		float DirExponentialHeightLineIntegral = ExponentialHeightLineIntegralShared * max(CameraToReceiverLength - DirectionalInscatteringStartDistance, 0.0f);
