Fixed negative sqrt.

https://udn.unrealengine.com/questions/191267/broken-pre-compute-of-atmospheric-fog.html

diff --git a/Engine/Shaders/AtmospherePrecompute.usf b/Engine/Shaders/AtmospherePrecompute.usf
--- a/Engine/Shaders/AtmospherePrecompute.usf
+++ b/Engine/Shaders/AtmospherePrecompute.usf
@@ -73,7 +73,7 @@ void IrradianceNPS(FAtmosphereVSOutput Input, out float4 OutColor : SV_Target0)
 {
     float Radius, MuS;
     GetIrradianceRMuS(Input.OutTexCoord, Radius, MuS);
-    float3 S = float3(max(sqrt(1.0 - MuS * MuS), 0.0), 0.0, MuS);
+    float3 S = float3(sqrt(max(1.0 - MuS * MuS, 0.0)), 0.0, MuS);
 
     float3 Result = float3(0.f, 0.f, 0.f);
     // integral over 2.PI around x with two nested loops over W directions (Theta,Phi) -- Eq (15)
