UE4 - Fix for NaN's
- Fixes UE-3079, UE-3428 (and a ton of WEX issues also!)

diff --git a/Engine/Shaders/MaterialTemplate.usf b/Engine/Shaders/MaterialTemplate.usf
--- a/Engine/Shaders/MaterialTemplate.usf
+++ b/Engine/Shaders/MaterialTemplate.usf
@@ -146,17 +146,17 @@ struct FMaterialPixelParameters
 // @todo compat hack
 FMaterialPixelParameters MakeInitializedMaterialPixelParameters()
 {
+	FMaterialPixelParameters MPP;
 #if SM5_PROFILE || SM4_PROFILE || PS4_PROFILE || (ES2_PROFILE || ES3_1_PROFILE)
-	return (FMaterialPixelParameters)0;
+	MPP = (FMaterialPixelParameters)0;
 #else
 	// @todo compat hack - put back to the 0 typecast (2 lines up)
-	FMaterialPixelParameters MPP;
-#if NUM_MATERIAL_TEXCOORDS
+	#if NUM_MATERIAL_TEXCOORDS
 	for (int C = 0; C < NUM_MATERIAL_TEXCOORDS; C++)
 	{
 		MPP.TexCoords[C] = 0;
 	}
-#endif
+	#endif
 	MPP.VertexColor = 0;
 	MPP.TangentNormal = 0;
 	MPP.WorldNormal = 0;
@@ -166,7 +166,6 @@ FMaterialPixelParameters MakeInitializedMaterialPixelParameters()
 	MPP.ScreenPosition = 0;
 	MPP.UnMirrored = 0;
 	MPP.TwoSidedSign = 0;
-	MPP.TangentToWorld = 0;
 	MPP.WorldPosition = 0;
 	MPP.WorldPosition_NoOffsets = 0;
 	MPP.WorldPosition_CamRelative = 0;
@@ -179,21 +178,22 @@ FMaterialPixelParameters MakeInitializedMaterialPixelParameters()
 	MPP.Particle.PositionAndSize = 0;
 	MPP.Particle.MacroUV = 0;
 	MPP.Particle.DynamicParameter = 0;
-#if USE_PARTICLE_SUBUVS
+	#if USE_PARTICLE_SUBUVS
 	MPP.Particle.SubUVCoords[0] = 0;
 	MPP.Particle.SubUVCoords[1] = 0;
 	MPP.Particle.SubUVLerp = 0;
-#endif
+	#endif
 	MPP.Particle.Size = 0;
-#if LIGHTMAP_UV_ACCESS
+	#if LIGHTMAP_UV_ACCESS
 	MPP.LightmapUVs = 0;
-#endif
-#if USE_INSTANCING
+	#endif
+	#if USE_INSTANCING
 	MPP.PerInstanceParams = 0;
-#endif
+	#endif
 	MPP.SVPosition = 0;
-	return MPP;
 #endif
+	MPP.TangentToWorld = float3x3(1,0,0,0,1,0,0,0,1);
+	return MPP;
 }
 
 /** 
