Fixed clearcoatroughness.

diff --git a/Engine/Shaders/DeferredLightingCommon.usf b/Engine/Shaders/DeferredLightingCommon.usf
--- a/Engine/Shaders/DeferredLightingCommon.usf
+++ b/Engine/Shaders/DeferredLightingCommon.usf
@@ -594,7 +594,9 @@ float4 GetDynamicLighting(float3 WorldPosition, float3 CameraVector, float2 InUV
 		{
 			const float3 LightColor = LightData.LightColorAndFalloffExponent.rgb;
 
-			float3 LobeRoughness = float3(0.1, ScreenSpaceData.GBuffer.Roughness, 1);
+			const float ClearCoatRoughness	= ScreenSpaceData.GBuffer.CustomData.y;
+
+			float3 LobeRoughness = float3(ClearCoatRoughness, ScreenSpaceData.GBuffer.Roughness, 1);
 			float3 LobeEnergy = AreaLightSpecular(LightData, LobeRoughness, ToLight, L, V, N);
 
 			// accumulate diffuse and specular
