UE4 - Fix landscape shadows not working on forward shading. Still not the correct fix!
This is only needed for PC ES2 landscape currently, OpenGL doesn't need it (nor Metal)

diff --git a/Engine/Shaders/BasePassForForwardShadingPixelShader.usf b/Engine/Shaders/BasePassForForwardShadingPixelShader.usf
--- a/Engine/Shaders/BasePassForForwardShadingPixelShader.usf
+++ b/Engine/Shaders/BasePassForForwardShadingPixelShader.usf
@@ -84,13 +84,13 @@ void Main(
 
 	//@mw todo:
 	//Go over with Brian to determine the proper defines to wrap it and test cases to fix it.
-	/*
+#if !(GL3_PROFILE || GL4_PROFILE || COMPILER_GLSL_ES2 || COMPILER_GLSL || METAL_PROFILE)
 	float2 PixelPos = SvPosition - View.ViewRectMin.xy;	
 	MaterialParameters.ScreenPosition.xy = ( PixelPos / View.ViewSizeAndSceneTexelSize.xy - 0.5 ) * float2(2,-2);
 	MaterialParameters.ScreenPosition.zw = SvPosition.zw;
 	
 	MaterialParameters.ScreenPosition.xyz *= MaterialParameters.ScreenPosition.w;
-	*/
+#endif
 
 	// Store the results in local variables and reuse instead of calling the functions multiple times.
 	// Store the results in local variables and reuse instead of calling the functions multiple times.
