Compile fix for ES2.

diff --git a/Engine/Shaders/PostProcessUpscale.usf b/Engine/Shaders/PostProcessUpscale.usf
--- a/Engine/Shaders/PostProcessUpscale.usf
+++ b/Engine/Shaders/PostProcessUpscale.usf
@@ -63,7 +63,7 @@ void MainVS(
 	// amount of Quads in X and Y, from GTesselatedScreenRectangleIndexBuffer
 	uint2 RectangleTessellation = uint2(TESS_RECT_X, TESS_RECT_Y);
 
-	float2 TexCoord = float2(VertexId % (RectangleTessellation.x + 1), VertexId / (RectangleTessellation.x + 1)) / float2(RectangleTessellation);
+	float2 TexCoord = float2((uint)fmod(VertexId, (RectangleTessellation.x + 1)), VertexId / (RectangleTessellation.x + 1)) / float2(RectangleTessellation);
 
 	// still in 0..1 range
 	float4 Position = float4(TexCoord.x, TexCoord.y, 0, 1);
