Tiled deferred shader no longer loops over lighting model Id, which was making huge compile times and is also slower with the shading models we have now (5.76ms -> 5.54ms)

diff --git a/Engine/Shaders/TiledDeferredLightShaders.usf b/Engine/Shaders/TiledDeferredLightShaders.usf
--- a/Engine/Shaders/TiledDeferredLightShaders.usf
+++ b/Engine/Shaders/TiledDeferredLightShaders.usf
@@ -240,49 +240,37 @@ void TiledDeferredLightingMain(
 	FScreenSpaceData ScreenSpaceData = GetScreenSpaceDataUint(PixelPos);
 	FGBufferData InGBufferData = ScreenSpaceData.GBuffer; 
 	float3 CameraVector = normalize(WorldPosition - View.ViewOrigin.xyz);
-
-	// Multiple shading models are being handled by unrolling a top level loop
-	// This generates faster code with the D3D11 shader compiler than a single pass over all lights, which branches on GBuffer.ShadingModelID in PointLightSubsurface
-	// Presumably because the compiler was able to remove the shading model branches in the unrolled loop iterations
-	UNROLL
-	for (uint TileShadingModelID = 1; TileShadingModelID < SHADINGMODELID_NUM; TileShadingModelID++)
+		
+	BRANCH
+	if (InGBufferData.ShadingModelID != SHADINGMODELID_UNLIT)
 	{
-		// Only light pixels with the shading model being processed in this loop iteration
-		BRANCH
-		if (InGBufferData.ShadingModelID == TileShadingModelID)
+		LOOP
+		for (uint TileLightIndex = 0; TileLightIndex < NumLightsAffectingTile; TileLightIndex++) 
 		{
-			LOOP
-			for (uint TileLightIndex = 0; TileLightIndex < NumLightsAffectingTile; TileLightIndex++) 
-			{
-				uint LightIndex = TileLightIndices[TileLightIndex];
+			uint LightIndex = TileLightIndices[TileLightIndex];
 
-				FDeferredLightData LightData = (FDeferredLightData)0;
-				LightData.LightPositionAndInvRadius = TiledDeferred.LightPositionAndInvRadius[LightIndex];
-				LightData.LightColorAndFalloffExponent = TiledDeferred.LightColorAndFalloffExponent[LightIndex];
-				{
-					float4 Value = TiledDeferred2.LightDirectionAndSpotlightMaskAndMinRoughness[LightIndex];
-					LightData.LightDirection = Value.xyz;
-					LightData.bSpotLight = Value.w > 0;
-					LightData.MinRoughness = abs(Value.w);
-				}
-				LightData.SpotAnglesAndSourceRadius = float4( TiledDeferred2.SpotAnglesAndSourceRadiusAndSimpleLighting[LightIndex].xyz, 0 );
-				LightData.ShadowMapChannelMask = TiledDeferred2.ShadowMapChannelMask[LightIndex];
-				LightData.bInverseSquared = LightData.LightColorAndFalloffExponent.w == 0;
-				// Only radial lights supported with tiled deferred
-				LightData.bRadialLight = true;
-				// The only type of shadowing supported for lights using tiled is static shadowing, so the light should only compute shadowing if it has static shadowing
-				LightData.bShadowed = dot(LightData.ShadowMapChannelMask, float4(1, 1, 1, 1));
-	
-				// Lights requiring light attenuation are not supported tiled for now
-				CompositedLighting += GetDynamicLighting(WorldPosition, CameraVector, ScreenUV, ScreenSpaceData, TileShadingModelID, LightData, float4(1, 1, 1, 1), uint2(0,0));
+			FDeferredLightData LightData = (FDeferredLightData)0;
+			LightData.LightPositionAndInvRadius = TiledDeferred.LightPositionAndInvRadius[LightIndex];
+			LightData.LightColorAndFalloffExponent = TiledDeferred.LightColorAndFalloffExponent[LightIndex];
+			{
+				float4 Value = TiledDeferred2.LightDirectionAndSpotlightMaskAndMinRoughness[LightIndex];
+				LightData.LightDirection = Value.xyz;
+				LightData.bSpotLight = Value.w > 0;
+				LightData.MinRoughness = abs(Value.w);
 			}
+			LightData.SpotAnglesAndSourceRadius = float4( TiledDeferred2.SpotAnglesAndSourceRadiusAndSimpleLighting[LightIndex].xyz, 0 );
+			LightData.ShadowMapChannelMask = TiledDeferred2.ShadowMapChannelMask[LightIndex];
+			LightData.bInverseSquared = LightData.LightColorAndFalloffExponent.w == 0;
+			// Only radial lights supported with tiled deferred
+			LightData.bRadialLight = true;
+			// The only type of shadowing supported for lights using tiled is static shadowing, so the light should only compute shadowing if it has static shadowing
+			LightData.bShadowed = dot(LightData.ShadowMapChannelMask, float4(1, 1, 1, 1));
+	
+			// Lights requiring light attenuation are not supported tiled for now
+			CompositedLighting += GetDynamicLighting(WorldPosition, CameraVector, ScreenUV, ScreenSpaceData, InGBufferData.ShadingModelID, LightData, float4(1, 1, 1, 1), uint2(0,0));
 		}
-	}
-		
-	// The simple shading model does not depend on ShadingModelID, so use it anytime the material is lit
-	BRANCH
-	if (InGBufferData.ShadingModelID != SHADINGMODELID_UNLIT)
-	{
+
+		// The simple shading model does not depend on ShadingModelID, so use it anytime the material is lit
 		LOOP
 		for (uint TileLightIndex = 0; TileLightIndex < NumSimpleLightsAffectingTile; TileLightIndex++) 
 		{
