TwoSided Foliage shading model
* Diffuse lighting on the opposite side of the leaf becomes subsurface lighting on the side being shaded.  SubsurfaceColor filters this lighting and can be used to represent varying thickness as well (leaf veins).  Subsurface lighting is a mix of wrap-around diffuse with a directionally dependent lobe.
* Opacity only affects shadow falloff distance

diff --git a/Engine/Shaders/BasePassPixelShader.usf b/Engine/Shaders/BasePassPixelShader.usf
--- a/Engine/Shaders/BasePassPixelShader.usf
+++ b/Engine/Shaders/BasePassPixelShader.usf
@@ -268,13 +268,15 @@ float3 IndirectlightingCacheMinUV;
 float3 IndirectlightingCacheMaxUV;
 
 /** Computes sky diffuse lighting, including precomputed shadowing. */
-float3 GetSkyLighting(float3 WorldNormal, float2 LightmapUV)
+void GetSkyLighting(float3 WorldNormal, float2 LightmapUV, out float3 OutDiffuseLighting, out float3 OutSubsurfaceLighting)
 {
-	float3 Lighting = 0;
+	OutDiffuseLighting = 0;
+	OutSubsurfaceLighting = 0;
 
 #if ENABLE_SKY_LIGHT
 
 	float SkyVisibility = 1;
+	float GeometryTerm = 1;
 	float3 SkyLightingNormal = WorldNormal;
 	
 	#if HQ_TEXTURE_LIGHTMAP || CACHED_POINT_INDIRECT_LIGHTING || CACHED_VOLUME_INDIRECT_LIGHTING
@@ -313,7 +315,7 @@ float3 GetSkyLighting(float3 WorldNormal, float2 LightmapUV)
 
 				float DotProductFactor = lerp(saturate(dot(NormalizedBentNormal, WorldNormal)), 1, BentNormalWeightFactor);
 				// Account for darkening due to the geometry term
-				SkyVisibility *= DotProductFactor;
+				GeometryTerm = DotProductFactor;
 			#endif
 		}
 	#endif
@@ -322,21 +324,28 @@ float3 GetSkyLighting(float3 WorldNormal, float2 LightmapUV)
 	float3 DiffuseLookup = GetSkySHDiffuse(SkyLightingNormal) * View.SkyLightColor.rgb;
 
 	// Apply AO to the sky diffuse
-	Lighting += DiffuseLookup * SkyVisibility;
-#endif
+	OutDiffuseLighting += DiffuseLookup * (SkyVisibility * GeometryTerm);
+
+	#if MATERIAL_SHADINGMODEL_TWOSIDED_FOLIAGE
+		float3 BackfaceDiffuseLookup = GetSkySHDiffuse(-WorldNormal) * View.SkyLightColor.rgb;
+		OutSubsurfaceLighting += BackfaceDiffuseLookup * SkyVisibility;
+	#endif
 
-	return Lighting;
+#endif
 }
 
 /** Calculates indirect lighting contribution on this object from precomputed data. */
-float3 GetPrecomputedIndirectLighting(
+void GetPrecomputedIndirectLightingAndSkyLight(
 	FMaterialPixelParameters MaterialParameters, 
 	FVertexFactoryInterpolantsVSToPS Interpolants,
 	FBasePassInterpolantsVSToPS BasePassInterpolants,
-	out float IndirectIrradiance)
+	out float3 OutDiffuseLighting,
+	out float3 OutSubsurfaceLighting,
+	out float OutIndirectIrradiance)
 {
-	IndirectIrradiance = 0;
-	float3 Lighting = 0;
+	OutIndirectIrradiance = 0;
+	OutDiffuseLighting = 0;
+	OutSubsurfaceLighting = 0;
 	float2 SkyOcclusionUV = 0;
 
 	// Method for movable components which want to use a volume texture of interpolated SH samples
@@ -353,7 +362,7 @@ float3 GetPrecomputedIndirectLighting(
 		#define AMBIENTONLY 0
 		#if AMBIENTONLY
 
-			Lighting = Vector0.rgb / SHAmbientFunction() / PI;
+			OutDiffuseLighting = Vector0.rgb / SHAmbientFunction() / PI;
 
 		#else
 
@@ -368,7 +377,12 @@ float3 GetPrecomputedIndirectLighting(
 
 			// Diffuse convolution
 			FTwoBandSHVector DiffuseTransferSH = CalcDiffuseTransferSH(MaterialParameters.WorldNormal, 1);
-			Lighting = max(half3(0,0,0), DotSH(CachedSH, DiffuseTransferSH)) / PI;
+			OutDiffuseLighting = max(half3(0,0,0), DotSH(CachedSH, DiffuseTransferSH)) / PI;
+
+			#if MATERIAL_SHADINGMODEL_TWOSIDED_FOLIAGE
+				FTwoBandSHVector SubsurfaceTransferSH = CalcDiffuseTransferSH(-MaterialParameters.WorldNormal, 1);
+				OutSubsurfaceLighting += max(half3(0,0,0), DotSH(CachedSH, SubsurfaceTransferSH)) / PI;
+			#endif
 
 		#endif
 
@@ -382,7 +396,12 @@ float3 GetPrecomputedIndirectLighting(
 			PointIndirectLighting.B.V = IndirectLightingSHCoefficients[2].x;
 
 			FOneBandSHVector DiffuseTransferSH = CalcDiffuseTransferSH1(MaterialParameters.WorldNormal, 1);
-			Lighting = DotSH1(PointIndirectLighting, DiffuseTransferSH);
+			OutDiffuseLighting = DotSH1(PointIndirectLighting, DiffuseTransferSH);
+
+			#if MATERIAL_SHADINGMODEL_TWOSIDED_FOLIAGE
+				FOneBandSHVector SubsurfaceTransferSH = CalcDiffuseTransferSH1(-MaterialParameters.WorldNormal, 1);
+				OutSubsurfaceLighting += DotSH1(PointIndirectLighting, SubsurfaceTransferSH);
+			#endif
 
 		#else
 
@@ -393,7 +412,12 @@ float3 GetPrecomputedIndirectLighting(
 
 			FTwoBandSHVector DiffuseTransferSH = CalcDiffuseTransferSH(MaterialParameters.WorldNormal, 1);
 			// Compute diffuse lighting which takes the normal into account
-			Lighting = max(half3(0,0,0), DotSH(PointIndirectLighting, DiffuseTransferSH));
+			OutDiffuseLighting = max(half3(0,0,0), DotSH(PointIndirectLighting, DiffuseTransferSH));
+
+			#if MATERIAL_SHADINGMODEL_TWOSIDED_FOLIAGE
+				FTwoBandSHVector SubsurfaceTransferSH = CalcDiffuseTransferSH(-MaterialParameters.WorldNormal, 1);
+				OutSubsurfaceLighting += max(half3(0,0,0), DotSH(PointIndirectLighting, SubsurfaceTransferSH));
+			#endif
 
 		#endif
 
@@ -403,24 +427,31 @@ float3 GetPrecomputedIndirectLighting(
 		float2 LightmapUV0, LightmapUV1;
 		GetLightMapCoordinates(Interpolants, LightmapUV0, LightmapUV1);
 		SkyOcclusionUV = LightmapUV0;
-		Lighting = GetLightMapColorHQ(LightmapUV0, LightmapUV1, MaterialParameters.WorldNormal).rgb;
+		GetLightMapColorHQ(LightmapUV0, LightmapUV1, MaterialParameters.WorldNormal, OutDiffuseLighting, OutSubsurfaceLighting);
 
 	// Low quality texture lightmaps
 	#elif LQ_TEXTURE_LIGHTMAP
 
 		float2 LightmapUV0, LightmapUV1;
 		GetLightMapCoordinates(Interpolants, LightmapUV0, LightmapUV1);
-		Lighting = GetLightMapColorLQ( LightmapUV0, LightmapUV1, MaterialParameters.WorldNormal ).rgb;
+		OutDiffuseLighting = GetLightMapColorLQ(LightmapUV0, LightmapUV1, MaterialParameters.WorldNormal).rgb;
 
 	#endif
 
-	// Sky lighting must contribute to IndirectIrradiance for ReflectionEnvironment lightmap mixing
-	float3 SkyLighting = GetSkyLighting(MaterialParameters.WorldNormal, SkyOcclusionUV);
+	// Apply indirect lighting scale while we have only accumulated lightmaps
+	OutDiffuseLighting *= View.IndirectLightingColorScale;
+
+	float3 SkyDiffuseLighting;
+	float3 SkySubsurfaceLighting;
+	GetSkyLighting(MaterialParameters.WorldNormal, SkyOcclusionUV, SkyDiffuseLighting, SkySubsurfaceLighting);
+
+	OutSubsurfaceLighting += SkySubsurfaceLighting;
 
 	BRANCH
 	if( View.UseLightmaps > 0 )
 	{
-		Lighting += SkyLighting;
+		// Sky lighting must contribute to IndirectIrradiance for ReflectionEnvironment lightmap mixing
+		OutDiffuseLighting += SkyDiffuseLighting;
 	}
 	else
 	{
@@ -428,16 +459,13 @@ float3 GetPrecomputedIndirectLighting(
 		float DiffuseMip = ComputeReflectionCaptureMipFromRoughness(1);
 		float4 DiffuseIBL = TextureCubeArraySampleLevel(ReflectionCubemap, ReflectionCubemapSampler, MaterialParameters.WorldNormal, CubemapArrayIndex, DiffuseMip);
 		
-		Lighting += DiffuseIBL.rgb + (1 - DiffuseIBL.a) * SkyLighting;
+		OutDiffuseLighting += DiffuseIBL.rgb + (1 - DiffuseIBL.a) * SkyDiffuseLighting;
 	#endif
 	}
 
 	#if HQ_TEXTURE_LIGHTMAP || LQ_TEXTURE_LIGHTMAP || CACHED_VOLUME_INDIRECT_LIGHTING || CACHED_POINT_INDIRECT_LIGHTING
-		Lighting *= View.IndirectLightingColorScale;
-		IndirectIrradiance = Luminance(Lighting);
+		OutIndirectIrradiance = Luminance(OutDiffuseLighting);
 	#endif
-
-	return Lighting;
 }
 
 #if EDITOR_PRIMITIVE_MATERIAL
@@ -576,11 +604,11 @@ void Main(
 
 	// If we don't use this shading model the color should be black (don't generate shader code for unused data, don't do indirectlighting cache lighting with this color).
 	float3 SubsurfaceColor = 0;
-#if MATERIAL_SHADINGMODEL_SUBSURFACE || MATERIAL_SHADINGMODEL_PREINTEGRATED_SKIN || MATERIAL_SHADINGMODEL_SUBSURFACE_PROFILE
+#if MATERIAL_SHADINGMODEL_SUBSURFACE || MATERIAL_SHADINGMODEL_PREINTEGRATED_SKIN || MATERIAL_SHADINGMODEL_SUBSURFACE_PROFILE || MATERIAL_SHADINGMODEL_TWOSIDED_FOLIAGE
 	{
 		float4 SubsurfaceData = GetMaterialSubsurfaceData(MaterialParameters);
 
-#if MATERIAL_SHADINGMODEL_SUBSURFACE || MATERIAL_SHADINGMODEL_PREINTEGRATED_SKIN
+#if MATERIAL_SHADINGMODEL_SUBSURFACE || MATERIAL_SHADINGMODEL_PREINTEGRATED_SKIN || MATERIAL_SHADINGMODEL_TWOSIDED_FOLIAGE
 		SubsurfaceColor = SubsurfaceData.rgb;
 #endif
 		SubsurfaceProfile = SubsurfaceData.a;
@@ -629,7 +657,19 @@ void Main(
 
 	#if !MATERIAL_SHADINGMODEL_UNLIT
 
-		Color += GetPrecomputedIndirectLighting(MaterialParameters, Interpolants, BasePassInterpolants, IndirectIrradiance) * (DiffuseColor + SubsurfaceColor) * MaterialAO;
+		float3 DiffuseIndirectLighting;
+		float3 SubsurfaceIndirectLighting;
+		GetPrecomputedIndirectLightingAndSkyLight(MaterialParameters, Interpolants, BasePassInterpolants, DiffuseIndirectLighting, SubsurfaceIndirectLighting, IndirectIrradiance);
+
+		float3 EffectiveDiffuseColor = DiffuseColor;
+
+		#if MATERIAL_SHADINGMODEL_SUBSURFACE || MATERIAL_SHADINGMODEL_PREINTEGRATED_SKIN
+			// Add subsurface energy to diffuse
+			//@todo - better subsurface handling for these shading models with skylight and precomputed GI
+			DiffuseColor += SubsurfaceColor;
+		#endif
+
+		Color += (DiffuseIndirectLighting * DiffuseColor + SubsurfaceIndirectLighting * SubsurfaceColor) * MaterialAO;
 
 		#if SIMPLE_DYNAMIC_LIGHTING
 			// always unshadowed so BiasedNDotL is not needed
@@ -762,6 +802,9 @@ void Main(
 
 			GBuffer.CustomData.x = ClearCoat;
 			GBuffer.CustomData.y = ClearCoatRoughness;
+		#elif MATERIAL_SHADINGMODEL_TWOSIDED_FOLIAGE
+			GBuffer.ShadingModelID = SHADINGMODELID_TWOSIDED_FOLIAGE;
+			GBuffer.CustomData = EncodeSubsurfaceColor(SubsurfaceColor);
 		#else
 			// missing shading model, compiler should report ShadingModelID is not set
 		#endif
diff --git a/Engine/Shaders/DeferredLightingCommon.usf b/Engine/Shaders/DeferredLightingCommon.usf
--- a/Engine/Shaders/DeferredLightingCommon.usf
+++ b/Engine/Shaders/DeferredLightingCommon.usf
@@ -13,9 +13,6 @@
 #include "MonteCarlo.usf"
 #include "IESLightProfilesCommon.usf"
 
-// 0/1, default is 0 as on NV670 there was a slowdown in average scene, keep for further experiments
-#define LIGHTCULLING_FINE 0
-
 /** 
  * Data about a single light.
  * Putting the light data in this struct allows the same lighting code to be used between standard deferred, 
@@ -248,6 +245,25 @@ float3 SubsurfaceShadingSubsurface( FGBufferData GBuffer, float3 L, float3 V, ha
 	return SubsurfaceColor * lerp(BackScatter, 1, InScatter);
 }
 
+float3 TwoSidedFoliageShadingSubsurface( FGBufferData GBuffer, float3 L, float3 V, half3 N )
+{
+	float VoL = saturate(dot(V, -L));
+	float VoLLobe = VoL * VoL;
+	VoLLobe *= VoLLobe;
+	float Exponent = 8;
+	// Modified phong energy conservation
+	VoLLobe = VoLLobe * (Exponent + 1) / (2.0f * PI);
+
+	float DiffuseWrap = .6f;
+	// Energy conserving wrap
+	float BackDiffuse = saturate(dot(-N, L) * DiffuseWrap * DiffuseWrap + (1 - DiffuseWrap) * DiffuseWrap) / PI;  
+	float ViewDependenceAmount = .5f;
+	float BackShading = lerp(BackDiffuse, VoLLobe, ViewDependenceAmount);
+
+	float3 SubsurfaceColor = ExtractSubsurfaceColor(GBuffer);
+	return BackShading * SubsurfaceColor;
+}
+
 Texture2D		PreIntegratedBRDF;
 SamplerState	PreIntegratedBRDFSampler;
 
@@ -270,6 +286,7 @@ float3 SurfaceShading( FGBufferData GBuffer, float3 LobeRoughness, float3 LobeEn
 		case SHADINGMODELID_SUBSURFACE:
 		case SHADINGMODELID_PREINTEGRATED_SKIN:
 		case SHADINGMODELID_SUBSURFACE_PROFILE:
+		case SHADINGMODELID_TWOSIDED_FOLIAGE:
 			return StandardShading( GBuffer, LobeRoughness, LobeEnergy, L, V, N, DiffSpecMask);
 		case SHADINGMODELID_CLEAR_COAT:
 			// this path does not support DiffSpecMask yet
@@ -287,6 +304,8 @@ float3 SubsurfaceShading( FGBufferData GBuffer, float3 L, float3 V, half3 N )
 			return SubsurfaceShadingSubsurface( GBuffer, L, V, N );
 		case SHADINGMODELID_PREINTEGRATED_SKIN:
 			return SubsurfaceShadingPreintegratedSkin( GBuffer, L, V, N );
+		case SHADINGMODELID_TWOSIDED_FOLIAGE:
+			return TwoSidedFoliageShadingSubsurface( GBuffer, L, V, N );
 		default:
 			return 0;
 	}
@@ -586,11 +605,6 @@ float4 GetDynamicLighting(float3 WorldPosition, float3 CameraVector, float2 InUV
 
 		LightAccumulator.EstimatedCost += 0.3f;		// add the cost of getting the shadow terms
 
-#if LIGHTCULLING_FINE == 1
-		// optimization that should help if there are a lot of shadowed pixels
-		BRANCH if (ShadowedAttenuation > 0 || ShadingModelID == SHADINGMODELID_SUBSURFACE)
-#endif
-
 		{
 			const float3 LightColor = LightData.LightColorAndFalloffExponent.rgb;
 
diff --git a/Engine/Shaders/DeferredShadingCommon.usf b/Engine/Shaders/DeferredShadingCommon.usf
--- a/Engine/Shaders/DeferredShadingCommon.usf
+++ b/Engine/Shaders/DeferredShadingCommon.usf
@@ -228,7 +228,16 @@ uint DecodeShadingModelId(float InPackedChannel)
 #define SHADINGMODELID_PREINTEGRATED_SKIN	3
 #define SHADINGMODELID_CLEAR_COAT			4
 #define SHADINGMODELID_SUBSURFACE_PROFILE	5
-#define SHADINGMODELID_NUM					6
+#define SHADINGMODELID_TWOSIDED_FOLIAGE		6
+#define SHADINGMODELID_NUM					7
+
+bool IsSubsurfaceModel(int ShadingModel)
+{
+	return ShadingModel == SHADINGMODELID_SUBSURFACE 
+		|| ShadingModel == SHADINGMODELID_PREINTEGRATED_SKIN 
+		|| ShadingModel == SHADINGMODELID_SUBSURFACE_PROFILE
+		|| ShadingModel == SHADINGMODELID_TWOSIDED_FOLIAGE;
+}
 
 // all values that are output by the forward rendering pass
 struct FGBufferData
@@ -324,7 +333,7 @@ void EncodeGBuffer(
 	}
 	else
 	{
-		bool bSubsurface = (GBuffer.ShadingModelID == SHADINGMODELID_SUBSURFACE || GBuffer.ShadingModelID == SHADINGMODELID_PREINTEGRATED_SKIN || GBuffer.ShadingModelID == SHADINGMODELID_SUBSURFACE_PROFILE);
+		bool bSubsurface = IsSubsurfaceModel(GBuffer.ShadingModelID);
 
 		// NOTE OutGBufferB.b is currently unused!
 		OutGBufferB.r = bSubsurface ? GBuffer.Opacity : GBuffer.Metallic;
@@ -369,7 +378,7 @@ FGBufferData DecodeGBufferData(
 	GBuffer.HasDistanceFieldRepresentation = PackedAlpha & 1;
 	GBuffer.ShadingModelID = DecodeShadingModelId(InGBufferB.a);
 	
-	bool bSubsurface = (GBuffer.ShadingModelID == SHADINGMODELID_SUBSURFACE || GBuffer.ShadingModelID == SHADINGMODELID_PREINTEGRATED_SKIN || GBuffer.ShadingModelID == SHADINGMODELID_SUBSURFACE_PROFILE);
+	bool bSubsurface = IsSubsurfaceModel(GBuffer.ShadingModelID);
 
 	GBuffer.BaseColor = DecodeBaseColor(InGBufferC.rgb);
 	GBuffer.Metallic = bSubsurface ? 0 : InGBufferB.r;
diff --git a/Engine/Shaders/LightmapCommon.usf b/Engine/Shaders/LightmapCommon.usf
--- a/Engine/Shaders/LightmapCommon.usf
+++ b/Engine/Shaders/LightmapCommon.usf
@@ -45,8 +45,10 @@ half4 GetLightMapColorLQ( float2 LightmapUV0, float2 LightmapUV1, half3 WorldNor
 	return half4( Color, Luma );
 }
 
-half4 GetLightMapColorHQ( float2 LightmapUV0, float2 LightmapUV1, half3 WorldNormal )
+void GetLightMapColorHQ( float2 LightmapUV0, float2 LightmapUV1, half3 WorldNormal, out half3 OutDiffuseLighting, out half3 OutSubsurfaceLighting )
 {
+	OutSubsurfaceLighting = 0;
+
 	half4 Lightmap0 = Texture2DSample( LightMapTexture, LightMapSampler, LightmapUV0 );
 	half4 Lightmap1 = Texture2DSample( LightMapTexture, LightMapSampler, LightmapUV1 );
 
@@ -71,6 +73,11 @@ half4 GetLightMapColorHQ( float2 LightmapUV0, float2 LightmapUV1, half3 WorldNor
 
 	// Sample SH with normal
 	half Directionality = max( 0.0, dot( SH, half4(WorldNormal.yzx, 1) ) );
+
+	#if MATERIAL_SHADINGMODEL_TWOSIDED_FOLIAGE
+		half SubsurfaceDirectionality = max( 0.0, dot( SH, half4(-WorldNormal.yzx, 1) ) );
+		OutSubsurfaceLighting = L * SubsurfaceDirectionality * UVW;
+	#endif
 #else
 	half Directionality = 0.6;
 #endif
@@ -78,7 +85,7 @@ half4 GetLightMapColorHQ( float2 LightmapUV0, float2 LightmapUV1, half3 WorldNor
 	half Luma = L * Directionality;
 	half3 Color = Luma * UVW;
 
-	return half4( Color, Luma );
+	OutDiffuseLighting = Color;
 }
 
 float4 GetSkyBentNormalAndOcclusion(float2 LightmapUV)
diff --git a/Engine/Shaders/MaterialTemplate.usf b/Engine/Shaders/MaterialTemplate.usf
--- a/Engine/Shaders/MaterialTemplate.usf
+++ b/Engine/Shaders/MaterialTemplate.usf
@@ -868,6 +868,7 @@ float3 GetShadingModelColor(FGBufferData GBuffer)
 	else if (GBuffer.ShadingModelID == SHADINGMODELID_SUBSURFACE) return float3(1.0f, 0.1f, 0.1f); // Red
 	else if (GBuffer.ShadingModelID == SHADINGMODELID_PREINTEGRATED_SKIN) return float3(0.6f, 0.4f, 0.1f); // Brown
 	else if (GBuffer.ShadingModelID == SHADINGMODELID_SUBSURFACE_PROFILE) return float3(0.2f, 0.6f, 0.5f); // Cyan
+	else if (GBuffer.ShadingModelID == SHADINGMODELID_TWOSIDED_FOLIAGE) return float3(0.2f, 0.2f, 0.8f); // Blue
 	else return float3(1.0f, 1.0f, 1.0f); // White
 #else
 	switch(GBuffer.ShadingModelID)
@@ -877,6 +878,7 @@ float3 GetShadingModelColor(FGBufferData GBuffer)
 		case SHADINGMODELID_SUBSURFACE: return float3(1.0f, 0.1f, 0.1f); // Red
 		case SHADINGMODELID_PREINTEGRATED_SKIN: return float3(0.6f, 0.4f, 0.1f); // Brown
 		case SHADINGMODELID_SUBSURFACE_PROFILE: return float3(0.2f, 0.6f, 0.5f); // Cyan
+		case SHADINGMODELID_TWOSIDED_FOLIAGE: return float3(0.2f, 0.2f, 0.8f); // Cyan
 		default: return float3(1.0f, 1.0f, 1.0f); // White
 	}
 #endif
@@ -928,7 +930,9 @@ float4 SceneTextureLookup(float2 UV, int SceneTextureIndex, bool bFiltered)
 		// PPI_SpecularColor
 		case 3: return float4(ScreenSpaceData.GBuffer.SpecularColor, 0);
 		// PPI_SubsurfaceColor
-		case 4: return float4(ScreenSpaceData.GBuffer.CustomData, 0);
+		case 4: 
+			float3 CustomData = IsSubsurfaceModel(ScreenSpaceData.GBuffer.ShadingModelID) ? ExtractSubsurfaceColor(ScreenSpaceData.GBuffer) : ScreenSpaceData.GBuffer.CustomData;
+			return float4(CustomData, 0);
 		// PPI_BaseColor
 		case 5: return float4(ScreenSpaceData.GBuffer.BaseColor, 0);
 		// PPI_Specular
diff --git a/Engine/Shaders/PostProcessSubsurface.usf b/Engine/Shaders/PostProcessSubsurface.usf
--- a/Engine/Shaders/PostProcessSubsurface.usf
+++ b/Engine/Shaders/PostProcessSubsurface.usf
@@ -31,7 +31,7 @@ float GetSubsurfaceStrength(float2 UV)
 {
 	FScreenSpaceData ScreenSpaceData = GetScreenSpaceData(UV);
 
-	float Mask = ScreenSpaceData.GBuffer.ShadingModelID >= SHADINGMODELID_SUBSURFACE_PROFILE;
+	float Mask = ScreenSpaceData.GBuffer.ShadingModelID == SHADINGMODELID_SUBSURFACE_PROFILE;
 
 	return Mask * ScreenSpaceData.GBuffer.Opacity;
 }
@@ -204,7 +204,7 @@ float4 SetupSubsurfaceForOnePixel(float2 UV)
 
 	FScreenSpaceData ScreenSpaceData = GetScreenSpaceData(UV);
 
-	FLATTEN if (ScreenSpaceData.GBuffer.ShadingModelID >= SHADINGMODELID_SUBSURFACE_PROFILE)
+	FLATTEN if (ScreenSpaceData.GBuffer.ShadingModelID == SHADINGMODELID_SUBSURFACE_PROFILE)
 	{
 		float4 SceneColor4 = Texture2DSample(PostprocessInput0, PostprocessInput0Sampler, UV);
 
@@ -352,7 +352,7 @@ void SubsurfaceRecombinePS(float4 UVAndScreenPos : TEXCOORD0, out float4 OutColo
 	}
 
 	// 0: dont' fade in subsurface half res,  1:fully fade in half res
-	float BlendFactor = (ScreenSpaceData.GBuffer.ShadingModelID >= SHADINGMODELID_SUBSURFACE_PROFILE);
+	float BlendFactor = (ScreenSpaceData.GBuffer.ShadingModelID == SHADINGMODELID_SUBSURFACE_PROFILE);
 
 #if 0
 	// fade out subsurface scattering if radius is too small to be more crips (not blend with half resolution)
@@ -394,7 +394,7 @@ void ExtractSpecularPS(float4 UVAndScreenPos : TEXCOORD0, out float4 OutColor :
 	FScreenSpaceData ScreenSpaceData = GetScreenSpaceData(GBufferUV);
 
 	// 0: dont' fade in subsurface half res,  1:fully fade in half res
-	float BlendFactor = (ScreenSpaceData.GBuffer.ShadingModelID >= SHADINGMODELID_SUBSURFACE_PROFILE);
+	float BlendFactor = (ScreenSpaceData.GBuffer.ShadingModelID == SHADINGMODELID_SUBSURFACE_PROFILE);
 
 	float3 ExtractedSpecular = 0;
 	 
diff --git a/Engine/Shaders/ShadowProjectionPixelShader.usf b/Engine/Shaders/ShadowProjectionPixelShader.usf
--- a/Engine/Shaders/ShadowProjectionPixelShader.usf
+++ b/Engine/Shaders/ShadowProjectionPixelShader.usf
@@ -102,7 +102,7 @@ void Main(
 	#if FEATURE_LEVEL >= FEATURE_LEVEL_SM4
 
 		BRANCH
-		if (InGBufferData.ShadingModelID == SHADINGMODELID_SUBSURFACE || InGBufferData.ShadingModelID == SHADINGMODELID_PREINTEGRATED_SKIN)
+		if (IsSubsurfaceModel(InGBufferData.ShadingModelID))
 		{
 			// Derive density from a heuristic using opacity, tweaked for useful falloff ranges and to give a linear depth falloff with opacity
 			float Density = -.05f * log(1 - min(InGBufferData.Opacity, .999f));
diff --git a/Engine/Shaders/SkyLighting.usf b/Engine/Shaders/SkyLighting.usf
--- a/Engine/Shaders/SkyLighting.usf
+++ b/Engine/Shaders/SkyLighting.usf
@@ -67,13 +67,20 @@ void SkyLightDiffusePS(
 		
 #endif
 
+		float ScalarFactors = lerp(SkyVisibility, 1, OcclusionTintAndMinOcclusion.w) * InGBufferData.GBufferAO * ScreenSpaceData.AmbientOcclusion;
+
+		if (InGBufferData.ShadingModelID == SHADINGMODELID_TWOSIDED_FOLIAGE)
+		{
+			float3 SubsurfaceLookup = GetSkySHDiffuse(-InGBufferData.WorldNormal) * View.SkyLightColor.rgb;
+			float3 SubsurfaceColor = ExtractSubsurfaceColor(InGBufferData);
+			Lighting += ScalarFactors * SubsurfaceLookup * SubsurfaceColor;
+		}
+
 		// Compute the preconvolved incoming lighting with the bent normal direction
 		float3 DiffuseLookup = GetSkySHDiffuse(SkyLightingNormal) * View.SkyLightColor.rgb;
 
-		float ScalarFactors = lerp(SkyVisibility, 1, OcclusionTintAndMinOcclusion.w) * DotProductFactor * InGBufferData.GBufferAO * ScreenSpaceData.AmbientOcclusion;
-
 		// Apply AO to the sky diffuse and account for darkening due to the geometry term
-		Lighting = ScalarFactors * DiffuseLookup * InGBufferData.DiffuseColor + (1 - SkyVisibility) * OcclusionTintAndMinOcclusion.xyz;
+		Lighting += (ScalarFactors * DotProductFactor) * DiffuseLookup * InGBufferData.DiffuseColor + (1 - SkyVisibility) * OcclusionTintAndMinOcclusion.xyz;
 
 		// For debugging
 		//Lighting = (Texture2DSampleLevel(BentNormalAOTexture, BentNormalAOSampler, UV, 0).xyz);
diff --git a/Engine/Source/Runtime/Engine/Classes/Components/SkyLightComponent.h b/Engine/Source/Runtime/Engine/Classes/Components/SkyLightComponent.h
--- a/Engine/Source/Runtime/Engine/Classes/Components/SkyLightComponent.h
+++ b/Engine/Source/Runtime/Engine/Classes/Components/SkyLightComponent.h
@@ -103,8 +103,9 @@ class ENGINE_API USkyLightComponent : public ULightComponentBase
 	float SkyDistanceThreshold;
 
 	/** 
-	 * Whether all lighting from the lower hemisphere should be set to zero.  
-	 * This is valid when lighting a scene on a planet, but can be disabled for model showcases, etc.
+	 * Whether all distant lighting from the lower hemisphere should be set to zero.  
+	 * Enabling this is accurate when lighting a scene on a planet, 
+	 * However disabling it can be useful to approximate skylight bounce lighting (eg Movable light).
 	 */
 	UPROPERTY(EditAnywhere, BlueprintReadOnly, Category=Light)
 	bool bLowerHemisphereIsBlack;
diff --git a/Engine/Source/Runtime/Engine/Classes/Engine/EngineTypes.h b/Engine/Source/Runtime/Engine/Classes/Engine/EngineTypes.h
--- a/Engine/Source/Runtime/Engine/Classes/Engine/EngineTypes.h
+++ b/Engine/Source/Runtime/Engine/Classes/Engine/EngineTypes.h
@@ -142,6 +142,7 @@ enum EMaterialShadingModel
 	MSM_PreintegratedSkin	UMETA(DisplayName="Preintegrated Skin"),
 	MSM_ClearCoat			UMETA(DisplayName="Clear Coat"),
 	MSM_SubsurfaceProfile	UMETA(DisplayName="Subsurface Profile"),
+	MSM_TwoSidedFoliage		UMETA(DisplayName="Two Sided Foliage"),
 	MSM_MAX,
 };
 
diff --git a/Engine/Source/Runtime/Engine/Classes/Materials/MaterialInstance.h b/Engine/Source/Runtime/Engine/Classes/Materials/MaterialInstance.h
--- a/Engine/Source/Runtime/Engine/Classes/Materials/MaterialInstance.h
+++ b/Engine/Source/Runtime/Engine/Classes/Materials/MaterialInstance.h
@@ -177,7 +177,7 @@ class UMaterialInstance : public UMaterialInterface
 	/** Flag to detect cycles in the material instance graph. */
 	uint32 ReentrantFlag:1;
 
-	/** Defines if SubsurfaceProfile from tis instance is used or it uses the parent one. */
+	/** Defines if SubsurfaceProfile from this instance is used or it uses the parent one. */
 	UPROPERTY(EditAnywhere, BlueprintReadOnly, Category = MaterialInstance)
 	uint32 bOverrideSubsurfaceProfile:1;
 
diff --git a/Engine/Source/Runtime/Engine/Private/Materials/HLSLMaterialTranslator.h b/Engine/Source/Runtime/Engine/Private/Materials/HLSLMaterialTranslator.h
--- a/Engine/Source/Runtime/Engine/Private/Materials/HLSLMaterialTranslator.h
+++ b/Engine/Source/Runtime/Engine/Private/Materials/HLSLMaterialTranslator.h
@@ -341,7 +341,7 @@ public:
 			EMaterialShadingModel MaterialShadingModel = Material->GetShadingModel();
 
 			if (Material->GetMaterialDomain() == MD_Surface
-				&& (MaterialShadingModel == MSM_Subsurface || MaterialShadingModel == MSM_PreintegratedSkin || MaterialShadingModel == MSM_SubsurfaceProfile))
+				&& IsSubsurfaceShadingModel(MaterialShadingModel))
 			{
 				// Note we don't test for the blend mode as you can have a translucent material using the subsurface shading model
 
@@ -527,13 +527,6 @@ public:
 
 			MaterialCompilationOutput.UniformExpressionSet.SetParameterCollections(ParameterCollections);
 
-/* test
-			if(MaterialShadingModel == MSM_SubsurfaceProfile)
-			{
-				MaterialCompilationOutput.UniformExpressionSet.SetCompileInSubsurfaceProfile();
-			}
-*/
-
 			// Create the material uniform buffer struct.
 			MaterialCompilationOutput.UniformExpressionSet.CreateBufferStruct();
 		}
@@ -764,15 +757,15 @@ protected:
 			check(!CodeChunks[InProperty][InFrequency][Index].UniformExpression || CodeChunks[InProperty][InFrequency][Index].UniformExpression->IsConstant());
 			if (CodeChunks[InProperty][InFrequency][Index].UniformExpression && CodeChunks[InProperty][InFrequency][Index].UniformExpression->IsConstant())
 			{
-					// Handle a constant uniform expression being the only code chunk hooked up to a material input
-					const FShaderCodeChunk& CodeChunk = CodeChunks[InProperty][InFrequency][Index];
+				// Handle a constant uniform expression being the only code chunk hooked up to a material input
+				const FShaderCodeChunk& CodeChunk = CodeChunks[InProperty][InFrequency][Index];
 				OutValue = CodeChunk.Definition;
-				}
-				else
-				{
-					const FShaderCodeChunk& CodeChunk = CodeChunks[InProperty][InFrequency][Index];
-					// Combine the definition lines and the return statement
-					check(CodeChunk.bInline || CodeChunk.SymbolName.Len() > 0);
+			}
+			else
+			{
+				const FShaderCodeChunk& CodeChunk = CodeChunks[InProperty][InFrequency][Index];
+				// Combine the definition lines and the return statement
+				check(CodeChunk.bInline || CodeChunk.SymbolName.Len() > 0);
 				OutDefinitions = GetDefinitions(InProperty, InFrequency);
 				OutValue = CodeChunk.bInline ? CodeChunk.Definition : CodeChunk.SymbolName;
 			}
diff --git a/Engine/Source/Runtime/Engine/Private/Materials/Material.cpp b/Engine/Source/Runtime/Engine/Private/Materials/Material.cpp
--- a/Engine/Source/Runtime/Engine/Private/Materials/Material.cpp
+++ b/Engine/Source/Runtime/Engine/Private/Materials/Material.cpp
@@ -3642,7 +3642,7 @@ bool UMaterial::IsPropertyActive(EMaterialProperty InProperty) const
 		break;
 	case MP_Opacity:
 		Active = IsTranslucentBlendMode((EBlendMode)BlendMode) && BlendMode != BLEND_Modulate;
-		if (ShadingModel == MSM_Subsurface || ShadingModel == MSM_PreintegratedSkin || ShadingModel == MSM_SubsurfaceProfile)
+		if (IsSubsurfaceShadingModel(ShadingModel))
 		{
 			Active = true;
 		}
@@ -3657,13 +3657,14 @@ bool UMaterial::IsPropertyActive(EMaterialProperty InProperty) const
 		Active = ShadingModel != MSM_Unlit;
 		break;
 	case MP_Metallic:
-		Active = ShadingModel != MSM_Unlit && ShadingModel != MSM_Subsurface && ShadingModel != MSM_PreintegratedSkin && ShadingModel != MSM_SubsurfaceProfile;
+		// Subsurface models store opacity in place of Metallic in the GBuffer
+		Active = ShadingModel != MSM_Unlit && !IsSubsurfaceShadingModel(ShadingModel);
 		break;
 	case MP_Normal:
 		Active = ShadingModel != MSM_Unlit || Refraction.IsConnected();
 		break;
 	case MP_SubsurfaceColor:
-		Active = ShadingModel == MSM_Subsurface || ShadingModel == MSM_PreintegratedSkin;
+		Active = ShadingModel == MSM_Subsurface || ShadingModel == MSM_PreintegratedSkin || ShadingModel == MSM_TwoSidedFoliage;
 		break;
 	case MP_ClearCoat:
 	case MP_ClearCoatRoughness:
@@ -3759,4 +3760,4 @@ bool UMaterial::HasFlippedCoordinates()
 }
 #endif //WITH_EDITORONLY_DATA
 
-#undef LOCTEXT_NAMESPACE
\ No newline at end of file
+#undef LOCTEXT_NAMESPACE
diff --git a/Engine/Source/Runtime/Engine/Private/Materials/MaterialShader.cpp b/Engine/Source/Runtime/Engine/Private/Materials/MaterialShader.cpp
--- a/Engine/Source/Runtime/Engine/Private/Materials/MaterialShader.cpp
+++ b/Engine/Source/Runtime/Engine/Private/Materials/MaterialShader.cpp
@@ -41,6 +41,7 @@ FString GetShadingModelString(EMaterialShadingModel ShadingModel)
 		case MSM_PreintegratedSkin:	ShadingModelName = TEXT("MSM_PreintegratedSkin"); break;
 		case MSM_ClearCoat:			ShadingModelName = TEXT("MSM_ClearCoat"); break;
 		case MSM_SubsurfaceProfile:	ShadingModelName = TEXT("MSM_SubsurfaceProfile"); break;
+		case MSM_TwoSidedFoliage:	ShadingModelName = TEXT("MSM_TwoSidedFoliage"); break;
 		default: ShadingModelName = TEXT("Unknown"); break;
 	}
 	return ShadingModelName;
@@ -84,6 +85,7 @@ void UpdateMaterialShaderCompilingStats(const FMaterial* Material)
 		case MSM_PreintegratedSkin:
 		case MSM_ClearCoat:
 		case MSM_SubsurfaceProfile:
+		case MSM_TwoSidedFoliage:
 			INC_DWORD_STAT_BY(STAT_ShaderCompiling_NumLitMaterialShaders,1);
 			break;
 		default: break;
diff --git a/Engine/Source/Runtime/Engine/Private/Materials/MaterialShared.cpp b/Engine/Source/Runtime/Engine/Private/Materials/MaterialShared.cpp
--- a/Engine/Source/Runtime/Engine/Private/Materials/MaterialShared.cpp
+++ b/Engine/Source/Runtime/Engine/Private/Materials/MaterialShared.cpp
@@ -1240,6 +1240,7 @@ void FMaterial::SetupMaterialEnvironment(
 		case MSM_PreintegratedSkin: OutEnvironment.SetDefine(TEXT("MATERIAL_SHADINGMODEL_PREINTEGRATED_SKIN"),	TEXT("1")); break;
 		case MSM_SubsurfaceProfile: OutEnvironment.SetDefine(TEXT("MATERIAL_SHADINGMODEL_SUBSURFACE_PROFILE"),	TEXT("1")); break;
 		case MSM_ClearCoat:			OutEnvironment.SetDefine(TEXT("MATERIAL_SHADINGMODEL_CLEAR_COAT"),			TEXT("1")); break;
+		case MSM_TwoSidedFoliage:	OutEnvironment.SetDefine(TEXT("MATERIAL_SHADINGMODEL_TWOSIDED_FOLIAGE"),	TEXT("1")); break;
 		default: 
 			UE_LOG(LogMaterial, Warning, TEXT("Unknown material shading model: %u  Setting to MSM_DefaultLit"),(int32)GetShadingModel());
 			OutEnvironment.SetDefine(TEXT("MATERIAL_SHADINGMODEL_DEFAULT_LIT"),TEXT("1"));
diff --git a/Engine/Source/Runtime/Engine/Public/MaterialShared.h b/Engine/Source/Runtime/Engine/Public/MaterialShared.h
--- a/Engine/Source/Runtime/Engine/Public/MaterialShared.h
+++ b/Engine/Source/Runtime/Engine/Public/MaterialShared.h
@@ -49,6 +49,11 @@ DECLARE_LOG_CATEGORY_EXTERN(LogMaterial,Log,Verbose);
 /** Creates a string that represents the given quality level. */
 extern void GetMaterialQualityLevelName(EMaterialQualityLevel::Type InMaterialQualityLevel, FString& OutName);
 
+inline bool IsSubsurfaceShadingModel(EMaterialShadingModel ShadingModel)
+{
+	return ShadingModel == MSM_Subsurface || ShadingModel == MSM_PreintegratedSkin || ShadingModel == MSM_SubsurfaceProfile || ShadingModel == MSM_TwoSidedFoliage;
+}
+
 /**
  * The types which can be used by materials.
  */
