fixed UE-5560 Skeletal Mesh with Tessellated Material causes Excessive Blur

diff --git a/Engine/Shaders/VelocityShader.usf b/Engine/Shaders/VelocityShader.usf
--- a/Engine/Shaders/VelocityShader.usf
+++ b/Engine/Shaders/VelocityShader.usf
@@ -12,12 +12,16 @@
 #define GPUSKIN_PASS_THROUGH 0
 #endif
 
-/*=============================================================================
- * Vertex Shader
- *============================================================================*/
+// Move all geometry a little bit towards the camera to not get z fighting with existing depth
+// buffer of the same meshes rendered with slightly different float computations.
+// This is a positive number as our z buffer is inverted for better precision.
+
+// for tessellation
+static const float GTessellationDepthBias = 0.001f;
+// for non tessllation: 0.0f as the isolate code should do a good job to prevent this case (can be increased if that is not the case e.g. if the platform ignores the isolate)
+static const float GNonTessellationDepthBias = 0.0f;
 
-// Move all geometry a little bit towards the camera
-static const float GStretchTimeScale = -0.001f;
+//
 
 struct FVelocityInterpsVSToPS
 {
@@ -90,8 +94,8 @@ struct FVelocityVSToPS
 			O.Position = mul(WorldPosition, View.TranslatedWorldToClip);
 		}
 
-		// Move all geometry a little bit towards the camera.
-		O.Position.z += GStretchTimeScale * O.Position.w;
+		// Move all geometry a little bit towards the camera (to not get z fighting with existing zbuffer of the same meshes rendered with slightly different float computations).
+		O.Position.z += GTessellationDepthBias * O.Position.w;
 		
 		O.Interps.PackedVelocityA = Interpolants.Interps.PackedVelocityA;
 		O.Interps.PackedVelocityC = Interpolants.Interps.PackedVelocityC;
@@ -101,6 +105,10 @@ struct FVelocityVSToPS
 	
 	#include "Tessellation.usf"
 #endif
+	
+/*=============================================================================
+ * Vertex Shader
+ *============================================================================*/
 
 void MainVertexShader(
 	FVertexFactoryInput Input,
@@ -161,8 +169,8 @@ void MainVertexShader(
 	}
 
 #if !USING_TESSELLATION
-	// Move all geometry a little bit towards the camera.
-	//Output.Position.z += GStretchTimeScale * Output.Position.w;
+	// Move all geometry a little bit towards the camera (to not get z fighting with existing zbuffer of the same meshes rendered with slightly different float computations).
+	Output.Position.z += GNonTessellationDepthBias * Output.Position.w;
 #endif
 
 #if USING_TESSELLATION
