fixed typo

diff --git a/Engine/Shaders/VelocityShader.usf b/Engine/Shaders/VelocityShader.usf
--- a/Engine/Shaders/VelocityShader.usf
+++ b/Engine/Shaders/VelocityShader.usf
@@ -18,7 +18,7 @@
 
 // for tessellation
 static const float GTessellationDepthBias = 0.001f;
-// for non tessllation: 0.0f as the isolate code should do a good job to prevent this case (can be increased if that is not the case e.g. if the platform ignores the isolate)
+// for non tessellation: 0.0f as the isolate code should do a good job to prevent this case (can be increased if that is not the case e.g. if the platform ignores the isolate)
 static const float GNonTessellationDepthBias = 0.0f;
 
 //
