* fixed UE-3599 UE4: RENDERING: VisualizeMotionBlur is broken

diff --git a/Engine/Shaders/PostProcessMotionBlur.usf b/Engine/Shaders/PostProcessMotionBlur.usf
--- a/Engine/Shaders/PostProcessMotionBlur.usf
+++ b/Engine/Shaders/PostProcessMotionBlur.usf
@@ -307,6 +307,17 @@ uint GetStepCountFromQuality()
 //  - Removes the background blur (seems higher quality).
 #define MOTION_BLUR_EARLY_EXIT 1
 
+
+
+
+
+#if MOTION_BLUR_QUALITY == 0
+	// this would break VisualizeMotionBlur so we disable it
+	#undef MOTION_BLUR_EARLY_EXIT 
+	#define MOTION_BLUR_EARLY_EXIT 0
+#endif
+
+
 // motionblur pixel shader
 // input:
 //   0: RGB:scene color, A: depth in half resolution from MotionBlurSetup
