cleanup

diff --git a/Engine/Shaders/PostProcessMotionBlur.usf b/Engine/Shaders/PostProcessMotionBlur.usf
--- a/Engine/Shaders/PostProcessMotionBlur.usf
+++ b/Engine/Shaders/PostProcessMotionBlur.usf
@@ -20,13 +20,27 @@
 // from the paper: We use SOFT Z EXTENT = 1mm to 10cm for our results
 #define SOFT_Z_EXTENT 1
 
+// This,
+//  - Gets the MainPS pass to low cost under no motion.
+//  - Somewhat cleans up the blur/no blur transition.
+//  - Reduces the motion blur on small motion (where it is not needed).
+//  - Removes the background blur (seems higher quality).
+#define MOTION_BLUR_EARLY_EXIT 1
+
 // 0: use more texture lookups avoiding .Gather() lookups
 #define GATHER_OPTIMIZATION (FEATURE_LEVEL >= FEATURE_LEVEL_SM5)
 
-
 // 0:off / 1:on, useful to debug the motionblur algorithm
 #define MOTIONBLUR_TESTCHART 0
 
+// --------------------------------------------------------------------------
+
+#if MOTION_BLUR_QUALITY == 0
+	// this would break VisualizeMotionBlur so we disable it
+	#undef MOTION_BLUR_EARLY_EXIT 
+	#define MOTION_BLUR_EARLY_EXIT 0
+#endif
+
 // helper function, from the paper but adjusted to avoid division by zero
 float cone(float2 X, float2 Y, float2 v)
 {
@@ -300,24 +314,6 @@ uint GetStepCountFromQuality()
 #endif
 }
 
-// This,
-//  - Gets the MainPS pass to low cost under no motion.
-//  - Somewhat cleans up the blur/no blur transition.
-//  - Reduces the motion blur on small motion (where it is not needed).
-//  - Removes the background blur (seems higher quality).
-#define MOTION_BLUR_EARLY_EXIT 1
-
-
-
-
-
-#if MOTION_BLUR_QUALITY == 0
-	// this would break VisualizeMotionBlur so we disable it
-	#undef MOTION_BLUR_EARLY_EXIT 
-	#define MOTION_BLUR_EARLY_EXIT 0
-#endif
-
-
 // motionblur pixel shader
 // input:
 //   0: RGB:scene color, A: depth in half resolution from MotionBlurSetup
