- Moved  ES31_PROFILE and SM4_PROFILE defines from Metal shader compiler into Common.usf (they should be removed from there as well, and not depend on them when METAL_*PROFILE is set)

diff --git a/Engine/Shaders/Common.usf b/Engine/Shaders/Common.usf
--- a/Engine/Shaders/Common.usf
+++ b/Engine/Shaders/Common.usf
@@ -91,15 +91,11 @@
 
 	#define FEATURE_LEVEL FEATURE_LEVEL_SM5
 
-#elif SM5_PROFILE || SM4_PROFILE
-	 
+#elif SM5_PROFILE
 	// SM5 = full dx11 features (high end UE4 rendering)
-	#if SM5_PROFILE
-		#define FEATURE_LEVEL FEATURE_LEVEL_SM5
-	#elif  SM4_PROFILE
-		#define FEATURE_LEVEL FEATURE_LEVEL_SM4
-	#endif
-
+	#define FEATURE_LEVEL FEATURE_LEVEL_SM5
+#elif SM4_PROFILE
+	#define FEATURE_LEVEL FEATURE_LEVEL_SM4
 #elif GL3_PROFILE || GL4_PROFILE
 
 	#if GL3_PROFILE
@@ -118,9 +114,15 @@
 	#endif
 #elif METAL_PROFILE
 	#define FEATURE_LEVEL  FEATURE_LEVEL_ES3_1
+	// @todo metal: remove this and make sure all uses handle METAL_PROFILE
+	#undef ES3_1_PROFILE
+	#define ES3_1_PROFILE 1
 	#define FCOLOR_COMPONENT_SWIZZLE .bgra  
 #elif METAL_MRT_PROFILE
 	#define FEATURE_LEVEL FEATURE_LEVEL_SM4
+	// @todo metal: remove this and make sure all uses handle METAL_MRT_PROFILE
+	#undef SM4_PROFILE
+	#define SM4_PROFILE 1
 	#define FCOLOR_COMPONENT_SWIZZLE .bgra  
 #elif ES2_PROFILE || ES3_1_PROFILE
 	//@todo ES3_1 GL
@@ -153,7 +155,8 @@
 // Whether the platform supports independent texture and samplers
 // When enabled, different texture lookups can share samplers to allow more artist samplers in the base pass
 // Ideally this would just be enabled for all SM4 and above feature level platforms
-#define SUPPORTS_INDEPENDENT_SAMPLERS (PS4_PROFILE || SM5_PROFILE || SM4_PROFILE)
+// @todo metal mrt: No reason this can't work with Metal, once cross compiler is fixed
+#define SUPPORTS_INDEPENDENT_SAMPLERS ((PS4_PROFILE || SM5_PROFILE || SM4_PROFILE) && !METAL_MRT_PROFILE)
 
 // This would need to be a #define in GLSL to ignore the SamplerState, however, it is currently a function call in HLSL
 // for type checking of the parameters - ironically the type checking is really only needed in GLSL!
diff --git a/Engine/Source/Developer/iOS/MetalShaderFormat/Private/MetalShaderCompiler.cpp b/Engine/Source/Developer/iOS/MetalShaderFormat/Private/MetalShaderCompiler.cpp
--- a/Engine/Source/Developer/iOS/MetalShaderFormat/Private/MetalShaderCompiler.cpp
+++ b/Engine/Source/Developer/iOS/MetalShaderFormat/Private/MetalShaderCompiler.cpp
@@ -841,13 +841,11 @@ void CompileShader_Metal(const FShaderCompilerInput& Input,FShaderCompilerOutput
 
 	if (Input.ShaderFormat == NAME_SF_METAL)
 	{
-	AdditionalDefines.SetDefine(TEXT("METAL_PROFILE"), 1);
-	AdditionalDefines.SetDefine(TEXT("ES3_1_PROFILE"), 1);
+		AdditionalDefines.SetDefine(TEXT("METAL_PROFILE"), 1);
 	}
 	else if (Input.ShaderFormat == NAME_SF_METAL_MRT)
 	{
 		AdditionalDefines.SetDefine(TEXT("METAL_MRT_PROFILE"), 1);
-		AdditionalDefines.SetDefine(TEXT("SM4_PROFILE"), 1);
 	}
 	else
 	{
