Improved sky specular occlusion for the TwoSidedFoliage shading model - applies AO directly instead of doing the reflection cone vs bent normal intersection.

diff --git a/Engine/Shaders/ReflectionEnvironmentComputeShaders.usf b/Engine/Shaders/ReflectionEnvironmentComputeShaders.usf
--- a/Engine/Shaders/ReflectionEnvironmentComputeShaders.usf
+++ b/Engine/Shaders/ReflectionEnvironmentComputeShaders.usf
@@ -408,7 +408,7 @@ float CountOverlap( float3 WorldPosition )
 	return Overlap;
 }
 
-void GatherRadiance( inout float4 Color, float3 WorldPosition, float3 RayDirection, float Roughness, float2 ScreenPosition, float IndirectIrradiance )
+void GatherRadiance(inout float4 Color, float3 WorldPosition, float3 RayDirection, float Roughness, float2 ScreenPosition, float IndirectIrradiance, uint ShadingModelID)
 {
 	float Mip = ComputeReflectionCaptureMipFromRoughness( Roughness );
 
@@ -552,7 +552,7 @@ void GatherRadiance( inout float4 Color, float3 WorldPosition, float3 RayDirecti
 		}
 
 		float2 ScreenUV = ScreenPosition * View.ScreenPositionScaleBias.xy + View.ScreenPositionScaleBias.wz;
-		float Visibility = GetDistanceFieldAOSpecularOcclusion(ScreenUV, RayDirection, Roughness);
+		float Visibility = GetDistanceFieldAOSpecularOcclusion(ScreenUV, RayDirection, Roughness, ShadingModelID == SHADINGMODELID_TWOSIDED_FOLIAGE);
 
 		// TODO use Mip
 		Color.rgb += (Color.a * Visibility) * SkyLighting;
@@ -644,7 +644,7 @@ void ReflectionEnvironmentTiledDeferredMain(
 
 		Color.a *= SpecularOcclusion;
 
-		GatherRadiance( Color, WorldPosition, R, GBuffer.Roughness, ScreenPosition, IndirectIrradiance );
+		GatherRadiance(Color, WorldPosition, R, GBuffer.Roughness, ScreenPosition, IndirectIrradiance, GBuffer.ShadingModelID);
 
 		BRANCH
 		if( GBuffer.ShadingModelID == SHADINGMODELID_CLEAR_COAT )
@@ -671,7 +671,7 @@ void ReflectionEnvironmentTiledDeferredMain(
 			Color.a *= 1 - SSR.a;
 			Color.a *= SpecularOcclusion;
 
-			GatherRadiance( Color, WorldPosition, R, ClearCoatRoughness, ScreenPosition, IndirectIrradiance );
+			GatherRadiance(Color, WorldPosition, R, ClearCoatRoughness, ScreenPosition, IndirectIrradiance, GBuffer.ShadingModelID);
 		}
 		else
 		{
diff --git a/Engine/Shaders/ReflectionEnvironmentShaders.usf b/Engine/Shaders/ReflectionEnvironmentShaders.usf
--- a/Engine/Shaders/ReflectionEnvironmentShaders.usf
+++ b/Engine/Shaders/ReflectionEnvironmentShaders.usf
@@ -527,7 +527,7 @@ void ReflectionApplyPS(
 				SkyLighting *= GBuffer.IndirectIrradiance;
 			}
 
-			float Visibility = GetDistanceFieldAOSpecularOcclusion(UV, ReflectionVector, GBuffer.Roughness); 
+			float Visibility = GetDistanceFieldAOSpecularOcclusion(UV, ReflectionVector, GBuffer.Roughness, GBuffer.ShadingModelID == SHADINGMODELID_TWOSIDED_FOLIAGE); 
 
 			SpecularLighting.rgb += (Visibility * SpecularLighting.a) * SkyLighting;
 		}
diff --git a/Engine/Shaders/SkyLightingShared.usf b/Engine/Shaders/SkyLightingShared.usf
--- a/Engine/Shaders/SkyLightingShared.usf
+++ b/Engine/Shaders/SkyLightingShared.usf
@@ -21,7 +21,7 @@ SamplerState BentNormalAOSampler;
 
 float ApplyBentNormalAO;
 
-float GetDistanceFieldAOSpecularOcclusion(float2 ScreenUV, float3 ReflectionVector, float Roughness)
+float GetDistanceFieldAOSpecularOcclusion(float2 ScreenUV, float3 ReflectionVector, float Roughness, bool bTwoSidedFoliage)
 {
 	float Visibility = 1;
 
@@ -29,14 +29,23 @@ float GetDistanceFieldAOSpecularOcclusion(float2 ScreenUV, float3 ReflectionVect
 	if (ApplyBentNormalAO > 0)
 	{
 		float3 BentNormalAO = Texture2DSampleLevel(BentNormalAOTexture, BentNormalAOSampler, ScreenUV, 0).xyz;
-		float ReflectionConeAngle = Roughness * PI;
-		float BentNormalLength = length(BentNormalAO);
-		float UnoccludedAngle = BentNormalLength * PI;
-		float AngleBetween = acos(dot(BentNormalAO, ReflectionVector) / max(BentNormalLength, .001f));
-		Visibility = ApproximateConeConeIntersection(ReflectionConeAngle, UnoccludedAngle, AngleBetween);
-
-		// Can't rely on the direction of the bent normal when close to fully occluded, lerp to shadowed
-		Visibility = lerp(0, Visibility, saturate((UnoccludedAngle - .1f) / .2f));
+
+		BRANCH
+		if (bTwoSidedFoliage)
+		{
+			Visibility = length(BentNormalAO);
+		}
+		else
+		{
+			float ReflectionConeAngle = Roughness * PI;
+			float BentNormalLength = length(BentNormalAO);
+			float UnoccludedAngle = BentNormalLength * PI;
+			float AngleBetween = acos(dot(BentNormalAO, ReflectionVector) / max(BentNormalLength, .001f));
+			Visibility = ApproximateConeConeIntersection(ReflectionConeAngle, UnoccludedAngle, AngleBetween);
+
+			// Can't rely on the direction of the bent normal when close to fully occluded, lerp to shadowed
+			Visibility = lerp(0, Visibility, saturate((UnoccludedAngle - .1f) / .2f));
+		}
 	}
 
 	return Visibility;
