Rendering preview for ES2 feature level now clips editor primitives to world geometry.

Fixes UE-5451

diff --git a/Engine/Shaders/Common.usf b/Engine/Shaders/Common.usf
--- a/Engine/Shaders/Common.usf
+++ b/Engine/Shaders/Common.usf
@@ -147,6 +147,10 @@
 
 #endif
 
+#ifndef ES2_EMULATION
+#define ES2_EMULATION (!COMPILER_GLSL_ES2 && COMPILE_SHADERS_FOR_DEVELOPMENT && !METAL_PROFILE)
+#endif
+
 // If we didn't request color component swizzling, just make it empty
 #ifndef FCOLOR_COMPONENT_SWIZZLE
 #define FCOLOR_COMPONENT_SWIZZLE
@@ -623,7 +627,7 @@ float CalcSceneDepth( float2 ScreenUV )
 #endif
 }
 
-#if FEATURE_LEVEL >= FEATURE_LEVEL_SM4
+#if FEATURE_LEVEL >= FEATURE_LEVEL_SM4 || ES2_EMULATION
 
 	// depth in in DeviceZ
 	Texture2D<float> SceneDepthTextureNonMS;
diff --git a/Engine/Shaders/SimpleElementPixelShader.usf b/Engine/Shaders/SimpleElementPixelShader.usf
--- a/Engine/Shaders/SimpleElementPixelShader.usf
+++ b/Engine/Shaders/SimpleElementPixelShader.usf
@@ -10,6 +10,7 @@
 #define USE_IES_PROFILE 1
 #include "IESLightProfilesCommon.usf"
 
+#define ENABLE_EDITOR_COMPOSITING (FEATURE_LEVEL >= FEATURE_LEVEL_SM4 || ES2_EMULATION)
 
 Texture2D InTexture;
 SamplerState InTextureSampler;
@@ -34,18 +35,23 @@ MaterialFloat4 ColourTexture2DSample(Texture2D Tex, SamplerState Sampler, float2
 	return Sample;
 }
 
-#if FEATURE_LEVEL >= FEATURE_LEVEL_SM4
+#if ENABLE_EDITOR_COMPOSITING
 
 void PerformEditorCompositingDepthTest(float4 SvPosition)
 {
-	bool bIsPerspective = (View.ViewToClip._m33 < 1.0f);
-	// Using SV_Position because manual interpolation fails in precision for the editor grid.
-	float SceneDepthDeviceZ = SceneDepthTextureNonMS.Load(int3(SvPosition.xy, 0)).r;
+	// depth test manually if compositing editor primitives
+	FLATTEN
+	if (bEnableEditorPrimitiveDepthTest)
+	{
+		bool bIsPerspective = (View.ViewToClip._m33 < 1.0f);
+		// Using SV_Position because manual interpolation fails in precision for the editor grid.
+		float SceneDepthDeviceZ = SceneDepthTextureNonMS.Load(int3(SvPosition.xy, 0)).r;
 
-	float InterpolatedDeviceZ = SvPosition.z;
+		float InterpolatedDeviceZ = SvPosition.z;
 
-	// This is for a reversed Z depth buffer.
-	clip(InterpolatedDeviceZ - SceneDepthDeviceZ);
+		// This is for a reversed Z depth buffer.
+		clip(InterpolatedDeviceZ - SceneDepthDeviceZ);
+	}
 }
 
 #endif
@@ -57,17 +63,14 @@ void Main(
 	out float4 OutColor : SV_Target0
 #if FEATURE_LEVEL >= FEATURE_LEVEL_SM4
 	,out float4 OutWorldNormal : SV_Target1 
+#endif
+#if ENABLE_EDITOR_COMPOSITING
 	,in float4 SvPosition : SV_Position
 #endif
 	)
 {
-#if FEATURE_LEVEL >= FEATURE_LEVEL_SM4
-	// depth test manually if compositing editor primitives
-	FLATTEN
-	if (bEnableEditorPrimitiveDepthTest)
-	{
-		PerformEditorCompositingDepthTest(SvPosition);
-	}
+#if ENABLE_EDITOR_COMPOSITING
+	PerformEditorCompositingDepthTest(SvPosition);
 #endif
 	 
 	float4 BaseColor = ColourTexture2DSample(InTexture, InTextureSampler,TextureCoordinate);
@@ -88,17 +91,14 @@ void AlphaOnlyMain(
 	out float4 OutColor : SV_Target0
 #if FEATURE_LEVEL >= FEATURE_LEVEL_SM4
 	,out float4 OutWorldNormal : SV_Target1 
+#endif
+#if ENABLE_EDITOR_COMPOSITING
 	,in float4 SvPosition : SV_Position
 #endif
 	)
 {
-#if FEATURE_LEVEL >= FEATURE_LEVEL_SM4
-	// depth test manually if compositing editor primitives
-	FLATTEN
-	if (bEnableEditorPrimitiveDepthTest)
-	{
-		PerformEditorCompositingDepthTest(SvPosition);
-	}
+#if ENABLE_EDITOR_COMPOSITING
+	PerformEditorCompositingDepthTest(SvPosition);
 #endif
 	
 	float4 TmpColor = Color;
@@ -122,19 +122,16 @@ void GammaMain(
 	in float4 Color : TEXCOORD1,
 	in float4 InHitProxyId : TEXCOORD2,			// needed to have d3ddebug warning supressed (VS and PS signatures with SV_Position
 	out float4 OutColor : SV_Target0
-#if FEATURE_LEVEL >= FEATURE_LEVEL_SM4
+#if ENABLE_EDITOR_COMPOSITING
 	,in float4 SvPosition : SV_Position
+#endif
+#if FEATURE_LEVEL >= FEATURE_LEVEL_SM4
 	,out float4 OutWorldNormal : SV_Target1
 #endif
 	)
 {
-#if FEATURE_LEVEL >= FEATURE_LEVEL_SM4
-	// depth test manually if compositing editor primitives
-	FLATTEN
-	if (bEnableEditorPrimitiveDepthTest)
-	{
-		PerformEditorCompositingDepthTest(SvPosition);
-	}
+#if ENABLE_EDITOR_COMPOSITING
+	PerformEditorCompositingDepthTest(SvPosition);
 #endif
 
 	float4 BaseColor = ColourTexture2DSample(InTexture, InTextureSampler,TextureCoordinate);
@@ -158,19 +155,16 @@ void GammaAlphaOnlyMain(
 	in float4 Color : TEXCOORD1,
 	in float4 InHitProxyId : TEXCOORD2,			// needed to have d3ddebug warning supressed (VS and PS signatures with SV_Position
 	out float4 OutColor : SV_Target0
-#if FEATURE_LEVEL >= FEATURE_LEVEL_SM4
+#if ENABLE_EDITOR_COMPOSITING
 	,in float4 SvPosition : SV_Position
+#endif
+#if FEATURE_LEVEL >= FEATURE_LEVEL_SM4
 	,out float4 OutWorldNormal : SV_Target1
 #endif
 	)
 {
-#if FEATURE_LEVEL >= FEATURE_LEVEL_SM4
-	// depth test manually if compositing editor primitives
-	FLATTEN
-	if (bEnableEditorPrimitiveDepthTest)
-	{
-		PerformEditorCompositingDepthTest(SvPosition);
-	}
+#if ENABLE_EDITOR_COMPOSITING
+	PerformEditorCompositingDepthTest(SvPosition);
 #endif
 
 	float4 TmpColor = Color;
@@ -195,19 +189,16 @@ void GammaMaskedMain(
 	in float4 Color : TEXCOORD1,
 	in float4 HitProxyId : TEXCOORD2,
 	out float4 OutColor : SV_Target0
-#if FEATURE_LEVEL >= FEATURE_LEVEL_SM4
+#if ENABLE_EDITOR_COMPOSITING
 	,in float4 SvPosition : SV_Position
+#endif
+#if FEATURE_LEVEL >= FEATURE_LEVEL_SM4
 	,out float4 OutWorldNormal : SV_Target1
 #endif
 	)
 {
-#if FEATURE_LEVEL >= FEATURE_LEVEL_SM4
-	// depth test manually if compositing editor primitives
-	FLATTEN
-	if (bEnableEditorPrimitiveDepthTest)
-	{
-		PerformEditorCompositingDepthTest(SvPosition);
-	}
+#if ENABLE_EDITOR_COMPOSITING
+	PerformEditorCompositingDepthTest(SvPosition);
 #endif
 
 	float4 BaseColor = ColourTexture2DSample(InTexture, InTextureSampler,TextureCoordinate);
diff --git a/Engine/Source/Runtime/Renderer/Private/BasePassRendering.h b/Engine/Source/Runtime/Renderer/Private/BasePassRendering.h
--- a/Engine/Source/Runtime/Renderer/Private/BasePassRendering.h
+++ b/Engine/Source/Runtime/Renderer/Private/BasePassRendering.h
@@ -343,8 +343,11 @@ public:
 			{
 				TranslucentLightingParameters.Set(RHICmdList, this, View);
 			}
+		}
 
 #if WITH_EDITOR
+		if ( GMaxRHIFeatureLevel >= ERHIFeatureLevel::SM4 )
+		{
 			if( MaterialResource.IsUsedWithEditorCompositing() )
 			{
 				// Compute parameters for converting from screen space to pixel 
@@ -354,7 +357,7 @@ public:
 
 				FVector4 ScreenPosToPixelValue(
 					ViewportExtent.X * 0.5f,
-					-ViewportExtent.Y * 0.5f, 
+					-ViewportExtent.Y * 0.5f,
 					ViewportExtent.X * 0.5f - 0.5f + ViewportOffset.X,
 					ViewportExtent.Y * 0.5f - 0.5f + ViewportOffset.Y);
 
@@ -375,8 +378,8 @@ public:
 				SetShaderValue(RHICmdList, ShaderRHI, EditorCompositeDepthTestParameter, bEnableEditorPrimitveDepthTest );
 				SetShaderValue(RHICmdList, ShaderRHI, MSAASampleCount, GSceneRenderTargets.EditorPrimitivesColor ? GSceneRenderTargets.EditorPrimitivesColor->GetDesc().NumSamples : 0 );
 			}
-#endif
 		}
+#endif
 	}
 
 	void SetMesh(FRHICommandList& RHICmdList, const FVertexFactory* VertexFactory,const FSceneView& View,const FPrimitiveSceneProxy* Proxy,const FMeshBatchElement& BatchElement, EBlendMode BlendMode)
