ES2_EMULATION define ensures es2 featurelevel is set.

diff --git a/Engine/Shaders/Common.usf b/Engine/Shaders/Common.usf
--- a/Engine/Shaders/Common.usf
+++ b/Engine/Shaders/Common.usf
@@ -148,7 +148,7 @@
 #endif
 
 #ifndef ES2_EMULATION
-#define ES2_EMULATION (!COMPILER_GLSL_ES2 && COMPILE_SHADERS_FOR_DEVELOPMENT && !METAL_PROFILE)
+#define ES2_EMULATION ((FEATURE_LEVEL == FEATURE_LEVEL_ES2) && (!COMPILER_GLSL_ES2 && COMPILE_SHADERS_FOR_DEVELOPMENT && !METAL_PROFILE))
 #endif
 
 // If we didn't request color component swizzling, just make it empty
diff --git a/Engine/Shaders/MaterialTemplate.usf b/Engine/Shaders/MaterialTemplate.usf
--- a/Engine/Shaders/MaterialTemplate.usf
+++ b/Engine/Shaders/MaterialTemplate.usf
@@ -554,7 +554,9 @@ MaterialFloat2 GetParticleMacroUV(FMaterialPixelParameters Parameters)
 	return (Parameters.ScreenPosition.xy / Parameters.ScreenPosition.w - Parameters.Particle.MacroUV.xy) * Parameters.Particle.MacroUV.zw + MaterialFloat2(.5, .5);
 }
 
-#define ES2_EMULATION (!COMPILER_GLSL_ES2 && COMPILE_SHADERS_FOR_DEVELOPMENT && !METAL_PROFILE)
+#ifndef ES2_EMULATION
+#define ES2_EMULATION ((FEATURE_LEVEL == FEATURE_LEVEL_ES2) && (!COMPILER_GLSL_ES2 && COMPILE_SHADERS_FOR_DEVELOPMENT && !METAL_PROFILE))
+#endif
 
 MaterialFloat4 ProcessMaterialColorTextureLookup(MaterialFloat4 TextureValue)
 {
