When rendering point light projection the Apple Intel HD driver/GLSL compiler fails to correctly index into the uniform array using the dynamically calculated CubeFaceIndex variable. On Mac OS X we have to use more branches and constant index values to avoid this. This was filed as radr://19190281 "Intel HD 4000 fails to index uniform array with dynamic variable, causes point lights to render incorrectly in UE4"

diff --git a/Engine/Shaders/ShadowProjectionCommon.usf b/Engine/Shaders/ShadowProjectionCommon.usf
--- a/Engine/Shaders/ShadowProjectionCommon.usf
+++ b/Engine/Shaders/ShadowProjectionCommon.usf
@@ -135,9 +135,46 @@ float CubemapHardwarePCF(float3 WorldPosition, float3 LightPosition, float Light
 		UpVector *= InvShadowmapResolution;
 
 		// Figure out which cube face we're sampling from
-		int CubeFaceIndex = 0;
 		float3 AbsLightVector = abs(LightVector);
 		float MaxCoordinate = max(AbsLightVector.x, max(AbsLightVector.y, AbsLightVector.z));
+#if MAC // @todo Remove once there is a fix for the Apple Intel HD driver/GLSL compiler failing to correctly index into the uniform array using the dynamically calculated CubeFaceIndex variable. radr://19190281 "Intel HD 4000 fails to index uniform array with dynamic variable, causes point lights to render incorrectly in UE4"
+		// Transform the position into shadow space
+		float4 ShadowPosition;
+		if (MaxCoordinate == AbsLightVector.x)
+		{
+			if(AbsLightVector.x == LightVector.x)
+			{
+				ShadowPosition = mul(float4(WorldPosition.xyz, 1), ShadowViewProjectionMatrices[0]);
+			}
+			else
+			{
+				ShadowPosition = mul(float4(WorldPosition.xyz, 1), ShadowViewProjectionMatrices[1]);
+			}
+		}
+		else if (MaxCoordinate == AbsLightVector.y)
+		{
+			if(AbsLightVector.y == LightVector.y)
+			{
+				ShadowPosition = mul(float4(WorldPosition.xyz, 1), ShadowViewProjectionMatrices[2]);
+			}
+			else
+			{
+				ShadowPosition = mul(float4(WorldPosition.xyz, 1), ShadowViewProjectionMatrices[3]);
+			}
+		}
+		else
+		{
+			if(AbsLightVector.z == LightVector.z)
+			{
+				ShadowPosition = mul(float4(WorldPosition.xyz, 1), ShadowViewProjectionMatrices[4]);
+			}
+			else
+			{
+				ShadowPosition = mul(float4(WorldPosition.xyz, 1), ShadowViewProjectionMatrices[5]);
+			}
+		}
+#else
+		int CubeFaceIndex = 0;
 		if (MaxCoordinate == AbsLightVector.x)
 		{
 			CubeFaceIndex = AbsLightVector.x == LightVector.x ? 0 : 1;
@@ -153,6 +190,7 @@ float CubemapHardwarePCF(float3 WorldPosition, float3 LightPosition, float Light
 
 		// Transform the position into shadow space
 		float4 ShadowPosition = mul(float4(WorldPosition.xyz, 1), ShadowViewProjectionMatrices[CubeFaceIndex]);
+#endif
 		// Calculate the Z buffer value that would have been stored for this position in the shadow map
 		float CompareDistance = ShadowPosition.z / ShadowPosition.w;
 		float ShadowDepthBias = - DepthBias / ShadowPosition.w;
