better quality for SubsurfaceProfile, full res for the non subsurface scattering light contribution (appears sharper but can be tweaked very quickly changing the SubSurfaceColor brightness)

diff --git a/Engine/Shaders/PostProcessSubsurface.usf b/Engine/Shaders/PostProcessSubsurface.usf
--- a/Engine/Shaders/PostProcessSubsurface.usf
+++ b/Engine/Shaders/PostProcessSubsurface.usf
@@ -17,6 +17,9 @@ SamplerState SSProfilesTextureSampler;
 // x:Radius*DistanceToProjectionWindow/KernelSize*0.5, y:DistanceToProjectionWindow, zw:xy scale correction for BufferSize!=ViewportSize
 float4 SSSParams;
 
+// 0: (testing, could be faster), 1: higher quality
+#define FULLRES_WHERE_POSSIBLE 1
+
 // ------------------------------------------
 
 // setup for "SeparableSSS.usf"
@@ -260,6 +263,41 @@ float4 SetupSubsurfaceForOnePixel(float2 UV)
 	return Ret;
 }
 
+void ExtractSpecularPS(float4 UVAndScreenPos : TEXCOORD0, out float4 OutColor : SV_Target0)
+{
+	float2 GBufferUV = UVAndScreenPos.xy;
+
+	float4 SceneColor4 = Texture2DSample(PostprocessInput0, PostprocessInput0Sampler, GBufferUV);
+
+	FScreenSpaceData ScreenSpaceData = GetScreenSpaceData(GBufferUV);
+
+	// 0: dont' fade in subsurface half res,  1:fully fade in half res
+	float BlendFactor = (ScreenSpaceData.GBuffer.ShadingModelID == SHADINGMODELID_SUBSURFACE_PROFILE);
+
+	float3 ExtractedSpecular = 0;
+	 
+#if SSS_RECOMBINE_METHOD
+	// we took the specular highlights out, now we add them back in
+	ExtractedSpecular += LightAccumulator_ReconstructNonDiffuseLighting(SceneColor4);
+#endif
+
+	// alpha channel can be anything - we consumed the value and likely output to a format that doesn't need the value
+	OutColor = float4(lerp(SceneColor4.rgb, ExtractedSpecular, BlendFactor), 0);
+
+#if FULLRES_WHERE_POSSIBLE
+	// 0..255, which SubSurface profile to pick
+	uint SubsurfaceProfileInt = ExtractSubsurfaceProfileInt(ScreenSpaceData.GBuffer);
+
+	// reference
+//	float3 DiffuseLighting = LightAccumulator_ReconstructDiffuseLighting(SceneColor4);
+	// cheaper:
+	float3 DiffuseLighting = SceneColor4 - ExtractedSpecular;
+
+	// add center sample from full resolution input for better quality
+	OutColor.rgb += DiffuseLighting * GetKernel(SSSS_N_KERNELWEIGHTOFFSET, SubsurfaceProfileInt).rgb;
+#endif
+}
+
 void SetupPS(in float4 UVAndScreenPos : TEXCOORD0, out float4 OutColor : SV_Target0)
 {
 	float2 UV = UVAndScreenPos.xy;
@@ -347,10 +385,16 @@ void MainPS(float4 UVAndScreenPos : TEXCOORD0, out float4 OutColor : SV_Target0)
 	// can be optimized
 	float2 GBufferUV = (ViewportUV * PostprocessInput0Size.xy * 2 + View.ViewRectMin.xy) * View.ViewSizeAndSceneTexelSize.zw;
 
-	OutColor = SSSSBlurPS(GBufferUV, ViewportUV, Direction, false);
+	bool bIncludeCenter = true;
+
+#if FULLRES_WHERE_POSSIBLE && SSS_DIRECTION == 1
+	bIncludeCenter = false;
+#endif
+
+	OutColor = SSSSBlurPS(GBufferUV, ViewportUV, Direction, false, bIncludeCenter);
 
 #if SSS_DIRECTION == 1
-	// second pass prepares the setup fro the recombine pass which doesn't need depth but wants to reconstruct the color
+	// second pass prepares the setup from the recombine pass which doesn't need depth but wants to reconstruct the color
 	OutColor.a = ComputeMaskFromDepthInAlpha(OutColor.a);
 #endif
 }
@@ -381,12 +425,6 @@ void SubsurfaceRecombinePS(float4 UVAndScreenPos : TEXCOORD0, out float4 OutColo
 
 		// renormalize to dilate RGB to fix half res upsampling artifacts
 		SSSColor = SSSColorWithAlpha.rgb / max(SSSColorWithAlpha.a, 0.00001f);
-
-		// 0..255, which SubSurface profile to pick
-//		uint SubsurfaceProfileInt = ExtractSubsurfaceProfileInt(ScreenSpaceData.GBuffer);
-
-		// add center sample from full resolution input for better quality
-//		SSSColor += SceneColor4.rgb * GetKernel(SSSS_N_KERNELWEIGHTOFFSET, SubsurfaceProfileInt).rgb;
 	}
 
 	// 0: don't fade in subsurface half res,  1:fully fade in half res
@@ -406,25 +444,3 @@ void SubsurfaceRecombinePS(float4 UVAndScreenPos : TEXCOORD0, out float4 OutColo
 	OutColor = float4(SSSColor * BlendFactor + SceneColor4.rgb, LerpFactor);
 }
 
-
-void ExtractSpecularPS(float4 UVAndScreenPos : TEXCOORD0, out float4 OutColor : SV_Target0)
-{
-	float2 GBufferUV = UVAndScreenPos.xy;
-
-	float4 SceneColor4 = Texture2DSample(PostprocessInput0, PostprocessInput0Sampler, GBufferUV);
-
-	FScreenSpaceData ScreenSpaceData = GetScreenSpaceData(GBufferUV);
-
-	// 0: dont' fade in subsurface half res,  1:fully fade in half res
-	float BlendFactor = (ScreenSpaceData.GBuffer.ShadingModelID == SHADINGMODELID_SUBSURFACE_PROFILE);
-
-	float3 ExtractedSpecular = 0;
-	 
-#if SSS_RECOMBINE_METHOD
-	// we took the specular highlights out, now we add them back in
-	ExtractedSpecular += LightAccumulator_ReconstructNonDiffuseLighting(SceneColor4);
-#endif
-
-	// alpha channel can be anything - we consumed the value and likely output to a format that doesn't need the value
-	OutColor = float4(lerp(SceneColor4.rgb, ExtractedSpecular, BlendFactor), 0);
-}
diff --git a/Engine/Shaders/SeparableSSS.usf b/Engine/Shaders/SeparableSSS.usf
--- a/Engine/Shaders/SeparableSSS.usf
+++ b/Engine/Shaders/SeparableSSS.usf
@@ -186,7 +186,7 @@ float3 SSSSTransmittance(float translucency, float sssWidth, float3 worldPositio
 // @param texcoord The usual quad texture coordinates.
 // @param dir Direction of the blur: First pass:   float2(1.0, 0.0), Second pass:  float2(0.0, 1.0)
 // @param initStencil indicates whether the stencil buffer should be initialized. Should be set to 'true' for the first pass if not previously initialized, to enable optimization of the second pass.
-float4 SSSSBlurPS(float2 GBufferUV, float2 ViewportUV, float2 dir, bool initStencil)
+float4 SSSSBlurPS(float2 GBufferUV, float2 ViewportUV, float2 dir, bool initStencil, bool bIncludeCenter)
 {
     // Fetch color of current pixel:
     float4 colorM = SSSSSampleSceneColorPoint(ViewportUV);
@@ -232,9 +232,10 @@ float4 SSSSBlurPS(float2 GBufferUV, float2 ViewportUV, float2 dir, bool initSten
 	// center sample is added in Recombine pass
 	colorBlurred.a = colorM.a;
 
-//#if SSS_DIRECTION == 0
-	colorBlurred.rgb = colorM.rgb * GetKernel(SSSS_N_KERNELWEIGHTOFFSET, SubsurfaceProfileInt).rgb;
-//#endif
+	if(bIncludeCenter)
+	{
+		colorBlurred.rgb = colorM.rgb * GetKernel(SSSS_N_KERNELWEIGHTOFFSET, SubsurfaceProfileInt).rgb;
+	}
 
 	// Accumulate the other samples:
     SSSS_UNROLL
diff --git a/Engine/Source/Runtime/Renderer/Private/CompositionLighting/CompositionLighting.cpp b/Engine/Source/Runtime/Renderer/Private/CompositionLighting/CompositionLighting.cpp
--- a/Engine/Source/Runtime/Renderer/Private/CompositionLighting/CompositionLighting.cpp
+++ b/Engine/Source/Runtime/Renderer/Private/CompositionLighting/CompositionLighting.cpp
@@ -366,6 +366,7 @@ void FCompositionLighting::ProcessLighting(FRHICommandListImmediate& RHICmdList,
 
 				FRenderingCompositePass* Pass1 = Context.Graph.RegisterPass(new(FMemStack::Get()) FRCPassPostProcessSubsurface(1));
 				Pass1->SetInput(ePId_Input0, Pass0);
+				Pass1->SetInput(ePId_Input1, PassSetup);
 
 				// full res composite pass, no blurring (Radius=0)
 				FRenderingCompositePass* RecombinePass = Context.Graph.RegisterPass(new(FMemStack::Get()) FRCPassPostProcessSubsurfaceRecombine());
diff --git a/Engine/Source/Runtime/Renderer/Private/PostProcess/PostProcessSubsurface.cpp b/Engine/Source/Runtime/Renderer/Private/PostProcess/PostProcessSubsurface.cpp
--- a/Engine/Source/Runtime/Renderer/Private/PostProcess/PostProcessSubsurface.cpp
+++ b/Engine/Source/Runtime/Renderer/Private/PostProcess/PostProcessSubsurface.cpp
@@ -458,6 +458,7 @@ class FPostProcessSubsurfaceExtractSpecularPS : public FGlobalShader
 public:
 	FPostProcessPassParameters PostprocessParameter;
 	FDeferredPixelShaderParameters DeferredParameters;
+	FSubsurfaceParameters SubsurfaceParameters;
 
 	/** Initialization constructor. */
 	FPostProcessSubsurfaceExtractSpecularPS(const ShaderMetaType::CompiledShaderInitializerType& Initializer)
@@ -465,6 +466,7 @@ public:
 	{
 		PostprocessParameter.Bind(Initializer.ParameterMap);
 		DeferredParameters.Bind(Initializer.ParameterMap);
+		SubsurfaceParameters.Bind(Initializer.ParameterMap);
 	}
 
 	void SetParameters(const FRenderingCompositePassContext& Context)
@@ -475,13 +477,14 @@ public:
 		FGlobalShader::SetParameters(Context.RHICmdList, ShaderRHI, Context.View);
 		PostprocessParameter.SetPS(ShaderRHI, Context, TStaticSamplerState<SF_Point,AM_Clamp,AM_Clamp,AM_Clamp>::GetRHI());
 		DeferredParameters.Set(Context.RHICmdList, ShaderRHI, Context.View);
+		SubsurfaceParameters.SetParameters(Context.RHICmdList, ShaderRHI, Context);
 	}
 
 	// FShader interface.
 	virtual bool Serialize(FArchive& Ar)
 	{
 		bool bShaderHasOutdatedParameters = FGlobalShader::Serialize(Ar);
-		Ar << PostprocessParameter << DeferredParameters;
+		Ar << PostprocessParameter << DeferredParameters << SubsurfaceParameters;
 		return bShaderHasOutdatedParameters;
 	}
 
diff --git a/Engine/Source/Runtime/Renderer/Private/PostProcess/PostProcessSubsurface.h b/Engine/Source/Runtime/Renderer/Private/PostProcess/PostProcessSubsurface.h
--- a/Engine/Source/Runtime/Renderer/Private/PostProcess/PostProcessSubsurface.h
+++ b/Engine/Source/Runtime/Renderer/Private/PostProcess/PostProcessSubsurface.h
@@ -56,8 +56,9 @@ public:
 
 // derives from TRenderingCompositePassBase<InputCount, OutputCount>
 // ePId_Input0: SceneColor (horizontal blur) or the pass before (vertical blur)
+// ePId_Input1: optional Setup pass (only for InDirection==1)
 // modifies SceneColor, uses some GBuffer attributes
-class FRCPassPostProcessSubsurface : public TRenderingCompositePassBase<1, 1>
+class FRCPassPostProcessSubsurface : public TRenderingCompositePassBase<2, 1>
 {
 public:
 	// constructor
