bumped shaderversion in case that was the casue of the recent discoloration

diff --git a/Engine/Shaders/ShaderVersion.usf b/Engine/Shaders/ShaderVersion.usf
--- a/Engine/Shaders/ShaderVersion.usf
+++ b/Engine/Shaders/ShaderVersion.usf
@@ -2,4 +2,4 @@
 // Each time the console command is executed it generates a new GUID. As this file is included
 // in common.usf (which should be included in any shader) it allows to invalidate the shader DDC.
 // 
-// GUID = {CAC8D6F0-2E83-4C3D-A767-584D1E3DC80E}
+// GUID = 0BD018124CF4B0CAD25B178DB468DF1B
\ No newline at end of file
