temp fix for particle flickering on ps4.

diff --git a/Engine/Shaders/MaterialTemplate.usf b/Engine/Shaders/MaterialTemplate.usf
--- a/Engine/Shaders/MaterialTemplate.usf
+++ b/Engine/Shaders/MaterialTemplate.usf
@@ -1370,12 +1370,16 @@ void CalcMaterialParameters(
 
 	Parameters.SVPosition = SvPosition;
 
+#if INTERPOLATER_DERIVED_SCREENPOSITION
+	Parameters.ScreenPosition = mul(PixelPosition, View.TranslatedWorldToClip);
+#else		
 	{
 		float2 PixelPos = SvPosition - View.ViewRectMin.xy;	
 		Parameters.ScreenPosition.xy = ( PixelPos / View.ViewSizeAndSceneTexelSize.xy - 0.5 ) * float2(2,-2);
 		Parameters.ScreenPosition.zw = SvPosition.zw;
 		Parameters.ScreenPosition.xyz *= Parameters.ScreenPosition.w;
 	}
+#endif
 
 	// PixelPosition is the world position translated to the camera position, which is just -CameraVector
 	Parameters.CameraVector = normalize(-PixelPosition.xyz);
