fixed opengl compiling (don't support new experiment mode yet on opengl)

diff --git a/Engine/Shaders/ForEachLight.usf b/Engine/Shaders/ForEachLight.usf
--- a/Engine/Shaders/ForEachLight.usf
+++ b/Engine/Shaders/ForEachLight.usf
@@ -33,6 +33,9 @@ float3 AccumulateForwardLights(FMaterialPixelParameters Parameters, float3 Diffu
 {
 	float3 Ret = 0;
 
+	// not yet supported on opengl (firstbitlow() does not exist)
+#ifndef OPENGL
+
 	// Which tile we are?
 	int2 Tile = int2((Parameters.SVPosition.xy - View.ViewRectMin.xy) * ForwardLightData.InvTileSize);
 
@@ -62,6 +65,8 @@ float3 AccumulateForwardLights(FMaterialPixelParameters Parameters, float3 Diffu
 
 		Ret += LightingForOneLight(Parameters, LightData, ScreenSpaceData);
 	}
+#endif
+
 	return Ret;
 }
 
