Feature: Metallic channel is now supported for non lit materials e.g. SubsurfaceProfile

diff --git a/Engine/Shaders/DeferredShadingCommon.usf b/Engine/Shaders/DeferredShadingCommon.usf
--- a/Engine/Shaders/DeferredShadingCommon.usf
+++ b/Engine/Shaders/DeferredShadingCommon.usf
@@ -333,9 +333,7 @@ void EncodeGBuffer(
 	}
 	else
 	{
-		bool bSubsurface = IsSubsurfaceModel(GBuffer.ShadingModelID);
-
-		OutGBufferB.r = bSubsurface ? GBuffer.Opacity : GBuffer.Metallic;
+		OutGBufferB.r = GBuffer.Metallic;
 		OutGBufferB.g = GBuffer.Specular;
 		OutGBufferB.b = GBuffer.Roughness;
 		OutGBufferC.rgb = EncodeBaseColor(GBuffer.BaseColor);
@@ -344,8 +342,7 @@ void EncodeGBuffer(
 		OutGBufferC.a = EncodeIndirectIrradiance(GBuffer.IndirectIrradiance * GBuffer.GBufferAO) + QuantizationBias * (1.0 / 255.0);
 		OutGBufferB.a = GBuffer.ShadingModelID * (1.0 / 255.0);
 
-		// unused, can be used to make a 4th CustomData channel
-		OutGBufferD.r = 0;
+		OutGBufferD.r = GBuffer.Opacity;
 		OutGBufferD.gba = GBuffer.CustomData;
 	}
 
@@ -377,12 +374,10 @@ FGBufferData DecodeGBufferData(
 	GBuffer.HasDistanceFieldRepresentation = PackedAlpha & 1;
 	GBuffer.ShadingModelID = DecodeShadingModelId(InGBufferB.a);
 	
-	bool bSubsurface = IsSubsurfaceModel(GBuffer.ShadingModelID);
-
 	GBuffer.BaseColor = DecodeBaseColor(InGBufferC.rgb);
-	GBuffer.Metallic = bSubsurface ? 0 : InGBufferB.r;
+	GBuffer.Metallic = InGBufferB.r;
 	GBuffer.Specular = InGBufferB.g;
-	GBuffer.Opacity = bSubsurface ? InGBufferB.r : 0;
+	GBuffer.Opacity = InGBufferD.r;
 	
 	GBuffer.Roughness = InGBufferB.b;
 	GBuffer.CustomData = InGBufferD.gba;
diff --git a/Engine/Source/Runtime/Engine/Private/Materials/Material.cpp b/Engine/Source/Runtime/Engine/Private/Materials/Material.cpp
--- a/Engine/Source/Runtime/Engine/Private/Materials/Material.cpp
+++ b/Engine/Source/Runtime/Engine/Private/Materials/Material.cpp
@@ -3597,7 +3597,7 @@ bool UMaterial::IsPropertyActive(EMaterialProperty InProperty) const
 		break;
 	case MP_Metallic:
 		// Subsurface models store opacity in place of Metallic in the GBuffer
-		Active = ShadingModel != MSM_Unlit && !IsSubsurfaceShadingModel(ShadingModel);
+		Active = ShadingModel != MSM_Unlit;
 		break;
 	case MP_Normal:
 		Active = ShadingModel != MSM_Unlit || Refraction.IsConnected();
