Improved tessellation performance. Raised max tessellation factor to 15.

diff --git a/Engine/Shaders/FlatTessellation.usf b/Engine/Shaders/FlatTessellation.usf
--- a/Engine/Shaders/FlatTessellation.usf
+++ b/Engine/Shaders/FlatTessellation.usf
@@ -59,17 +59,11 @@ struct FFlatTessellationHSToDS
 		TessellationMultipliers.x = 0.5 * (I[0].TessellationMultiplier + I[1].TessellationMultiplier);
 		TessellationMultipliers.y = 0.5 * (I[1].TessellationMultiplier + I[2].TessellationMultiplier);
 		TessellationMultipliers.z = 0.5 * (I[2].TessellationMultiplier + I[0].TessellationMultiplier);
-		TessellationMultipliers.w = (I[0].TessellationMultiplier + I[1].TessellationMultiplier + I[2].TessellationMultiplier) / 3.0;
+		TessellationMultipliers.w = 0.333 * (I[0].TessellationMultiplier + I[1].TessellationMultiplier + I[2].TessellationMultiplier);
 
 		//@todo: This should be vectorized and done in the MainHull shader.
 		float4 CompositeTessellationFactors = TessellationMultipliers * CalculateCompositeTessellationFactors(WorldPostion0,WorldPostion1,WorldPostion2);
-
-		// Limit the tessellation factor to <= 8. The multiplies and saturates handle any potential NANs/INFs in CompositeTessellationFactors.
-		float4 MaxTessFactor = 8.0f;
-		float4 InvMaxTessFactor = 1.0f / MaxTessFactor;
-		
-		// clamp at 1 to prevent triangle to disappear with 
-		CompositeTessellationFactors = max(1, MaxTessFactor * saturate(CompositeTessellationFactors * InvMaxTessFactor));
+		CompositeTessellationFactors = clamp( CompositeTessellationFactors, 1, 15 );
 
 		O.TessFactor[0] = CompositeTessellationFactors.x; // 1->2 edge
 		O.TessFactor[1] = CompositeTessellationFactors.y; // 2->0 edge
@@ -83,6 +77,7 @@ struct FFlatTessellationHSToDS
 	[domain("tri")]
 	[patchconstantfunc("HS_FlatTrianglesConstant")]
 	[outputcontrolpoints(3)]
+	[maxtessfactor(15)]
 	TESSELLATION_ATTRIBUTES
 	FFlatTessellationHSToDS MainHull( InputPatch<FPassSpecificVSToDS, IN_CONTROL_POINTS> I, uint ControlPointID : SV_OutputControlPointID )
 	{
diff --git a/Engine/Shaders/PNTriangles.usf b/Engine/Shaders/PNTriangles.usf
--- a/Engine/Shaders/PNTriangles.usf
+++ b/Engine/Shaders/PNTriangles.usf
@@ -77,17 +77,11 @@ float4 ComputeControlPoint(float4 PositionA, float4 PositionB, float4 NormalA)
 		TessellationMultipliers.x = 0.5 * (I[0].TessellationMultiplier + I[1].TessellationMultiplier);
 		TessellationMultipliers.y = 0.5 * (I[1].TessellationMultiplier + I[2].TessellationMultiplier);
 		TessellationMultipliers.z = 0.5 * (I[2].TessellationMultiplier + I[0].TessellationMultiplier);
-		TessellationMultipliers.w = (I[0].TessellationMultiplier + I[1].TessellationMultiplier + I[2].TessellationMultiplier) / 3.0;
+		TessellationMultipliers.w = 0.333 * (I[0].TessellationMultiplier + I[1].TessellationMultiplier + I[2].TessellationMultiplier);
 
 		//@todo: This should be vectorized and done in the MainHull shader.
 		float4 CompositeTessellationFactors = TessellationMultipliers * CalculateCompositeTessellationFactors(B300.xyz,B030.xyz,B003.xyz);
-
-		// Limit the tessellation factor to <= 8. The multiplies and saturates handle any potential NANs/INFs in CompositeTessellationFactors.
-		float4 MaxTessFactor = 8.0f;
-		float4 InvMaxTessFactor = 1.0f / MaxTessFactor;
-
-		// clamp at 1 to prevent triangle to disappear with 
-		CompositeTessellationFactors = max(1, MaxTessFactor * saturate(CompositeTessellationFactors * InvMaxTessFactor));
+		CompositeTessellationFactors = clamp( CompositeTessellationFactors, 1, 15 );
 
 		O.TessFactor[0] = CompositeTessellationFactors.x; // 1->2 edge
 		O.TessFactor[1] = CompositeTessellationFactors.y; // 2->0 edge
@@ -101,6 +95,7 @@ float4 ComputeControlPoint(float4 PositionA, float4 PositionB, float4 NormalA)
 	[domain("tri")]
 	[patchconstantfunc("PNSharedConstantHull")]
 	[outputcontrolpoints(OUT_CONTROL_POINTS)]
+	[maxtessfactor(15)]
 	TESSELLATION_ATTRIBUTES
 	FPNTessellationHSToDS MainHull( InputPatch<FPassSpecificVSToDS, IN_CONTROL_POINTS> I, uint ControlPointID : SV_OutputControlPointID )
 	{
diff --git a/Engine/Shaders/Tessellation.usf b/Engine/Shaders/Tessellation.usf
--- a/Engine/Shaders/Tessellation.usf
+++ b/Engine/Shaders/Tessellation.usf
@@ -62,13 +62,12 @@ float GetTessellationMultiplier(FPassSpecificVSToDS I)
 	return GetMaterialTessellationMultiplier(MaterialParameters);
 }
 
-// get diameter of a bounding sphere in post projection (scaled by post proj z of origin) in Y.
-float GetPostProjectionSphereExtent( float3 Origin, float Diameter )
+int3 GetClipFlag( float3 Origin )
 {
 	float4 ClipPos = mul( float4( Origin, 1.0 ), View.TranslatedWorldToClip );
-	// Saturate to prevent NANs/INFs if ClipPos.w is near zero.
-	float ClipY = saturate(abs(Diameter * View.ViewToClip._m11 / ClipPos.w));
-	return ClipY;
+    int3 Clip0 = ClipPos.xyz < ClipPos.w;
+    int3 Clip1 = ClipPos.xyz > -ClipPos.w;
+    return Clip0 + 2 * Clip1;
 }
 
 // This is a heuristic that maps a bounding sphere on relevant edges/primitive to post projection space and uses that to scale the tessellation on the edges/primitive.
@@ -81,26 +80,35 @@ float GetPostProjectionSphereExtent( float3 Origin, float Diameter )
 float4 CalculateCompositeTessellationFactors(float3 Control0, float3 Control1, float3 Control2)
 {
 #if USE_ADAPTIVE_TESSELLATION_FACTOR
-	float Edge0 = distance(Control0, Control1);
-	float Edge1 = distance(Control1, Control2);
-	float Edge2 = distance(Control2, Control0);
+#if 1
+	// Frustum cull
+	int3 ClipFlag = 0;
+	ClipFlag  = GetClipFlag( Control0 );
+	ClipFlag |= GetClipFlag( Control1 );
+	ClipFlag |= GetClipFlag( Control2 );
+	if( any( ClipFlag != 3 ) )
+	{
+		return 0;
+	}
+#endif
 
-	// Average lengths to get an inside factor.
-	float InsideDiameter = (Edge0 + Edge1 + Edge2) / 3;
+	float3 Edge0 = ( Control0 - Control1 );
+	float3 Edge1 = ( Control1 - Control2 );
+	float3 Edge2 = ( Control2 - Control0 );
 
-	// Midpoints.
-	float3 Midpoint0 = (Control0 + Control1) / 2;
-	float3 Midpoint1 = (Control1 + Control2) / 2;
-	float3 Midpoint2 = (Control2 + Control0) / 2;
-	float3 MidpointAll = (Control0 + Control1 + Control2) / 3;
+	float3 ToMidpoint0 = 0.5 * ( Control0 + Control1 ) - View.TranslatedViewOrigin.xyz;
+	float3 ToMidpoint1 = 0.5 * ( Control1 + Control2 ) - View.TranslatedViewOrigin.xyz;
+	float3 ToMidpoint2 = 0.5 * ( Control2 + Control0 ) - View.TranslatedViewOrigin.xyz;
 
+	// Use spherical projection instead of planer
 	float4 CompositeFactors = float4(
-		GetPostProjectionSphereExtent( Midpoint1, Edge1 ),
-		GetPostProjectionSphereExtent( Midpoint2, Edge2 ),
-		GetPostProjectionSphereExtent( Midpoint0, Edge0 ),
-		GetPostProjectionSphereExtent( MidpointAll, InsideDiameter ) );
+		sqrt( dot( Edge1, Edge1 ) / dot( ToMidpoint1, ToMidpoint1 ) ),
+		sqrt( dot( Edge2, Edge2 ) / dot( ToMidpoint2, ToMidpoint2 ) ),
+		sqrt( dot( Edge0, Edge0 ) / dot( ToMidpoint0, ToMidpoint0 ) ),
+		1 );
+	CompositeFactors.w = 0.333 * ( CompositeFactors.x + CompositeFactors.y + CompositeFactors.z );
 
-	// The adaptive tessellation factor is 0.5 * ViewSizeY / PixelsPerEdge and CompositeFactors is 2 * PercentageOfScreen.  
+	// The adaptive tessellation factor is 0.5 * View.ViewToClip[1][1] * ViewSizeY / PixelsPerEdge and CompositeFactors is 2 * PercentageOfScreen.  
 	return View.AdaptiveTessellationFactor * CompositeFactors;
 #else
 	return float4( 1.0,1.0,1.0,1.0 );
diff --git a/Engine/Source/Runtime/Renderer/Private/SceneRendering.cpp b/Engine/Source/Runtime/Renderer/Private/SceneRendering.cpp
--- a/Engine/Source/Runtime/Renderer/Private/SceneRendering.cpp
+++ b/Engine/Source/Runtime/Renderer/Private/SceneRendering.cpp
@@ -372,7 +372,7 @@ TUniformBufferRef<FViewUniformShaderParameters> FViewInfo::CreateUniformBuffer(
 		// CVar setting is pixels/tri which is nice and intuitive.  But we want pixels/tessellated edge.  So use a heuristic.
 		float TessellationAdaptivePixelsPerEdge = FMath::Sqrt(2.f * CVarTessellationAdaptivePixelsPerTriangle.GetValueOnRenderThread());
 
-		ViewUniformShaderParameters.AdaptiveTessellationFactor = 0.5f * float(ViewRect.Height()) / TessellationAdaptivePixelsPerEdge;
+		ViewUniformShaderParameters.AdaptiveTessellationFactor = 0.5f * ViewMatrices.ProjMatrix.M[1][1] * float(ViewRect.Height()) / TessellationAdaptivePixelsPerEdge;
 	}
 
 	//white texture should act like a shadowmap cleared to the farplane.
