removed bad usage of OPENGL in shaders

diff --git a/Engine/Shaders/Definitions.usf b/Engine/Shaders/Definitions.usf
--- a/Engine/Shaders/Definitions.usf
+++ b/Engine/Shaders/Definitions.usf
@@ -109,11 +109,7 @@
 #endif
 
 #ifndef USING_TESSELLATION
-#define USING_TESSELLATION							0
-#endif
-
-#ifndef OPENGL
-#define OPENGL											0
+#define USING_TESSELLATION								0
 #endif
 
 #define PC_D3D											SM5_PROFILE
diff --git a/Engine/Shaders/SpeedTreeCommon.usf b/Engine/Shaders/SpeedTreeCommon.usf
--- a/Engine/Shaders/SpeedTreeCommon.usf
+++ b/Engine/Shaders/SpeedTreeCommon.usf
@@ -96,7 +96,7 @@ float3x3 RotationMatrix(float3 vAxis, float fAngle)
 {
     // compute sin/cos of fAngle
     float2 vSinCos;
-	#ifdef OPENGL
+	#if COMPILER_GLSL || COMPILER_GLSL_ES2
 		vSinCos.x = sin(fAngle);
 		vSinCos.y = cos(fAngle);
 	#else
