fixed comment

diff --git a/Engine/Shaders/BasePassPixelShader.usf b/Engine/Shaders/BasePassPixelShader.usf
--- a/Engine/Shaders/BasePassPixelShader.usf
+++ b/Engine/Shaders/BasePassPixelShader.usf
@@ -660,7 +660,6 @@ void Main(
 		Color += (DiffuseIndirectLighting * DiffuseColor + SubsurfaceIndirectLighting * SubsurfaceColor) * MaterialAO;
 
 		#if SIMPLE_DYNAMIC_LIGHTING
-			// always unshadowed so BiasedNDotL is not needed
 			half Lambert = saturate(dot(MaterialParameters.WorldNormal, View.DirectionalLightDirection));
 			Color += DiffuseColor * Lambert * View.DirectionalLightColor.rgb;
 
