fixed ForEachLight (Translucency per pixel lighting), added basic backlighting

diff --git a/Engine/Shaders/DeferredLightingCommon.usf b/Engine/Shaders/DeferredLightingCommon.usf
--- a/Engine/Shaders/DeferredLightingCommon.usf
+++ b/Engine/Shaders/DeferredLightingCommon.usf
@@ -245,7 +245,7 @@ float3 SubsurfaceShadingSubsurface( FGBufferData GBuffer, float3 L, float3 V, ha
 	return SubsurfaceColor * lerp(BackScatter, 1, InScatter);
 }
 
-float3 TwoSidedFoliageShadingSubsurface( FGBufferData GBuffer, float3 L, float3 V, half3 N )
+float3 TwoSidedFoliageShadingSubsurface( float3 SubsurfaceColor, float3 L, float3 V, half3 N )
 {
 	float VoL = saturate(dot(V, -L));
 	float VoLLobe = VoL * VoL;
@@ -260,7 +260,6 @@ float3 TwoSidedFoliageShadingSubsurface( FGBufferData GBuffer, float3 L, float3
 	float ViewDependenceAmount = .5f;
 	float BackShading = lerp(BackDiffuse, VoLLobe, ViewDependenceAmount);
 
-	float3 SubsurfaceColor = ExtractSubsurfaceColor(GBuffer);
 	return BackShading * SubsurfaceColor;
 }
 
@@ -298,6 +297,8 @@ float3 SurfaceShading( FGBufferData GBuffer, float3 LobeRoughness, float3 LobeEn
 
 float3 SubsurfaceShading( FGBufferData GBuffer, float3 L, float3 V, half3 N )
 {
+	float3 SubsurfaceColor = ExtractSubsurfaceColor(GBuffer);
+
 	switch( GBuffer.ShadingModelID )
 	{
 		case SHADINGMODELID_SUBSURFACE:
@@ -305,7 +306,7 @@ float3 SubsurfaceShading( FGBufferData GBuffer, float3 L, float3 V, half3 N )
 		case SHADINGMODELID_PREINTEGRATED_SKIN:
 			return SubsurfaceShadingPreintegratedSkin( GBuffer, L, V, N );
 		case SHADINGMODELID_TWOSIDED_FOLIAGE:
-			return TwoSidedFoliageShadingSubsurface( GBuffer, L, V, N );
+			return TwoSidedFoliageShadingSubsurface( SubsurfaceColor, L, V, N );
 		default:
 			return 0;
 	}
diff --git a/Engine/Shaders/ForEachLight.usf b/Engine/Shaders/ForEachLight.usf
--- a/Engine/Shaders/ForEachLight.usf
+++ b/Engine/Shaders/ForEachLight.usf
@@ -10,12 +10,6 @@
 #include "DeferredLightingCommon.usf"	// FSimpleDeferredLightData and GetSimpleDynamicLighting() 
 Buffer<uint> LightGrid;
 
-// greatly reduces shadow mapping artifacts
-float BiasedNDotL(float NDotLWithoutSaturate )
-{
-	return saturate(NDotLWithoutSaturate * 1.08f - 0.08f);
-}
-
 float3 SimplePointLightDiffuse( FScreenSpaceData ScreenSpaceData )
 {
 	return Diffuse_Lambert(ScreenSpaceData.GBuffer.DiffuseColor);
@@ -23,9 +17,23 @@ float3 SimplePointLightDiffuse( FScreenSpaceData ScreenSpaceData )
 
 float3 LightingForOneLight(FMaterialPixelParameters Parameters, FSimpleDeferredLightData LightData, FScreenSpaceData ScreenSpaceData)
 {
+	float3 CameraVector = normalize(Parameters.WorldPosition - View.ViewOrigin.xyz);
+
 	// At the moment we only support very simple lighting (no shadows, phong, no area lights, basic attenuation).
 	// Later we can extend this function.
-	return GetSimpleDynamicLighting(Parameters.WorldPosition, Parameters.CameraVector, ScreenSpaceData, LightData);
+
+	float3 Surface = GetSimpleDynamicLighting(Parameters.WorldPosition, CameraVector, ScreenSpaceData, LightData);
+	// todo: investigate why this doesn't work
+//	float3 Surface = GetSimpleDynamicLighting(Parameters.WorldPosition, Parameters.CameraVector, ScreenSpaceData, LightData);
+
+	half Opacity = GetMaterialOpacity(Parameters);
+	float3 ToLight = LightData.LightPositionAndInvRadius.xyz - Parameters.WorldPosition;
+
+	float3 Subsurface = (1 - Opacity) * TwoSidedFoliageShadingSubsurface(ScreenSpaceData.GBuffer.DiffuseColor, normalize(ToLight), -CameraVector, ScreenSpaceData.GBuffer.WorldNormal);
+
+
+	return Subsurface + Surface;
+
 }
 
 // @param WorldNormal is not assumed to be normalized
@@ -34,7 +42,7 @@ float3 AccumulateForwardLights(FMaterialPixelParameters Parameters, float3 Diffu
 	float3 Ret = 0;
 
 	// not yet supported on opengl (firstbitlow() does not exist)
-#ifndef OPENGL
+#if !COMPILER_GLSL && !COMPILER_GLSL_ES2
 
 	// Which tile we are?
 	int2 Tile = int2((Parameters.SVPosition.xy - View.ViewRectMin.xy) * ForwardLightData.InvTileSize);
