fix attenuation for translucency per pixel lighting

diff --git a/Engine/Shaders/ForEachLight.usf b/Engine/Shaders/ForEachLight.usf
--- a/Engine/Shaders/ForEachLight.usf
+++ b/Engine/Shaders/ForEachLight.usf
@@ -26,10 +26,23 @@ float3 LightingForOneLight(FMaterialPixelParameters Parameters, FSimpleDeferredL
 	// todo: investigate why this doesn't work
 //	float3 Surface = GetSimpleDynamicLighting(Parameters.WorldPosition, Parameters.CameraVector, ScreenSpaceData, LightData);
 
+	// todo: cleanup (share with former code)
+
+	// assuming inverse square falloff
+
 	half Opacity = GetMaterialOpacity(Parameters);
 	float3 ToLight = LightData.LightPositionAndInvRadius.xyz - Parameters.WorldPosition;
 
-	float3 Subsurface = (1 - Opacity) * TwoSidedFoliageShadingSubsurface(ScreenSpaceData.GBuffer.DiffuseColor, normalize(ToLight), -CameraVector, ScreenSpaceData.GBuffer.WorldNormal);
+	// Sphere falloff (technically just 1/d2 but this avoids inf)
+	float DistanceSqr = dot( ToLight, ToLight );
+	float DistanceAttenuation = 1 / ( DistanceSqr + 1 );
+	
+	float LightRadiusMask = Square( saturate( 1 - Square( DistanceSqr * Square(LightData.LightPositionAndInvRadius.w) ) ) );
+
+	float SubsurfaceAttenuation	= DistanceAttenuation * LightRadiusMask;
+
+	float3 Subsurface = (1 - Opacity) * LightData.LightColorAndFalloffExponent.rgb * SubsurfaceAttenuation
+		* TwoSidedFoliageShadingSubsurface(ScreenSpaceData.GBuffer.DiffuseColor, normalize(ToLight), -CameraVector, ScreenSpaceData.GBuffer.WorldNormal);
 
 
 	return Subsurface + Surface;
