improved/optimized SSAO
 good: larger radius, cache optimized using HZB, removed border weight code as Z=0 outside the texture works well, removed tweak parameter, defines to customize, single pass looks better
 bad: look changes, can look noisier (with large radius), Using HZB is an approximation which can cause artifacs especially with a smaller radius
 note: quality in PostProcessSettings can be adjusted, without TemporalAA there is more noise/dither pattern, if needed a more sample version can be made

diff --git a/Engine/Config/BaseScalability.ini b/Engine/Config/BaseScalability.ini
--- a/Engine/Config/BaseScalability.ini
+++ b/Engine/Config/BaseScalability.ini
@@ -101,8 +101,9 @@ r.DistanceFieldAO=1
 
 [PostProcessQuality@0]
 r.MotionBlurQuality=0
+r.AmbientOcclusionSampleSetQuality=0
 r.AmbientOcclusionLevels=0
-r.AmbientOcclusionRadiusScale=1.7
+r.AmbientOcclusionRadiusScale=1.2
 r.DepthOfFieldQuality=0
 r.RenderTargetPoolMin=300
 r.LensFlareQuality=0
@@ -116,8 +117,9 @@ r.LightShaftQuality=0
 
 [PostProcessQuality@1]
 r.MotionBlurQuality=3
+r.AmbientOcclusionSampleSetQuality=0
 r.AmbientOcclusionLevels=1
-r.AmbientOcclusionRadiusScale=1.7
+r.AmbientOcclusionRadiusScale=1.5
 r.DepthOfFieldQuality=1
 r.RenderTargetPoolMin=350
 r.LensFlareQuality=0
@@ -131,7 +133,8 @@ r.LightShaftQuality=0
 
 [PostProcessQuality@2]
 r.MotionBlurQuality=3
-r.AmbientOcclusionLevels=2
+r.AmbientOcclusionSampleSetQuality=-1
+r.AmbientOcclusionLevels=1
 r.AmbientOcclusionRadiusScale=1.5
 r.DepthOfFieldQuality=2
 r.RenderTargetPoolMin=400
@@ -146,7 +149,8 @@ r.LightShaftQuality=1
 
 [PostProcessQuality@3]
 r.MotionBlurQuality=4
-r.AmbientOcclusionLevels=3
+r.AmbientOcclusionSampleSetQuality=-1
+r.AmbientOcclusionLevels=2
 r.AmbientOcclusionRadiusScale=1.0
 r.DepthOfFieldQuality=2
 r.RenderTargetPoolMin=400
diff --git a/Engine/Shaders/PostProcessAmbientOcclusion.usf b/Engine/Shaders/PostProcessAmbientOcclusion.usf
--- a/Engine/Shaders/PostProcessAmbientOcclusion.usf
+++ b/Engine/Shaders/PostProcessAmbientOcclusion.usf
@@ -19,15 +19,23 @@
 // 0:no / 1:yes
 // USE_UPSAMPLE
 
+// 0:normal method using depth computed in setup passes / 1:faster but slightly lower quality
+#define USE_HZB	1
+
+// 0: classic with weighted sample, 1: don't normalize and adjust the formula to be simpler and faster - can look better and is cheaper (Alchemy like?)
+#define OPTIMIZATION_O1 1
+
+// 1:lowest quality, 2:medium , 3:high, more doesn't give too much (maybe HZB mip computations should be adjusted)
+#define SAMPLE_STEPS 3
+
 // 0:does not use normals (slightly faster, lower quality, WIP = not finished), 1: uses normals (slower but higher quality)
 #define SSAO_NORMAL_AFFECTED 1
 
 // ambient occlusion
 // AO_SAMPLE_QUALITY = 0 : no AO sampling, only upsampling
 // AO_SAMPLE_QUALITY = 1 : no dither/per pixel randomization
-// AO_SAMPLE_QUALITY = 2 : efficient high frequency 4x4 pattern
-// AO_SAMPLE_QUALITY = 3 : efficient high frequency 4x4 pattern changing over time
-// AO_SAMPLE_QUALITY = 4 : extreme sample count
+// AO_SAMPLE_QUALITY = 2 : efficient high frequency 4x4 pattern without jitter for TemporalAA
+// AO_SAMPLE_QUALITY = 3 : efficient high frequency 4x4 pattern with jitter for TemporalAA
 
 #if SAMPLESET_QUALITY == 0
 	// low
@@ -61,7 +69,7 @@
 	#define SAMPLESET_ARRAY_SIZE 1
 	static const float2 OcclusionSamplesOffsets[SAMPLESET_ARRAY_SIZE]=
 	{
-		// 6 points distributed on the unit disc, spiral order and distance
+		// one sample, for testing
 		float2(0.500, 0.500), 
 	};
 #elif USE_SAMPLESET == 1
@@ -69,7 +77,7 @@
 	static const float2 OcclusionSamplesOffsets[SAMPLESET_ARRAY_SIZE]=
 	{
 		// 3 points distributed on the unit disc, spiral order and distance
-		float2(0, -1.0f) * 0.3f, 
+		float2(0, -1.0f) * 0.43f, 
 		float2(0.58f, 0.814f) * 0.7f, 
 		float2(-0.58f, 0.814f) 
 	};
@@ -77,12 +85,12 @@
 	#define SAMPLESET_ARRAY_SIZE 5
 	static const float2 OcclusionSamplesOffsets[SAMPLESET_ARRAY_SIZE]=
 	{
-		// 5 points distributed on a unit circle
+		// 5 points distributed on a ring
 		float2(0.156434, 0.987688),
-		float2(0.987688, 0.156434),
-		float2(0.453990, -0.891007),
-		float2(-0.707107, -0.707107),
-		float2(-0.891006, 0.453991),
+		float2(0.987688, 0.156434)*0.9,
+		float2(0.453990, -0.891007)*0.8,
+		float2(-0.707107, -0.707107)*0.7,
+		float2(-0.891006, 0.453991)*0.65,
 	};
 #else // USE_SAMPLESET == 3
 	#define SAMPLESET_ARRAY_SIZE 6
@@ -144,13 +152,6 @@ float2 ScreenPosToUV(float2 ScreenPos, float4 InputSize)
 	return (ScreenPos * ScreenPosToPixel.xy + ScreenPosToPixel.zw + 0.5f) * InputSize.zw;
 }
 
-// @return true: inside, false outside
-bool IsScreenPosInsideBounds(float2 ScreenPos)
-{
-	// in range -1..1 for x and y
-	return all(abs(ScreenPos) < 1);
-}
-
 // 0: not similar .. 1:very similar
 float ComputeDepthSimilarity(float DepthA, float DepthB, float TweakScale)
 {
@@ -276,7 +277,7 @@ float4 ComputeUpsampleContribution(float SceneDepth, float2 InUV, float3 CenterW
 
 //todo		if(SSAO_NORMAL_AFFECTED == 1)
 		{
-			float3 LocalWorldNormal = NormalAndSampleDepth.xyz * 2 - 1;
+			float3 LocalWorldNormal = NormalAndSampleDepth.xyz*2-1;
 
 			Weight *= saturate(dot(LocalWorldNormal, CenterWorldNormal));
 		}
@@ -326,6 +327,149 @@ float acosApproxNormAngle(float x)
 	return saturate((1 - x) * 0.5f * ContrastTweak);
 }
 
+// @return float2(InvNormAngle, Weight)
+float2 Wedge(float Scale, float2 LocalRandom, float3 FovFix, float3 InvFovFix, float3 ViewSpacePosition, float3 ScaledViewSpaceNormal, float InvHaloSize, float MipLevel)
+{
+	float AORadiusInShader = ScreenSpaceAOandSSRShaderParams[1].z;		// 40 is default, /400 is constant
+	float ScaleRadiusInWorldSpace = ScreenSpaceAOandSSRShaderParams[2].z;
+
+	float3 SamplePosition2 = ViewSpacePosition + float3(LocalRandom, 0) * Scale;
+
+	float2 ScreenSpacePos = ReconstructSSPosFromCS(SamplePosition2 * FovFix);
+
+#if USE_HZB
+	float2 UV = ScreenSpacePos * float2(0.5f, -0.5f) + 0.5f;
+
+	float Abs = ConvertFromDeviceZ(Texture2DSampleLevel( PostprocessInput3, PostprocessInput3Sampler, UV, MipLevel).r);
+#else
+	float2 UV = ScreenPosToUV(ScreenSpacePos, PostprocessInput0Size);
+
+	float Abs = Texture2DSample(PostprocessInput1, PostprocessInput1Sampler, UV).a * Constant_Float16F_Scale;
+#endif
+
+	float3 SamplePosition = ReconstructCSPos(Abs, ScreenSpacePos) * InvFovFix;
+		
+	float3 Delta = SamplePosition - ViewSpacePosition;
+
+#if OPTIMIZATION_O1
+	float InvNormAngle = saturate(dot(Delta, ScaledViewSpaceNormal) / dot(Delta, Delta));
+#else
+	float InvNormAngle = saturate(dot(Delta, ScaledViewSpaceNormal) / length(Delta));
+#endif
+
+	float Weight = 1;
+
+#if !OPTIMIZATION_O1
+	Weight = saturate(1.0f - length(Delta) * InvHaloSize);
+#endif
+
+	return float2(InvNormAngle, Weight);
+}
+
+
+// @return float2(InvNormAngle, Weight)
+float2 Wedge2(float Scale, float2 LocalRandom, float3 FovFix, float3 InvFovFix, float3 ViewSpacePosition, float3 ScaledViewSpaceNormal, float InvHaloSize, float MipLevel)
+{
+	float3 DeltaLeft;
+	float AbsLeft;
+	float WeightLeft;
+	float3 SamplePositionLeft;
+	{
+		float AORadiusInShader = ScreenSpaceAOandSSRShaderParams[1].z;		// 40 is default, /400 is constant
+		float ScaleRadiusInWorldSpace = ScreenSpaceAOandSSRShaderParams[2].z;
+
+		float3 SamplePosition2 = ViewSpacePosition + float3(LocalRandom, 0) * Scale;
+
+		float2 ScreenSpacePos = ReconstructSSPosFromCS(SamplePosition2 * FovFix);
+
+	#if USE_HZB
+		float2 UV = ScreenSpacePos * float2(0.5f, -0.5f) + 0.5f;
+
+		AbsLeft = ConvertFromDeviceZ(Texture2DSampleLevel( PostprocessInput3, PostprocessInput3Sampler, UV, MipLevel).r);
+	#else
+		float2 UV = ScreenPosToUV(ScreenSpacePos, PostprocessInput0Size);
+
+		AbsLeft = Texture2DSample(PostprocessInput1, PostprocessInput1Sampler, UV).a * Constant_Float16F_Scale;
+	#endif
+
+		SamplePositionLeft = ReconstructCSPos(AbsLeft, ScreenSpacePos) * InvFovFix;
+		
+		DeltaLeft = SamplePositionLeft - ViewSpacePosition;
+
+		WeightLeft = 1;
+
+#if !OPTIMIZATION_O1
+		WeightLeft = saturate(1.0f - length(DeltaLeft) * InvHaloSize);
+#endif
+	}
+
+	float3 DeltaRight;
+	float AbsRight;
+	float WeightRight;
+	float3 SamplePositionRight;
+	{
+		float AORadiusInShader = ScreenSpaceAOandSSRShaderParams[1].z;		// 40 is default, /400 is constant
+		float ScaleRadiusInWorldSpace = ScreenSpaceAOandSSRShaderParams[2].z;
+
+		float3 SamplePosition2 = ViewSpacePosition - float3(LocalRandom, 0) * Scale;
+
+		float2 ScreenSpacePos = ReconstructSSPosFromCS(SamplePosition2 * FovFix);
+
+	#if USE_HZB
+		float2 UV = ScreenSpacePos * float2(0.5f, -0.5f) + 0.5f;
+
+		AbsRight = ConvertFromDeviceZ(Texture2DSampleLevel( PostprocessInput3, PostprocessInput3Sampler, UV, MipLevel).r);
+	#else
+		float2 UV = ScreenPosToUV(ScreenSpacePos, PostprocessInput0Size);
+
+		AbsRight = Texture2DSample(PostprocessInput1, PostprocessInput1Sampler, UV).a * Constant_Float16F_Scale;
+	#endif
+
+		SamplePositionRight = ReconstructCSPos(AbsRight, ScreenSpacePos) * InvFovFix;
+		
+		DeltaRight = SamplePositionRight - ViewSpacePosition;
+		
+		WeightRight = 1;
+
+#if !OPTIMIZATION_O1
+		WeightRight = saturate(1.0f - length(DeltaRight) * InvHaloSize);
+#endif
+	}
+
+
+	float FlatSurfaceBias = 5.0f;
+
+	float left = ViewSpacePosition.z - AbsLeft;
+	float right = ViewSpacePosition.z - AbsRight;
+
+	// OptionA: accurate angle computation
+	float NormAngle = acosApproxNormAngle( dot(DeltaLeft, DeltaRight) / sqrt(length2(DeltaLeft) * length2(DeltaRight)));
+	// OptionB(fade out in near distance): float NormAngle = acosApproxNormAngle( (- left - right) * 20);
+	// OptionC(look consistent but more noisy, should be much faster): float NormAngle = 0;
+
+
+	// not 100% correct but simple
+	// bias is needed to avoid flickering on almost perfectly flat surfaces
+	//	    if((leftAbs  + rightAbs) * 0.5f > SceneDepth - 0.0001f)
+	if(left + right < FlatSurfaceBias)
+	{
+		// fix concave case
+		NormAngle = 1;
+	}
+
+	// to avoid halos around objects
+	float Weight = 1;
+				
+	float InvAmbientOcclusionDistance = ScreenSpaceAOandSSRShaderParams[0].z;
+	float ViewDepthAdd = 1.0f - ViewSpacePosition.z * InvAmbientOcclusionDistance;
+
+	Weight *= saturate(SamplePositionLeft.z * InvAmbientOcclusionDistance + ViewDepthAdd);
+	Weight *= saturate(SamplePositionRight.z * InvAmbientOcclusionDistance + ViewDepthAdd);
+
+//	return float2(1 - NormAngle, (WeightLeft + WeightRight) * 0.5f);
+	return float2((1-NormAngle) / (Weight + 0.001f), Weight);
+}
+
 // the main pixel shader that computes ambient occlusion
 void MainPS(in float4 UVAndScreenPos : TEXCOORD0, out float4 OutColor : SV_Target0)
 {
@@ -334,7 +478,7 @@ void MainPS(in float4 UVAndScreenPos : TEXCOORD0, out float4 OutColor : SV_Targe
 	// the following constants as set up on C++ side
 	float AmbientOcclusionPower = ScreenSpaceAOandSSRShaderParams[0].x;
 	float Ratio = ScreenSpaceAOandSSRShaderParams[1].w;
-	float2 AORadiusInShader = ScreenSpaceAOandSSRShaderParams[1].zz;
+	float AORadiusInShader = ScreenSpaceAOandSSRShaderParams[1].z;
 	float InvAmbientOcclusionDistance = ScreenSpaceAOandSSRShaderParams[0].z;
 	float AmbientOcclusionIntensity = ScreenSpaceAOandSSRShaderParams[0].w;
 	float2 ViewportUVToRandomUV = ScreenSpaceAOandSSRShaderParams[1].xy;
@@ -354,171 +498,108 @@ void MainPS(in float4 UVAndScreenPos : TEXCOORD0, out float4 OutColor : SV_Targe
 
 	float3 ViewSpaceNormal = normalize(mul(WorldNormal, (float3x3)View.TranslatedWorldToView));
 	
-	AORadiusInShader *= lerp(SceneDepth, 1, ScaleRadiusInWorldSpace);
-	
 	float3 ViewSpacePosition = ReconstructCSPos(SceneDepth, ScreenPos);
 
 	//Invert affects of fov and aspect in projection matrix. Reasonably sure the maths works out on this. The results look right.
 	ViewSpacePosition *= InvFovFix;
-	AORadiusInShader *= InvFovFix.xy;
+	float2 ActualAORadiusXY = AORadiusInShader * lerp(SceneDepth, 1, ScaleRadiusInWorldSpace) * InvFovFix.xy;
 
-	//Add bias after fixup
+	//Add bias after fixup (causes minor banding - not needed with larger radius)
 	ViewSpacePosition += ViewSpaceNormal * (AmbientOcclusionBias * SceneDepth) * ScaleFactor;
 
-	float ViewDepthAdd = 1.0f - ViewSpacePosition.z * InvAmbientOcclusionDistance;
-
-	float WeightAccumulator = 0.0001f;
-	float WeightedData = WeightAccumulator;
+	float2 WeightAccumulator = 0.0001f;
 	
 #if AO_SAMPLE_QUALITY != 0
 	// no per pixel randomization
 
 #if AO_SAMPLE_QUALITY == 1
 	// 4x4 random base direction
-	float4 RandomBase = float4(0,1,1,0) * float4(AORadiusInShader, AORadiusInShader);
+	float4 RandomBase = float4(0,1,1,0) * float4(ActualAORadiusXY, ActualAORadiusXY);
 	{
 #elif AO_SAMPLE_QUALITY == 2
 	// extract one of 16 base vectors (rotation and scale) from a texture that repeats 4x4
-	float4 RandomBase = (Texture2DSample(RandomNormalTexture, RandomNormalTextureSampler, UV * ViewportUVToRandomUV) * 2 - 1 ) * float4(AORadiusInShader, AORadiusInShader);
+	float4 RandomBase = (Texture2DSample(RandomNormalTexture, RandomNormalTextureSampler, UV * ViewportUVToRandomUV) * 2 - 1 ) * float4(ActualAORadiusXY, ActualAORadiusXY);
 	{
-#elif AO_SAMPLE_QUALITY == 3
+#else // AO_SAMPLE_QUALITY == 3
 	// extract one of 16 base vectors (rotation and scale) from a texture that repeats 4x4, changing over time if TemporalAA is enabled
-	float2 TemporalOffset = (View.TemporalAAParams.r % 8) * float2(2, 7) / 64.0f;
+	// todo: can be moved into the vertex shader
+	float2 TemporalOffset = (View.TemporalAAParams.r % 8) * float2(2.48f, 6.72f) / 64.0f;
 
 	float4 RandomNorm = Texture2DSample(RandomNormalTexture, RandomNormalTextureSampler, TemporalOffset + UV * ViewportUVToRandomUV);
-	float4 RandomBase = (RandomNorm * 2 - 1) * float4(AORadiusInShader, AORadiusInShader);
+	float4 RandomBase = (RandomNorm * 2 - 1) * float4(ActualAORadiusXY, ActualAORadiusXY);
 
 	{
-#else // AO_SAMPLE_QUALITY == 4
-	// can be optimized
-	int2 PixelPos = int2(UVAndScreenPos.zw * ScreenPosToPixel.xy + ScreenPosToPixel.zw + 0.5f);
-	// repeat a cross pattern in a 5x5 tile
-	// .1. 12345
-	// 234 34512
-	// .5. 51234
-	//     23451
-	//     45673 
-	int Index = ((uint)(PixelPos.x) + 2 * (uint)(PixelPos.y)) % 5;
-
-	// QualityAmplifier 1:normal 2:2x more expensive, 3:3x more expensive, ...
-	const uint QualityAmplifier = 1;
-
-	UNROLL for(uint e = 0; e < QualityAmplifier; ++e)
-	{
-		uint CombinedIndex = Index + e * 5;
-		// a sample pair covers 180 degree so more samples we but in that range
-		float RandomAngle = CombinedIndex / (QualityAmplifier * 5.0f) * PI;
-		float2 RandomDirection = float2(sin(RandomAngle), cos(RandomAngle)); 
-		float4 RandomBase = float4(RandomDirection, RandomDirection.y, -RandomDirection.x) * float4(AORadiusInShader, AORadiusInShader);
-#endif // AO_SAMPLE_QUALITY == 0
+#endif // AO_SAMPLE_QUALITY == 
+
+		// .x means for very anisotropic viewports we scale by x
+		float InvHaloSize = 1.0f / (ActualAORadiusXY.x * 2);
+
+		float3 ScaledViewSpaceNormal = ViewSpaceNormal;
+
+#if OPTIMIZATION_O1
+//		ScaledViewSpaceNormal *= 100.0f * View.GeneralPurposeTweak;
+		ScaledViewSpaceNormal *= 0.08f * lerp(SceneDepth, 1000, ScaleRadiusInWorldSpace) * View.GeneralPurposeTweak;
+#endif
 
 		UNROLL for(int i = 0; i < SAMPLESET_ARRAY_SIZE; ++i)
 		{
+			// -1..1
 			float2 UnrotatedRandom = OcclusionSamplesOffsets[i].xy;
 
+			// -ActualAORadiusXY..ActualAORadiusXY
 			float2 LocalRandom = UnrotatedRandom.x * RandomBase.xy + UnrotatedRandom.y * RandomBase.zw;
 
-			float3 SamplePosition0 = ViewSpacePosition + float3(LocalRandom, 0);
-			// mirrored sampled
-			float3 SamplePosition1 = ViewSpacePosition - float3(LocalRandom, 0);
-
-			float2 ScreenSpacePosLeft = ReconstructSSPosFromCS(SamplePosition0 * FovFix);
-			float2 ScreenSpacePosRight = ReconstructSSPosFromCS(SamplePosition1 * FovFix);
-
-			float2 UVLeft = ScreenPosToUV(ScreenSpacePosLeft, PostprocessInput0Size);
-			float2 UVRight = ScreenPosToUV(ScreenSpacePosRight, PostprocessInput0Size);
-
-			float leftAbs = Texture2DSample(PostprocessInput1, PostprocessInput1Sampler, UVLeft).a * Constant_Float16F_Scale;
-			float rightAbs = Texture2DSample(PostprocessInput1, PostprocessInput1Sampler, UVRight).a * Constant_Float16F_Scale;
-
-			SamplePosition0 = ReconstructCSPos(leftAbs, ScreenSpacePosLeft) * InvFovFix;
-			SamplePosition1 = ReconstructCSPos(rightAbs, ScreenSpacePosRight) * InvFovFix;
-		
-			float3 DeltaLeft = SamplePosition0 - ViewSpacePosition;
-			float3 DeltaRight = SamplePosition1 - ViewSpacePosition;
-
 			// Ambient Occlusion affected by per pixel normal
 
 			if(SSAO_NORMAL_AFFECTED == 0)
 			{
-				float FlatSurfaceBias = 5.0f;
-
-				float left = SceneDepth - leftAbs;
-				float right = SceneDepth - rightAbs;
+				float2 LocalAccumulator = float2(0, 0);
 
-				// OptionA: accurate angle computation
-				float NormAngle = acosApproxNormAngle( dot(DeltaLeft, DeltaRight) / sqrt(length2(DeltaLeft) * length2(DeltaRight)));
-				// OptionB(fade out in near distance): float NormAngle = acosApproxNormAngle( (- left - right) * 20);
-				// OptionC(look consistent but more noisy, should be much faster): float NormAngle = 0;
-
-				// not 100% correct but simple
-				// bias is needed to avoid flickering on almost perfectly flat surfaces
-				//	    if((leftAbs  + rightAbs) * 0.5f > SceneDepth - 0.0001f)
-				if(left + right < FlatSurfaceBias)
+				UNROLL for(uint step = 0; step < SAMPLE_STEPS; ++step)
 				{
-					// fix concave case
-					NormAngle = 1;
-				}
-
-				// to avoid halos around objects
-				float Weight = 1;
-				
-				Weight *= saturate(SamplePosition0.z * InvAmbientOcclusionDistance + ViewDepthAdd);
-				Weight *= saturate(SamplePosition1.z * InvAmbientOcclusionDistance + ViewDepthAdd);
+					// constant at run time
+					float Scale = (step + 1) / (float)SAMPLE_STEPS;
+					// constant at run time (higher is better for texture cache / performance, lower is better quality
+					float MipLevel = max(step, 1);
 
-				// what is outside of the screen should not contribute to the AO
-				Weight *= IsScreenPosInsideBounds(ScreenSpacePosLeft);
-				Weight *= IsScreenPosInsideBounds(ScreenSpacePosRight);
+					float2 StepSample = Wedge2(Scale, LocalRandom, FovFix, InvFovFix, ViewSpacePosition, ScaledViewSpaceNormal, InvHaloSize, MipLevel);
 
-				// the area scales quadratic with the angle ? - with this line it gets a bit darker
-//				float Area = NormAngle * NormAngle;
-				// wrong? looks better/more similar to no normal
-				float Area = NormAngle;
+					// combine horizon samples
+					LocalAccumulator = lerp(LocalAccumulator, float2(max(LocalAccumulator.x, StepSample.x), 1), StepSample.y);
+				}
 
-				WeightedData += Area * Weight;
-				WeightAccumulator += Weight;
+				// Square(): the area scales quadratic with the angle - it gets a bit darker
+				WeightAccumulator += float2(Square(1 - LocalAccumulator.x) * LocalAccumulator.y, LocalAccumulator.y);
 			}
 			else // SSAO_NORMAL_AFFECTED == 1
 			{
-				// left
-				{
-					float NormAngle = 1.0f - saturate(dot(DeltaLeft, ViewSpaceNormal) / length(DeltaLeft));
-
-					// to avoid halos around objects
-					float Weight = saturate(1 + (SamplePosition0.z - ViewSpacePosition.z) * InvAmbientOcclusionDistance);
-
-					// what is outside of the screen should not contribute to the AO
-					Weight *= IsScreenPosInsideBounds(ScreenSpacePosLeft);
-
-					// the area scales quadratic with the angle ? - with this line it gets a bit darker
-					float Area = NormAngle * NormAngle;
-
-					WeightedData += Area * Weight;
-					WeightAccumulator += Weight;
-				}
-
-				// right
+				UNROLL for(int mirror = -1; mirror <= 1; mirror +=2)
 				{
-					float NormAngle = 1.0f - saturate(dot(DeltaRight, ViewSpaceNormal) / length(DeltaRight));
+					float2 LocalAccumulator = float2(0, 0);
 
-					float Weight = saturate(SamplePosition1.z * InvAmbientOcclusionDistance + ViewDepthAdd);
+					UNROLL for(uint step = 0; step < SAMPLE_STEPS; ++step)
+					{
+						// constant at run time
+						float Scale = (step + 1) / (float)SAMPLE_STEPS * mirror;
+						// constant at run time (higher is better for texture cache / performance, lower is better quality
+						float MipLevel = max(step, 1);
 
-					Weight *= IsScreenPosInsideBounds(ScreenSpacePosRight);
+						float2 StepSample = Wedge(Scale, LocalRandom, FovFix, InvFovFix, ViewSpacePosition, ScaledViewSpaceNormal, InvHaloSize, MipLevel);
 
-					// the area scales quadratic with the angle ? - with this line it gets a bit darker
-					float Area = NormAngle * NormAngle;
+						// combine horizon samples
+						LocalAccumulator = lerp(LocalAccumulator, float2(max(LocalAccumulator.x, StepSample.x), 1), StepSample.y);
+					}
 
-					WeightedData += Area * Weight;
-					WeightAccumulator += Weight;
+					// Square(): the area scales quadratic with the angle - it gets a bit darker
+					WeightAccumulator += float2(Square(1 - LocalAccumulator.x) * LocalAccumulator.y, LocalAccumulator.y);
 				}
 			}
 		}
 	}
 
-	WeightedData.r /= WeightAccumulator;
 #endif // #if AO_SAMPLE_QUALITY == 0
 	
-	OutColor.r = WeightedData.r;
+	OutColor.r = WeightAccumulator.x / WeightAccumulator.y;
 	OutColor.gb = float2(0, 0);
 
 	{
diff --git a/Engine/Source/Runtime/Engine/Classes/Engine/Scene.h b/Engine/Source/Runtime/Engine/Classes/Engine/Scene.h
--- a/Engine/Source/Runtime/Engine/Classes/Engine/Scene.h
+++ b/Engine/Source/Runtime/Engine/Classes/Engine/Scene.h
@@ -206,7 +206,7 @@ struct FPostProcessSettings
 	uint32 bOverride_AmbientOcclusionFadeRadius:1;
 
 	UPROPERTY(BlueprintReadWrite, Category=Overrides, meta=(PinHiddenByDefault))
-	uint32 bOverride_AmbientOcclusionDistance:1;
+	uint32 bOverride_AmbientOcclusionDistance_DEPRECATED:1;
 
 	UPROPERTY(BlueprintReadWrite, Category=Overrides, meta=(PinHiddenByDefault))
 	uint32 bOverride_AmbientOcclusionRadiusInWS:1;
@@ -618,7 +618,7 @@ struct FPostProcessSettings
 	float AmbientOcclusionStaticFraction;
 
 	/** >0, in unreal units, bigger values means even distant surfaces affect the ambient occlusion */
-	UPROPERTY(interp, BlueprintReadWrite, Category=AmbientOcclusion, meta=(ClampMin = "0.1", UIMax = "200.0", editcondition = "bOverride_AmbientOcclusionRadius", DisplayName = "Radius"))
+	UPROPERTY(interp, BlueprintReadWrite, Category=AmbientOcclusion, meta=(ClampMin = "0.1", UIMax = "500.0", editcondition = "bOverride_AmbientOcclusionRadius", DisplayName = "Radius"))
 	float AmbientOcclusionRadius;
 
 	/** true: AO radius is in world space units, false: AO radius is locked the view space in 400 units */
@@ -633,9 +633,9 @@ struct FPostProcessSettings
 	UPROPERTY(interp, BlueprintReadWrite, Category=AmbientOcclusion, AdvancedDisplay, meta=(ClampMin = "0.0", UIMax = "20000.0", editcondition = "bOverride_AmbientOcclusionFadeRadius", DisplayName = "Fade Out Radius"))
 	float AmbientOcclusionFadeRadius;
 
-	/** >0, in unreal units, how wide the ambient occlusion effect should affect the geometry (in depth) */
-	UPROPERTY(interp, BlueprintReadWrite, Category=AmbientOcclusion, AdvancedDisplay, meta=(ClampMin = "0.0", UIMax = "2000.0", editcondition = "bOverride_AmbientOcclusionDistance", DisplayName = "Occlusion Distance"))
-	float AmbientOcclusionDistance;
+	/** >0, in unreal units, how wide the ambient occlusion effect should affect the geometry (in depth), will be removed - only used for non normal method which is not exposed */
+	UPROPERTY(interp, BlueprintReadWrite, Category=AmbientOcclusion, AdvancedDisplay, meta=(ClampMin = "0.0", UIMax = "2000.0", editcondition = "bOverride_AmbientOcclusionDistance_DEPRECATED", DisplayName = "Occlusion Distance"))
+	float AmbientOcclusionDistance_DEPRECATED;
 
 	/** >0, in unreal units, bigger values means even distant surfaces affect the ambient occlusion */
 	UPROPERTY(interp, BlueprintReadWrite, Category=AmbientOcclusion, AdvancedDisplay, meta=(ClampMin = "0.1", UIMax = "8.0", editcondition = "bOverride_AmbientOcclusionPower", DisplayName = "Power"))
@@ -853,8 +853,8 @@ struct FPostProcessSettings
 		AmbientOcclusionIntensity = .5f;
 		// next value might get overwritten by r.DefaultFeature.AmbientOcclusionStaticFraction
 		AmbientOcclusionStaticFraction = 1.0f;
-		AmbientOcclusionRadius = 40.0f;
-		AmbientOcclusionDistance = 80.0f;
+		AmbientOcclusionRadius = 200.0f;
+		AmbientOcclusionDistance_DEPRECATED = 80.0f;
 		AmbientOcclusionFadeDistance = 8000.0f;
 		AmbientOcclusionFadeRadius = 5000.0f;
 		AmbientOcclusionPower = 2.0f;
diff --git a/Engine/Source/Runtime/Renderer/Private/CompositionLighting/CompositionLighting.cpp b/Engine/Source/Runtime/Renderer/Private/CompositionLighting/CompositionLighting.cpp
--- a/Engine/Source/Runtime/Renderer/Private/CompositionLighting/CompositionLighting.cpp
+++ b/Engine/Source/Runtime/Renderer/Private/CompositionLighting/CompositionLighting.cpp
@@ -135,7 +135,7 @@ static void AddPostProcessingAmbientCubemap(FPostprocessContext& Context, FRende
 }
 
 // @param Levels 0..4, how many different resolution levels we want to render
-static FRenderingCompositeOutputRef AddPostProcessingAmbientOcclusion(FPostprocessContext& Context, uint32 Levels)
+static FRenderingCompositeOutputRef AddPostProcessingAmbientOcclusion(FRHICommandListImmediate& RHICmdList, FPostprocessContext& Context, uint32 Levels)
 {
 	check(Levels >= 0 && Levels <= 4);
 
@@ -163,6 +163,19 @@ static FRenderingCompositeOutputRef AddPostProcessingAmbientOcclusion(FPostproce
 		AmbientOcclusionInMip3->SetInput(ePId_Input1, FRenderingCompositeOutputRef(AmbientOcclusionInMip2, ePId_Output0));
 	}
 
+	FRenderingCompositePass* HZBInput = 0;
+
+	FSceneViewState* ViewState = (FSceneViewState*)Context.View.State;
+
+	if(ViewState)
+	{
+		void BuildHZB(FRHICommandListImmediate& RHICmdList, const FViewInfo& View);
+
+		BuildHZB(RHICmdList, Context.View);
+
+		HZBInput = Context.Graph.RegisterPass( new FRCPassPostProcessInput( ViewState->HZB.Texture ) );
+	}
+
 	// upsample from lower resolution
 
 	if(Levels >= 4)
@@ -170,6 +183,7 @@ static FRenderingCompositeOutputRef AddPostProcessingAmbientOcclusion(FPostproce
 		AmbientOcclusionPassMip3 = Context.Graph.RegisterPass(new(FMemStack::Get()) FRCPassPostProcessAmbientOcclusion());
 		AmbientOcclusionPassMip3->SetInput(ePId_Input0, AmbientOcclusionInMip3);
 		AmbientOcclusionPassMip3->SetInput(ePId_Input1, AmbientOcclusionInMip3);
+		AmbientOcclusionPassMip3->SetInput(ePId_Input3, HZBInput);
 	}
 
 	if(Levels >= 3)
@@ -178,6 +192,7 @@ static FRenderingCompositeOutputRef AddPostProcessingAmbientOcclusion(FPostproce
 		AmbientOcclusionPassMip2->SetInput(ePId_Input0, AmbientOcclusionInMip2);
 		AmbientOcclusionPassMip2->SetInput(ePId_Input1, AmbientOcclusionInMip2);
 		AmbientOcclusionPassMip2->SetInput(ePId_Input2, AmbientOcclusionPassMip3);
+		AmbientOcclusionPassMip2->SetInput(ePId_Input3, HZBInput);
 	}
 
 	if(Levels >= 2)
@@ -186,6 +201,7 @@ static FRenderingCompositeOutputRef AddPostProcessingAmbientOcclusion(FPostproce
 		AmbientOcclusionPassMip1->SetInput(ePId_Input0, AmbientOcclusionInMip1);
 		AmbientOcclusionPassMip1->SetInput(ePId_Input1, AmbientOcclusionInMip1);
 		AmbientOcclusionPassMip1->SetInput(ePId_Input2, AmbientOcclusionPassMip2);
+		AmbientOcclusionPassMip1->SetInput(ePId_Input3, HZBInput);
 	}
 
 	FRenderingCompositePass* GBufferA = Context.Graph.RegisterPass(new(FMemStack::Get()) FRCPassPostProcessInput(GSceneRenderTargets.GBufferA));
@@ -196,6 +212,7 @@ static FRenderingCompositeOutputRef AddPostProcessingAmbientOcclusion(FPostproce
 	AmbientOcclusionPassMip0->SetInput(ePId_Input0, GBufferA);
 	AmbientOcclusionPassMip0->SetInput(ePId_Input1, AmbientOcclusionInMip1);
 	AmbientOcclusionPassMip0->SetInput(ePId_Input2, AmbientOcclusionPassMip1);
+	AmbientOcclusionPassMip0->SetInput(ePId_Input3, HZBInput);
 
 	// to make sure this pass is processed as well (before), needed to make process decals before computing AO
 	AmbientOcclusionInMip1->AddDependency(Context.FinalOutput);
@@ -297,7 +314,7 @@ void FCompositionLighting::ProcessAfterBasePass(FRHICommandListImmediate& RHICmd
 
 		if(uint32 Levels = ComputeAmbientOcclusionPassCount(Context))
 		{
-			AmbientOcclusion = AddPostProcessingAmbientOcclusion(Context, Levels);
+			AmbientOcclusion = AddPostProcessingAmbientOcclusion(RHICmdList, Context, Levels);
 		}
 
 		if(IsAmbientCubemapPassRequired(Context))
diff --git a/Engine/Source/Runtime/Renderer/Private/CompositionLighting/PostProcessAmbientOcclusion.cpp b/Engine/Source/Runtime/Renderer/Private/CompositionLighting/PostProcessAmbientOcclusion.cpp
--- a/Engine/Source/Runtime/Renderer/Private/CompositionLighting/PostProcessAmbientOcclusion.cpp
+++ b/Engine/Source/Runtime/Renderer/Private/CompositionLighting/PostProcessAmbientOcclusion.cpp
@@ -16,9 +16,9 @@ static TAutoConsoleVariable<int32> CVarAmbientOcclusionSampleSetQuality(
 	-1,
 	TEXT("Defines how many samples we use for ScreenSpace Ambient Occlusion\n")
 	TEXT("-1: sample count depends on post process settings (default)\n")
-	TEXT("0: low sample count (defined in shader, 3 * 2)\n")
-	TEXT("1: high sample count (defined in shader, 6 * 2)"),
-	ECVF_RenderThreadSafe);
+	TEXT("0: low sample count (defined in shader, 3 * 2 per pixel)\n")
+	TEXT("1: high sample count (defined in shader, 6 * 2 per pixel)"),
+	ECVF_Scalability | ECVF_RenderThreadSafe);
 
 
 /** Encapsulates the post processing ambient occlusion pixel shader. */
@@ -343,9 +343,8 @@ IMPLEMENT_SHADER_TYPE(,FPostProcessBasePassAOPS,TEXT("PostProcessAmbientOcclusio
 
 // --------------------------------------------------------
 
-
 template <uint32 bInitialSetup>
-void FRCPassPostProcessAmbientOcclusionSetup::SetShaderSetupTempl(const FRenderingCompositePassContext& Context)
+FShader* FRCPassPostProcessAmbientOcclusionSetup::SetShaderSetupTempl(const FRenderingCompositePassContext& Context)
 {
 	TShaderMapRef<FPostProcessVS> VertexShader(Context.GetShaderMap());
 	TShaderMapRef<FPostProcessAmbientOcclusionSetupPS<bInitialSetup> > PixelShader(Context.GetShaderMap());
@@ -357,6 +356,8 @@ void FRCPassPostProcessAmbientOcclusionSetup::SetShaderSetupTempl(const FRenderi
 
 	VertexShader->SetParameters(Context);
 	PixelShader->SetParameters(Context);
+
+	return *VertexShader;
 }
 
 void FRCPassPostProcessAmbientOcclusionSetup::Process(FRenderingCompositePassContext& Context)
@@ -385,27 +386,27 @@ void FRCPassPostProcessAmbientOcclusionSetup::Process(FRenderingCompositePassCon
 	Context.RHICmdList.SetRasterizerState(TStaticRasterizerState<>::GetRHI());
 	Context.RHICmdList.SetDepthStencilState(TStaticDepthStencilState<false, CF_Always>::GetRHI());
 
+	FShader* VertexShader = 0;
+
 	if(IsInitialPass())
 	{
-		SetShaderSetupTempl<1>(Context);
+		VertexShader = SetShaderSetupTempl<1>(Context);
 	}
 	else
 	{
-		SetShaderSetupTempl<0>(Context);
+		VertexShader = SetShaderSetupTempl<0>(Context);
 	}
 
-	TShaderMapRef<FPostProcessVS> VertexShader(Context.GetShaderMap());
-
 	// Draw a quad mapping scene color to the view's render target
 	DrawRectangle( 
 		Context.RHICmdList,
 		0, 0,
 		DestRect.Width(), DestRect.Height(),
-		SrcRect.Min.X, SrcRect.Min.Y, 
+			SrcRect.Min.X, SrcRect.Min.Y, 
 		SrcRect.Width(), SrcRect.Height(),
 		DestRect.Size(),
 		GSceneRenderTargets.GetBufferSizeXY(),
-		*VertexShader,
+		VertexShader,
 		EDRF_UseTriangleOptimization);
 
 	Context.RHICmdList.CopyToResolveTarget(DestRenderTarget.TargetableTexture, DestRenderTarget.ShaderResourceTexture, false, FResolveParams());
@@ -555,7 +556,7 @@ void FRCPassPostProcessAmbientOcclusion::Process(FRenderingCompositePassContext&
 		Context.RHICmdList,
 		0, 0,
 		ViewRect.Width(), ViewRect.Height(),
-		ViewRect.Min.X, ViewRect.Min.Y, 
+		ViewRect.Min.X, ViewRect.Min.Y,
 		ViewRect.Width(), ViewRect.Height(),
 		ViewRect.Size(),
 		TexSize,
@@ -570,8 +571,12 @@ FPooledRenderTargetDesc FRCPassPostProcessAmbientOcclusion::ComputeOutputDesc(EP
 {
 	if(!bAOSetupAsInput)
 	{
+		FPooledRenderTargetDesc Ret;
+
+		Ret.DebugName = TEXT("AmbientOcclusionDirect");
+
 		// we render directly to the buffer, no need for an intermediate target, we output in a single channel
-		return FPooledRenderTargetDesc();
+		return Ret;
 	}
 
 	FPooledRenderTargetDesc Ret = PassInputs[0].GetOutput()->RenderTargetDesc;
@@ -581,7 +586,6 @@ FPooledRenderTargetDesc FRCPassPostProcessAmbientOcclusion::ComputeOutputDesc(EP
 	Ret.Format = PF_B8G8R8A8;
 	Ret.TargetableFlags &= ~TexCreate_DepthStencilTargetable;
 	Ret.TargetableFlags |= TexCreate_RenderTargetable;
-
 	Ret.DebugName = TEXT("AmbientOcclusion");
 
 	return Ret;
diff --git a/Engine/Source/Runtime/Renderer/Private/CompositionLighting/PostProcessAmbientOcclusion.h b/Engine/Source/Runtime/Renderer/Private/CompositionLighting/PostProcessAmbientOcclusion.h
--- a/Engine/Source/Runtime/Renderer/Private/CompositionLighting/PostProcessAmbientOcclusion.h
+++ b/Engine/Source/Runtime/Renderer/Private/CompositionLighting/PostProcessAmbientOcclusion.h
@@ -25,15 +25,17 @@ private:
 	// otherwise this is a down sampling pass which takes two MRT inputs from the setup pass before
 	bool IsInitialPass() const;
 
+	// @return VertexShader
 	template <uint32 bInitialSetup>
-	void SetShaderSetupTempl(const FRenderingCompositePassContext& Context);
+	FShader* SetShaderSetupTempl(const FRenderingCompositePassContext& Context);
 };
 
 // ePId_Input0: defines the resolution we compute AO and provides the normal
 // ePId_Input1: setup in same resolution as ePId_Input1 for depth expect when running in full resolution, then it's half
 // ePId_Input2: optional AO result one lower resolution
+// ePId_Input3: optional HZB
 // derives from TRenderingCompositePassBase<InputCount, OutputCount> 
-class FRCPassPostProcessAmbientOcclusion : public TRenderingCompositePassBase<3, 1>
+class FRCPassPostProcessAmbientOcclusion : public TRenderingCompositePassBase<4, 1>
 {
 public:
 	// @param bInAOSetupAsInput true:use AO setup as input, false: use GBuffer normal and native z depth
@@ -132,7 +134,7 @@ public:
 		float InvFadeRadius = 1.0f / FadeRadius;
 
 		// /1000 to be able to define the value in that distance
-		Value[0] = FVector4(Settings.AmbientOcclusionPower, Settings.AmbientOcclusionBias / 1000.0f, 1.0f / Settings.AmbientOcclusionDistance, Settings.AmbientOcclusionIntensity);
+		Value[0] = FVector4(Settings.AmbientOcclusionPower, Settings.AmbientOcclusionBias / 1000.0f, 1.0f / Settings.AmbientOcclusionDistance_DEPRECATED, Settings.AmbientOcclusionIntensity);
 		Value[1] = FVector4(ViewportUVToRandomUV.X, ViewportUVToRandomUV.Y, AORadiusInShader, Ratio);
 		Value[2] = FVector4(ScaleToFullRes, Settings.AmbientOcclusionMipThreshold / ScaleToFullRes, ScaleRadiusInWorldSpace, Settings.AmbientOcclusionMipBlend);
 		Value[3] = FVector4(0, 0, StaticFraction, InvTanHalfFov);
