better comment

diff --git a/Engine/Shaders/PostProcessAmbientOcclusion.usf b/Engine/Shaders/PostProcessAmbientOcclusion.usf
--- a/Engine/Shaders/PostProcessAmbientOcclusion.usf
+++ b/Engine/Shaders/PostProcessAmbientOcclusion.usf
@@ -19,7 +19,7 @@
 // 0:no / 1:yes
 // USE_UPSAMPLE
 
-// 0:normal method using depth computed in setup passes / 1:faster but slightly lower quality
+// 0:normal method using depth computed in setup passes (might need a clear outside the viewport, test with stereo/splitscreen) / 1:faster but slightly lower quality
 #define USE_HZB	1
 
 // 0: classic with weighted sample, 1: don't normalize and adjust the formula to be simpler and faster - can look better and is cheaper (Alchemy like?)
