Two sided foliage shading model gets less view dependent contribution, matches observed leaves better, subsurface lighting looks brighter from a perpendicular angle to the light

diff --git a/Engine/Shaders/DeferredLightingCommon.usf b/Engine/Shaders/DeferredLightingCommon.usf
--- a/Engine/Shaders/DeferredLightingCommon.usf
+++ b/Engine/Shaders/DeferredLightingCommon.usf
@@ -257,7 +257,8 @@ float3 TwoSidedFoliageShadingSubsurface( float3 SubsurfaceColor, float3 L, float
 	float DiffuseWrap = .6f;
 	// Energy conserving wrap
 	float BackDiffuse = saturate(dot(-N, L) * DiffuseWrap * DiffuseWrap + (1 - DiffuseWrap) * DiffuseWrap) / PI;  
-	float ViewDependenceAmount = .5f;
+	// Leaves have very small amounts of directional backscattering
+	float ViewDependenceAmount = .2f;
 	float BackShading = lerp(BackDiffuse, VoLLobe, ViewDependenceAmount);
 
 	return BackShading * SubsurfaceColor;
