Fixed GLSL compile error in PostProcessAmbientOcclusion.usf

diff --git a/Engine/Shaders/PostProcessAmbientOcclusion.usf b/Engine/Shaders/PostProcessAmbientOcclusion.usf
--- a/Engine/Shaders/PostProcessAmbientOcclusion.usf
+++ b/Engine/Shaders/PostProcessAmbientOcclusion.usf
@@ -560,7 +560,7 @@ void MainPS(in float4 UVAndScreenPos : TEXCOORD0, out float4 OutColor : SV_Targe
 					// constant at run time
 					float Scale = (step + 1) / (float)SAMPLE_STEPS;
 					// constant at run time (higher is better for texture cache / performance, lower is better quality
-					float MipLevel = max(step, 1);
+					float MipLevel = max(step, (uint)1);
 
 					float2 StepSample = Wedge2(Scale, LocalRandom, FovFix, InvFovFix, ViewSpacePosition, ScaledViewSpaceNormal, InvHaloSize, MipLevel);
 
@@ -582,7 +582,7 @@ void MainPS(in float4 UVAndScreenPos : TEXCOORD0, out float4 OutColor : SV_Targe
 						// constant at run time
 						float Scale = (step + 1) / (float)SAMPLE_STEPS * mirror;
 						// constant at run time (higher is better for texture cache / performance, lower is better quality
-						float MipLevel = max(step, 1);
+						float MipLevel = max(step, (uint)1);
 
 						float2 StepSample = Wedge(Scale, LocalRandom, FovFix, InvFovFix, ViewSpacePosition, ScaledViewSpaceNormal, InvHaloSize, MipLevel);
 
