Fixed swizzled UV1 for displacement

diff --git a/Engine/Shaders/LocalVertexFactory.usf b/Engine/Shaders/LocalVertexFactory.usf
--- a/Engine/Shaders/LocalVertexFactory.usf
+++ b/Engine/Shaders/LocalVertexFactory.usf
@@ -628,7 +628,7 @@ float4 VertexFactoryGetPreviousWorldPosition(FVertexFactoryInput Input, FVertexF
 			Result.TexCoords[CoordinateIndex] = Interpolants.InterpolantsVSToPS.TexCoords[CoordinateIndex/2].xy;
 			if(CoordinateIndex + 1 < NUM_MATERIAL_TEXCOORDS)
 			{
-				Result.TexCoords[CoordinateIndex + 1] = Interpolants.InterpolantsVSToPS.TexCoords[CoordinateIndex/2].wz;
+				Result.TexCoords[CoordinateIndex + 1] = Interpolants.InterpolantsVSToPS.TexCoords[CoordinateIndex/2].zw;
 			}
 		}
 	#endif	// NUM_MATERIAL_TEXCOORDS
