UE4 - Stub shader file to unblock testing

diff --git a/Engine/Shaders/ScreenSpaceRayCast.usf b/Engine/Shaders/ScreenSpaceRayCast.usf
new file mode 100644
--- /dev/null
+++ b/Engine/Shaders/ScreenSpaceRayCast.usf
@@ -0,0 +1,11 @@
+// Copyright 1998-2015 Epic Games, Inc. All Rights Reserved.
+
+/*=============================================================================
+	ScreenSpaceRayCast.usf
+=============================================================================*/
+
+#ifndef __SCREENSPACERAYCAST__
+#define __SCREENSPACERAYCAST__
+
+
+#endif // __SCREENSPACERAYCAST__
