UE4 - Stub out functions

diff --git a/Engine/Shaders/ScreenSpaceRayCast.usf b/Engine/Shaders/ScreenSpaceRayCast.usf
--- a/Engine/Shaders/ScreenSpaceRayCast.usf
+++ b/Engine/Shaders/ScreenSpaceRayCast.usf
@@ -7,5 +7,15 @@
 #ifndef __SCREENSPACERAYCAST__
 #define __SCREENSPACERAYCAST__
 
+float4 RayCastCheap(Texture2D HZBTexture, SamplerState HZBSampler, float3 ReflectionVector, float ScreenPositionW, float3 WorldPosition_CamRelative, float StepOffset)
+{
+	// Really cheap ;)
+	return 0;
+}
+
+float4 RayCast(Texture2D PostprocessInput1, SamplerState PostprocessInput1Sampler, float3 R, float Roughness, float SceneDepth, float3 PositionTranslatedWorld, int NumSteps, float StepOffset )
+{
+	return 0;
+}
 
 #endif // __SCREENSPACERAYCAST__
