UE4 - More compile fixes

diff --git a/Engine/Shaders/ScreenSpaceRayCast.usf b/Engine/Shaders/ScreenSpaceRayCast.usf
--- a/Engine/Shaders/ScreenSpaceRayCast.usf
+++ b/Engine/Shaders/ScreenSpaceRayCast.usf
@@ -18,4 +18,9 @@ float4 RayCast(Texture2D PostprocessInput1, SamplerState PostprocessInput1Sample
 	return 0;
 }
 
+float4 SampleScreenColor(Texture2D PostprocessInput1, SamplerState PostprocessInput1Sampler, float3 HitUVz)
+{
+	return 1;
+}
+
 #endif // __SCREENSPACERAYCAST__
