Real version of file

diff --git a/Engine/Shaders/ScreenSpaceRayCast.usf b/Engine/Shaders/ScreenSpaceRayCast.usf
--- a/Engine/Shaders/ScreenSpaceRayCast.usf
+++ b/Engine/Shaders/ScreenSpaceRayCast.usf
@@ -1,26 +1,353 @@
 // Copyright 1998-2015 Epic Games, Inc. All Rights Reserved.
 
-/*=============================================================================
-	ScreenSpaceRayCast.usf
-=============================================================================*/
+#pragma once
 
-#ifndef __SCREENSPACERAYCAST__
-#define __SCREENSPACERAYCAST__
+#include "Common.usf"	
 
-float4 RayCastCheap(Texture2D HZBTexture, SamplerState HZBSampler, float3 ReflectionVector, float ScreenPositionW, float3 WorldPosition_CamRelative, float StepOffset)
+// 0:reference (slower, potentially higher quality) 1:use mips of the HZB depth (better performance)
+#define USE_HZB		1
+
+#define SCALAR_BRANCHLESS		0
+#define VECTORIZED_BRANCHLESS	0
+#define VECTORIZED_EARLY_OUT	1
+
+float4 SampleDepthTexture( Texture2D Texture, SamplerState Sampler, float Level, float4 SampleUV0, float4 SampleUV1 )
 {
-	// Really cheap ;)
-	return 0;
+	float4 SampleDepth;
+#if USE_HZB
+	SampleDepth.x = Texture2DSampleLevel( Texture, Sampler, SampleUV0.xy, Level ).r;
+	SampleDepth.y = Texture2DSampleLevel( Texture, Sampler, SampleUV0.zw, Level ).r;
+	SampleDepth.z = Texture2DSampleLevel( Texture, Sampler, SampleUV1.xy, Level ).r;
+	SampleDepth.w = Texture2DSampleLevel( Texture, Sampler, SampleUV1.zw, Level ).r;
+#else
+	SampleDepth.x = Texture2DSampleLevel( SceneDepthTexture, SceneDepthTextureSampler, SampleUV0.xy, 0 ).r;
+	SampleDepth.y = Texture2DSampleLevel( SceneDepthTexture, SceneDepthTextureSampler, SampleUV0.zw, 0 ).r;
+	SampleDepth.z = Texture2DSampleLevel( SceneDepthTexture, SceneDepthTextureSampler, SampleUV1.xy, 0 ).r;
+	SampleDepth.w = Texture2DSampleLevel( SceneDepthTexture, SceneDepthTextureSampler, SampleUV1.zw, 0 ).r;
+#endif
+	return SampleDepth;
 }
 
-float4 RayCast(Texture2D PostprocessInput1, SamplerState PostprocessInput1Sampler, float3 R, float Roughness, float SceneDepth, float3 PositionTranslatedWorld, int NumSteps, float StepOffset )
+float4 RayCast( Texture2D Texture, SamplerState Sampler, float3 R, float Roughness, float SceneDepth, float3 PositionTranslatedWorld, int NumSteps, float StepOffset )
 {
-	return 0;
+	// TODO provide RayStartUVz
+
+	// NOTE could clip ray against frustum planes
+
+	// TODO use screen position and skip matrix mul
+	float4 RayStartClip	= mul( float4( PositionTranslatedWorld, 1 ), View.TranslatedWorldToClip );
+	float4 RayEndClip	= mul( float4( PositionTranslatedWorld + R * SceneDepth, 1 ), View.TranslatedWorldToClip );
+
+	float3 RayStartScreen = RayStartClip.xyz / RayStartClip.w;
+	float3 RayEndScreen = RayEndClip.xyz / RayEndClip.w;
+
+	// Normalize 2D length
+	float3 RayStepScreen = ( RayEndScreen - RayStartScreen ) / length( RayEndScreen.xy - RayStartScreen.xy );
+	RayStepScreen *= 1.5;
+
+#if USE_HZB
+	float3 RayStartUVz = float3( RayStartScreen.xy * float2( 0.5, -0.5 ) + 0.5, RayStartScreen.z );
+	float3 RayStepUVz = float3( RayStepScreen.xy * float2( 0.5, -0.5 ), RayStepScreen.z );
+#else
+	float3 RayStartUVz = float3( RayStartScreen.xy * View.ScreenPositionScaleBias.xy + View.ScreenPositionScaleBias.wz, RayStartScreen.z );
+	float3 RayStepUVz = float3( RayStepScreen.xy * View.ScreenPositionScaleBias.xy, RayStepScreen.z );
+#endif
+	
+	const float Step = 1.0 / (NumSteps + 1);
+
+	// *2 to get less morie pattern in extreme cases, larger values make object appear not grounded in reflections
+	const float CompareTolerance = abs( RayStepUVz.z ) * Step * 2;
+
+	// avoid bugs with early returns inside of loops on certain platform compilers.
+	float4 Result = float4( 0, 0, 0, 1 );
+
+#if SCALAR_BRANCHLESS
+
+	float MinHitTime = 1;
+	float LastDiff = 0;
+
+	float SampleTime = StepOffset * Step + Step;
+
+	UNROLL
+	for( int i = 0; i < NumSteps; i++ )
+	{
+		float3 SampleUVz = RayStartUVz + RayStepUVz * SampleTime;
+		
+		// Use lower res for farther samples
+		float Level = Roughness * (i * 4.0 / NumSteps);
+		float SampleDepth = Texture.SampleLevel( Sampler, SampleUVz.xy, Level ).r;
+
+		float DepthDiff = SampleUVz.z - SampleDepth;
+		bool Hit = abs( DepthDiff + CompareTolerance ) < CompareTolerance;
+
+		// Find more accurate hit using line segment intersection
+		float TimeLerp = saturate( LastDiff / (LastDiff - DepthDiff) );
+		float IntersectTime = SampleTime + TimeLerp * Step - Step;
+		float HitTime = Hit ? IntersectTime : 1;
+		MinHitTime = min( MinHitTime, HitTime );
+
+		LastDiff = DepthDiff;
+
+		SampleTime += Step;
+	}
+
+	float3 HitUVz = RayStartUVz + RayStepUVz * MinHitTime;
+#if USE_HZB
+	HitUVz.xy = HitUVz.xy * float2( 2, -2 ) + float2( -1, 1 );
+	HitUVz.xy = HitUVz.xy * View.ScreenPositionScaleBias.xy + View.ScreenPositionScaleBias.wz;
+#endif
+	Result = float4( HitUVz, MinHitTime );
+
+#elif VECTORIZED_BRANCHLESS
+
+	float MinHitTime = 1;
+	float LastDiff = 0;
+	float Level = 0;
+	
+	// Vectorized to group fetches
+	float4 SampleTime = ( StepOffset + float4( 1, 2, 3, 4 ) ) * Step;
+	float4 SampleUV0 = RayStartUVz.xyxy + RayStepUVz.xyxy * SampleTime.xxyy;
+	float4 SampleUV1 = RayStartUVz.xyxy + RayStepUVz.xyxy * SampleTime.zzww;
+	float4 SampleZ   = RayStartUVz.zzzz + RayStepUVz.zzzz * SampleTime;
+
+	LOOP
+	for( int i = 0; i < NumSteps; i += 4 )
+	{
+		// Use lower res for farther samples
+		float4 SampleDepth = SampleDepthTexture( Texture, Sampler, Level, SampleUV0, SampleUV1 );
+
+		#if 1
+			float4 DepthDiff = SampleZ - SampleDepth;
+			float4 IntersectTime = ( SampleDepth - RayStartUVz.zzzz ) / RayStepUVz.zzzz;
+		
+			bool4 Hit = abs( DepthDiff + CompareTolerance ) < CompareTolerance;
+	
+			// If hit set to intersect time. If missed set to 1, beyond end of ray
+			float4 HitTime = Hit ? IntersectTime : 1;
+
+			// Take closest hit
+			HitTime.xy = min( HitTime.xy, HitTime.zw );
+			MinHitTime = min( HitTime.x, HitTime.y );
+		#else
+			// Line segment intersection
+			float4 DepthDiff1 = SampleZ - SampleDepth;
+			float4 DepthDiff0 = float4( LastDiff, DepthDiff1.xyz );
+			float4 TimeLerp = saturate( DepthDiff0 / (DepthDiff0 - DepthDiff1) );
+			float4 IntersectTime = SampleTime + (TimeLerp - 1) / (NumSteps + 1);
+		
+			bool4 Hit = abs( DepthDiff1 + CompareTolerance ) < CompareTolerance;
+	
+			// If hit set to intersect time. If missed set to 1, beyond end of ray
+			float4 HitTime = Hit ? IntersectTime : 1;
+
+			// Take closest hit
+			HitTime.xy = min( HitTime.xy, HitTime.zw );
+			MinHitTime = min( MinHitTime, min( HitTime.x, HitTime.y ) );
+
+			LastDiff = DepthDiff1.w;
+		#endif
+		Level += Roughness * (16.0 / NumSteps);
+
+		SampleTime	+= 4.0 / (NumSteps + 1);
+		SampleUV0	+= RayStepUVz.xyxy * 4.0 / (NumSteps + 1);
+		SampleUV1	+= RayStepUVz.xyxy * 4.0 / (NumSteps + 1);
+		SampleZ		+= RayStepUVz.zzzz * 4.0 / (NumSteps + 1);
+	}
+
+	float3 HitUVz = RayStartUVz + RayStepUVz * MinHitTime;
+#if USE_HZB
+	HitUVz.xy = HitUVz.xy = HitUVz.xy * float2( 2, -2 ) + float2( -1, 1 );
+	HitUVz.xy = HitUVz.xy * View.ScreenPositionScaleBias.xy + View.ScreenPositionScaleBias.wz;
+#endif
+	Result = float4( HitUVz, MinHitTime );
+
+#else // VECTORIZED_EARLY_OUT
+
+	float LastDiff = 0;
+	float Level = 0;
+
+	float4 SampleTime = ( StepOffset + float4( 1, 2, 3, 4 ) ) * Step;
+
+	LOOP
+	for( int i = 0; i < NumSteps; i += 4 )
+	{
+		// Vectorized to group fetches
+		float4 SampleUV0 = RayStartUVz.xyxy + RayStepUVz.xyxy * SampleTime.xxyy;
+		float4 SampleUV1 = RayStartUVz.xyxy + RayStepUVz.xyxy * SampleTime.zzww;
+		float4 SampleZ   = RayStartUVz.zzzz + RayStepUVz.zzzz * SampleTime;
+		
+		// Use lower res for farther samples
+		float4 SampleDepth = SampleDepthTexture( Texture, Sampler, Level, SampleUV0, SampleUV1 );
+
+		float4 DepthDiff1 = SampleZ - SampleDepth;
+		bool4 Hit = abs( -DepthDiff1 - CompareTolerance ) < CompareTolerance;
+
+		BRANCH if( any( Hit ) )
+		{
+			// Find more accurate hit using line segment intersection
+			float4 DepthDiff0 = float4( LastDiff, DepthDiff1.xyz );
+			float4 TimeLerp = saturate( DepthDiff0 / (DepthDiff0 - DepthDiff1) );
+			float4 IntersectTime = SampleTime + (TimeLerp - 1) / (NumSteps + 1);
+			float4 HitTime = Hit ? IntersectTime : 1;
+
+			// Take closest hit
+			HitTime.xy = min( HitTime.xy, HitTime.zw );
+			float MinHitTime = min( HitTime.x, HitTime.y );
+			
+			float3 HitUVz = RayStartUVz + RayStepUVz * MinHitTime;
+#if USE_HZB
+			HitUVz.xy = HitUVz.xy * float2( 2, -2 ) + float2( -1, 1 );
+			HitUVz.xy = HitUVz.xy * View.ScreenPositionScaleBias.xy + View.ScreenPositionScaleBias.wz;
+#endif
+			Result = float4( HitUVz, MinHitTime );
+			break;
+		}
+
+		LastDiff = DepthDiff1.w;
+		Level += Roughness * (16.0 / NumSteps);
+
+		SampleTime += 4.0 / (NumSteps + 1);
+	}
+
+#endif
+
+	return Result;
 }
 
-float4 SampleScreenColor(Texture2D PostprocessInput1, SamplerState PostprocessInput1Sampler, float3 HitUVz)
+float4 RayCastCheap( Texture2D Texture, SamplerState Sampler, float3 R, float SceneDepth, float3 PositionTranslatedWorld, float StepOffset )
 {
-	return 1;
+	// TODO provide RayStartUVz
+
+	// NOTE could clip ray against frustum planes
+
+	// TODO use screen position and skip matrix mul
+	float4 RayStartClip	= mul( float4( PositionTranslatedWorld, 1 ), View.TranslatedWorldToClip );
+	float4 RayEndClip	= mul( float4( PositionTranslatedWorld + R * SceneDepth, 1 ), View.TranslatedWorldToClip );
+
+	float3 RayStartScreen = RayStartClip.xyz / RayStartClip.w;
+	float3 RayEndScreen = RayEndClip.xyz / RayEndClip.w;
+
+	// Normalize 2D length
+	float3 RayStepScreen = ( RayEndScreen - RayStartScreen ) / length( RayEndScreen.xy - RayStartScreen.xy );
+	RayStepScreen *= 0.75;
+
+#if USE_HZB
+	float3 RayStartUVz = float3( RayStartScreen.xy * float2( 0.5, -0.5 ) + 0.5, RayStartScreen.z );
+	float3 RayStepUVz = float3( RayStepScreen.xy * float2( 0.5, -0.5 ), RayStepScreen.z );
+#else
+	float3 RayStartUVz = float3( RayStartScreen.xy * View.ScreenPositionScaleBias.xy + View.ScreenPositionScaleBias.wz, RayStartScreen.z );
+	float3 RayStepUVz = float3( RayStepScreen.xy * View.ScreenPositionScaleBias.xy, RayStepScreen.z );
+#endif
+	
+	const float Step = 1.0 / (4 + 1);
+
+	// *2 to get less morie pattern in extreme cases, larger values make object appear not grounded in reflections
+	const float CompareTolerance = abs( RayStepUVz.z ) * Step * 2;
+
+	// Vectorized to group fetches
+	float4 SampleTime = ( StepOffset + float4( 1, 2, 3, 4 ) ) * Step;
+	float4 SampleUV0 = RayStartUVz.xyxy + RayStepUVz.xyxy * SampleTime.xxyy;
+	float4 SampleUV1 = RayStartUVz.xyxy + RayStepUVz.xyxy * SampleTime.zzww;
+	float4 SampleZ   = RayStartUVz.zzzz + RayStepUVz.zzzz * SampleTime;
+
+	float4 SampleDepth;
+#if USE_HZB
+	SampleDepth.x = Texture2DSampleLevel( Texture, Sampler, SampleUV0.xy, 0 ).r;
+	SampleDepth.y = Texture2DSampleLevel( Texture, Sampler, SampleUV0.zw, 0 ).r;
+	SampleDepth.z = Texture2DSampleLevel( Texture, Sampler, SampleUV1.xy, 1 ).r;
+	SampleDepth.w = Texture2DSampleLevel( Texture, Sampler, SampleUV1.zw, 1 ).r;
+#else
+	SampleDepth.x = Texture2DSampleLevel( SceneDepthTexture, SceneDepthTextureSampler, SampleUV0.xy, 0 ).r;
+	SampleDepth.y = Texture2DSampleLevel( SceneDepthTexture, SceneDepthTextureSampler, SampleUV0.zw, 0 ).r;
+	SampleDepth.z = Texture2DSampleLevel( SceneDepthTexture, SceneDepthTextureSampler, SampleUV1.xy, 0 ).r;
+	SampleDepth.w = Texture2DSampleLevel( SceneDepthTexture, SceneDepthTextureSampler, SampleUV1.zw, 0 ).r;
+#endif
+
+	#if 1
+		float4 DepthDiff = SampleZ - SampleDepth;
+		float4 IntersectTime = ( SampleDepth - RayStartUVz.zzzz ) / RayStepUVz.zzzz;
+		
+		bool4 Hit = abs( DepthDiff + CompareTolerance ) < CompareTolerance;
+	
+		// If hit set to intersect time. If missed set to 1, beyond end of ray
+		float4 HitTime = Hit ? IntersectTime : 1;
+
+		// Take closest hit
+		HitTime.xy = min( HitTime.xy, HitTime.zw );
+		float MinHitTime = min( HitTime.x, HitTime.y );
+	#else
+		// Line segment intersection
+		float4 DepthDiff1 = SampleZ - SampleDepth;
+		float4 DepthDiff0 = float4( 0, DepthDiff1.xyz );
+		float4 TimeLerp = saturate( DepthDiff0 / (DepthDiff0 - DepthDiff1) );
+		float4 IntersectTime = SampleTime + (TimeLerp - 1) * Step;
+		
+		bool4 Hit = abs( DepthDiff1 + CompareTolerance ) < CompareTolerance;
+	
+		// If hit set to intersect time. If missed set to 1, beyond end of ray
+		float4 HitTime = Hit ? IntersectTime : 1;
+
+		// Take closest hit
+		HitTime.xy = min( HitTime.xy, HitTime.zw );
+		float MinHitTime = min( HitTime.x, HitTime.y );
+	#endif
+
+	float3 HitUVz = RayStartUVz + RayStepUVz * MinHitTime;
+
+	BRANCH
+	if( MinHitTime == 1 && RayStepUVz.z < 0 )
+	{
+		float4 RayInfClip	= mul( float4( PositionTranslatedWorld + R * 4000, 1 ), View.TranslatedWorldToClip );
+		float3 RayInfScreen = RayInfClip.xyz / RayInfClip.w;
+		float3 RayInfUVz = float3( RayInfScreen.xy * float2( 0.5, -0.5 ) + 0.5, RayInfScreen.z );
+
+		HitUVz = RayInfUVz;
+		MinHitTime = 0.5;
+	}
+
+#if USE_HZB
+	HitUVz.xy = HitUVz.xy = HitUVz.xy * float2( 2, -2 ) + float2( -1, 1 );
+	HitUVz.xy = HitUVz.xy * View.ScreenPositionScaleBias.xy + View.ScreenPositionScaleBias.wz;
+#endif
+	return float4( HitUVz, MinHitTime );
 }
 
-#endif // __SCREENSPACERAYCAST__
+float4 SampleScreenColor( Texture2D Texture, SamplerState Sampler, float3 HitUVz )
+{
+	float4 OutColor;
+
+#if PREV_FRAME_COLOR
+	// Find previous screen position for hit since color buffer is from last frame
+	// TODO combine to single matrix
+	float4 HitClip = float4( (HitUVz.xy - View.ScreenPositionScaleBias.wz) / View.ScreenPositionScaleBias.xy, HitUVz.z, 1 );
+	float4 HitTranslatedWorld = mul( HitClip, View.ClipToTranslatedWorld );
+	HitTranslatedWorld /= HitTranslatedWorld.w;
+		
+	float3 PrevTranslatedWorld = HitTranslatedWorld.xyz + (View.PrevPreViewTranslation - View.PreViewTranslation);
+	float4 PrevClip = mul( float4( PrevTranslatedWorld, 1 ), View.PrevTranslatedWorldToClip );
+	float2 PrevScreen = PrevClip.xy / PrevClip.w;
+	float2 PrevUV = PrevScreen.xy * View.ScreenPositionScaleBias.xy + View.ScreenPositionScaleBias.wz;
+
+	OutColor.rgb = Texture.SampleLevel( Sampler, PrevUV, 0 ).rgb;
+	OutColor.a = 1;
+
+	// Off screen masking
+	float2 Vignette = saturate( abs( PrevScreen ) * 5 - 4 );
+#else
+	OutColor.rgb = Texture.SampleLevel( Sampler, HitUVz.xy, 0 ).rgb;
+	OutColor.a = 1;
+
+	// Off screen masking
+	float2 HitScreenPos = ( HitUVz.xy - View.ScreenPositionScaleBias.wz ) / View.ScreenPositionScaleBias.xy;
+	float2 Vignette = saturate( abs( HitScreenPos ) * 5 - 4 );
+#endif
+		
+	//PrevScreen sometimes has NaNs or Infs.  DX11 is protected because saturate turns NaNs -> 0.
+	//Do a SafeSaturate so other platforms get the same protection.
+	OutColor *= SafeSaturate( 1.0 - dot( Vignette, Vignette ) );
+
+	// Transform NaNs to black, transform negative colors to black.
+	OutColor.rgb = -min(-OutColor.rgb, 0.0);
+
+	return OutColor;
+}
\ No newline at end of file
