UE4 - Correct delta velocity for WPO materials

diff --git a/Engine/Shaders/BasePassForForwardShadingVertexShader.usf b/Engine/Shaders/BasePassForForwardShadingVertexShader.usf
--- a/Engine/Shaders/BasePassForForwardShadingVertexShader.usf
+++ b/Engine/Shaders/BasePassForForwardShadingVertexShader.usf
@@ -45,7 +45,7 @@ void Main(
 	float4 WorldPosition = WorldPositionExcludingWPO;
 
 	half3x3 TangentToLocal = VertexFactoryGetTangentToLocal(Input, VFIntermediates);	
-	FMaterialVertexParameters VertexParameters = GetMaterialVertexParameters(Input, VFIntermediates, WorldPosition.xyz, TangentToLocal);
+	FMaterialVertexParameters VertexParameters = GetMaterialVertexParameters(Input, VFIntermediates, WorldPosition.xyz, TangentToLocal, View.GameTime, View.RealTime);
 
 	half3 WorldPositionOffset = GetMaterialWorldPositionOffset(VertexParameters);
 	
diff --git a/Engine/Shaders/BasePassVertexShader.usf b/Engine/Shaders/BasePassVertexShader.usf
--- a/Engine/Shaders/BasePassVertexShader.usf
+++ b/Engine/Shaders/BasePassVertexShader.usf
@@ -18,7 +18,7 @@ void Main(
 	float4 WorldPosition = WorldPositionExcludingWPO;
 
 	float3x3 TangentToLocal = VertexFactoryGetTangentToLocal(Input, VFIntermediates);	
-	FMaterialVertexParameters VertexParameters = GetMaterialVertexParameters(Input, VFIntermediates, WorldPosition.xyz, TangentToLocal);
+	FMaterialVertexParameters VertexParameters = GetMaterialVertexParameters(Input, VFIntermediates, WorldPosition.xyz, TangentToLocal, View.GameTime, View.RealTime);
 
 	// Isolate instructions used for world position offset
 	// As these cause the optimizer to generate different position calculating instructions in each pass, resulting in self-z-fighting.
diff --git a/Engine/Shaders/DepthOnlyVertexShader.usf b/Engine/Shaders/DepthOnlyVertexShader.usf
--- a/Engine/Shaders/DepthOnlyVertexShader.usf
+++ b/Engine/Shaders/DepthOnlyVertexShader.usf
@@ -83,7 +83,7 @@ void Main(
 	float4 WorldPos = VertexFactoryGetWorldPosition(Input, VFIntermediates);
 
 	float3x3 TangentToLocal = VertexFactoryGetTangentToLocal(Input, VFIntermediates);
-	FMaterialVertexParameters VertexParameters = GetMaterialVertexParameters(Input, VFIntermediates, WorldPos.xyz, TangentToLocal);
+	FMaterialVertexParameters VertexParameters = GetMaterialVertexParameters(Input, VFIntermediates, WorldPos.xyz, TangentToLocal, View.GameTime, View.RealTime);
 
 	// Isolate instructions used for world position offset
 	// As these cause the optimizer to generate different position calculating instructions in each pass, resulting in self-z-fighting.
diff --git a/Engine/Shaders/DistortAccumulateVS.usf b/Engine/Shaders/DistortAccumulateVS.usf
--- a/Engine/Shaders/DistortAccumulateVS.usf
+++ b/Engine/Shaders/DistortAccumulateVS.usf
@@ -74,7 +74,7 @@ void Main(
 	float4	WorldPosition = VertexFactoryGetWorldPosition(Input, VFIntermediates);
 	float3x3 TangentToLocal = VertexFactoryGetTangentToLocal(Input, VFIntermediates);
 
-	FMaterialVertexParameters VertexParameters = GetMaterialVertexParameters(Input, VFIntermediates, WorldPosition.xyz, TangentToLocal);
+	FMaterialVertexParameters VertexParameters = GetMaterialVertexParameters(Input, VFIntermediates, WorldPosition.xyz, TangentToLocal, View.GameTime, View.RealTime);
 	// Isolate instructions used for world position offset on xbox 360, 
 	// As these cause the optimizer to generate different position calculating instructions in each pass, resulting in self-z-fighting.
 	// This is only necessary for shaders used in passes that have depth testing enabled.
diff --git a/Engine/Shaders/GpuSkinVertexFactory.usf b/Engine/Shaders/GpuSkinVertexFactory.usf
--- a/Engine/Shaders/GpuSkinVertexFactory.usf
+++ b/Engine/Shaders/GpuSkinVertexFactory.usf
@@ -203,12 +203,14 @@ struct FVertexFactoryIntermediates
 };
 
 /** Converts from vertex factory specific input to a FMaterialVertexParameters, which is used by vertex shader material inputs. */
-FMaterialVertexParameters GetMaterialVertexParameters( FVertexFactoryInput Input, FVertexFactoryIntermediates Intermediates, float3 WorldPosition, float3x3 TangentToLocal )
+FMaterialVertexParameters GetMaterialVertexParameters( FVertexFactoryInput Input, FVertexFactoryIntermediates Intermediates, float3 WorldPosition, float3x3 TangentToLocal, float GameTime, float RealTime)
 {
 	FMaterialVertexParameters Result = (FMaterialVertexParameters)0;
 	Result.WorldPosition = WorldPosition - View.PreViewTranslation.xyz;
 	Result.VertexColor = Intermediates.Color;
 	Result.TangentToWorld = mul(TangentToLocal, GetLocalToWorld3x3());
+	Result.GameTime = GameTime;
+	Result.RealTime = RealTime;
 
 #if NUM_MATERIAL_TEXCOORDS_VERTEX
 	for(int CoordinateIndex = 0; CoordinateIndex < NUM_MATERIAL_TEXCOORDS_VERTEX; CoordinateIndex++)
diff --git a/Engine/Shaders/HitProxyVertexShader.usf b/Engine/Shaders/HitProxyVertexShader.usf
--- a/Engine/Shaders/HitProxyVertexShader.usf
+++ b/Engine/Shaders/HitProxyVertexShader.usf
@@ -93,7 +93,7 @@ void Main(
 	float4	WorldPosition = VertexFactoryGetWorldPosition(Input, VFIntermediates);
 	float3x3 TangentToLocal = VertexFactoryGetTangentToLocal(Input, VFIntermediates);
 
-	FMaterialVertexParameters VertexParameters = GetMaterialVertexParameters(Input, VFIntermediates, WorldPosition.xyz, TangentToLocal);
+	FMaterialVertexParameters VertexParameters = GetMaterialVertexParameters(Input, VFIntermediates, WorldPosition.xyz, TangentToLocal, View.GameTime, View.RealTime);
 	WorldPosition.xyz += GetMaterialWorldPositionOffset(VertexParameters);
 
 #if USE_INSTANCING
diff --git a/Engine/Shaders/LandscapeVertexFactory.usf b/Engine/Shaders/LandscapeVertexFactory.usf
--- a/Engine/Shaders/LandscapeVertexFactory.usf
+++ b/Engine/Shaders/LandscapeVertexFactory.usf
@@ -359,12 +359,14 @@ FMaterialPixelParameters GetMaterialPixelParameters(FVertexFactoryInterpolantsVS
 }
 
 /** Converts from vertex factory specific input to a FMaterialVertexParameters, which is used by vertex shader material inputs. */
-FMaterialVertexParameters GetMaterialVertexParameters( FVertexFactoryInput Input, FVertexFactoryIntermediates Intermediates, float3 WorldPosition, float3x3 TangentToLocal )
+FMaterialVertexParameters GetMaterialVertexParameters( FVertexFactoryInput Input, FVertexFactoryIntermediates Intermediates, float3 WorldPosition, float3x3 TangentToLocal, float GameTime, float RealTime)
 {
 	FMaterialVertexParameters Result = (FMaterialVertexParameters)0;
 	Result.WorldPosition = WorldPosition - View.PreViewTranslation;
 	Result.VertexColor = float4(1,1,1,1);
 	Result.TangentToWorld = mul(TangentToLocal, GetLocalToWorld3x3()); 
+	Result.GameTime = GameTime;
+	Result.RealTime = RealTime;
 
 	FLandscapeTexCoords LandscapeTexCoords = GetLandscapeTexCoords(Input, Intermediates);
 
diff --git a/Engine/Shaders/LightMapDensityShader.usf b/Engine/Shaders/LightMapDensityShader.usf
--- a/Engine/Shaders/LightMapDensityShader.usf
+++ b/Engine/Shaders/LightMapDensityShader.usf
@@ -78,7 +78,7 @@ void MainVertexShader(
 	float4 WorldPosition = VertexFactoryGetWorldPosition(Input, VFIntermediates);
 	float3x3 TangentToLocal = VertexFactoryGetTangentToLocal(Input, VFIntermediates);
 
-	FMaterialVertexParameters VertexParameters = GetMaterialVertexParameters(Input, VFIntermediates, WorldPosition.xyz, TangentToLocal);
+	FMaterialVertexParameters VertexParameters = GetMaterialVertexParameters(Input, VFIntermediates, WorldPosition.xyz, TangentToLocal, View.GameTime, View.RealTime);
 	WorldPosition.xyz += GetMaterialWorldPositionOffset(VertexParameters);
 
 #if USING_TESSELLATION
diff --git a/Engine/Shaders/LocalVertexFactory.usf b/Engine/Shaders/LocalVertexFactory.usf
--- a/Engine/Shaders/LocalVertexFactory.usf
+++ b/Engine/Shaders/LocalVertexFactory.usf
@@ -250,11 +250,13 @@ half3x3 CalcTangentToWorldNoScale(in half3x3 TangentToLocal)
 }
 
 /** Converts from vertex factory specific input to a FMaterialVertexParameters, which is used by vertex shader material inputs. */
-FMaterialVertexParameters GetMaterialVertexParameters( FVertexFactoryInput Input, FVertexFactoryIntermediates Intermediates, float3 WorldPosition, half3x3 TangentToLocal )
+FMaterialVertexParameters GetMaterialVertexParameters( FVertexFactoryInput Input, FVertexFactoryIntermediates Intermediates, float3 WorldPosition, half3x3 TangentToLocal, float GameTime, float RealTime)
 {
 	FMaterialVertexParameters Result = (FMaterialVertexParameters)0;
 	Result.WorldPosition = WorldPosition - View.PreViewTranslation;
 	Result.VertexColor = Intermediates.Color;
+	Result.GameTime = GameTime;
+	Result.RealTime = RealTime;
 
 	// does not handle instancing!
 	Result.TangentToWorld = Intermediates.TangentToWorld;
diff --git a/Engine/Shaders/MaterialTemplate.usf b/Engine/Shaders/MaterialTemplate.usf
--- a/Engine/Shaders/MaterialTemplate.usf
+++ b/Engine/Shaders/MaterialTemplate.usf
@@ -242,10 +242,13 @@ struct FMaterialVertexParameters
 	half4 VertexColor;
 #if NUM_MATERIAL_TEXCOORDS_VERTEX
 	float2 TexCoords[NUM_MATERIAL_TEXCOORDS_VERTEX];
-#if (ES2_PROFILE || ES3_1_PROFILE)
+	#if (ES2_PROFILE || ES3_1_PROFILE)
 	float2 TexCoordOffset; // Offset for UV localization for large UV values
+	#endif
 #endif
-#endif
+
+	float GameTime;
+	float RealTime;
 
 	/** Per-particle properties. Only valid for particle vertex factories. */
 	FMaterialParticleParameters Particle;
diff --git a/Engine/Shaders/MeshParticleVertexFactory.usf b/Engine/Shaders/MeshParticleVertexFactory.usf
--- a/Engine/Shaders/MeshParticleVertexFactory.usf
+++ b/Engine/Shaders/MeshParticleVertexFactory.usf
@@ -215,11 +215,13 @@ FMaterialPixelParameters GetMaterialPixelParameters(FVertexFactoryInterpolantsVS
 }
 
 /** Converts from vertex factory specific input to a FMaterialVertexParameters, which is used by vertex shader material inputs. */
-FMaterialVertexParameters GetMaterialVertexParameters( FVertexFactoryInput Input, FVertexFactoryIntermediates Intermediates, float3 WorldPosition, float3x3 TangentToLocal )
+FMaterialVertexParameters GetMaterialVertexParameters( FVertexFactoryInput Input, FVertexFactoryIntermediates Intermediates, float3 WorldPosition, float3x3 TangentToLocal, float GameTime, float RealTime)
 {
 	FMaterialVertexParameters Result = (FMaterialVertexParameters)0;
 	Result.WorldPosition = WorldPosition - View.PreViewTranslation;
 	Result.VertexColor = Intermediates.VertexColor;
+	Result.GameTime = GameTime;
+	Result.RealTime = RealTime;
 
 	Result.Particle.PositionAndSize.xyz = Intermediates.ParticleWorldPosition.xyz;
 	Result.Particle.PositionAndSize.w = 1;
diff --git a/Engine/Shaders/ParticleBeamTrailVertexFactory.usf b/Engine/Shaders/ParticleBeamTrailVertexFactory.usf
--- a/Engine/Shaders/ParticleBeamTrailVertexFactory.usf
+++ b/Engine/Shaders/ParticleBeamTrailVertexFactory.usf
@@ -88,13 +88,15 @@ FMaterialPixelParameters GetMaterialPixelParameters(FVertexFactoryInterpolantsVS
 }
 
 /** Converts from vertex factory specific input to a FMaterialVertexParameters, which is used by vertex shader material inputs. */
-FMaterialVertexParameters GetMaterialVertexParameters( FVertexFactoryInput Input, FVertexFactoryIntermediates Intermediates, float3 WorldPosition, float3x3 TangentToLocal )
+FMaterialVertexParameters GetMaterialVertexParameters( FVertexFactoryInput Input, FVertexFactoryIntermediates Intermediates, float3 WorldPosition, float3x3 TangentToLocal, float GameTime, float RealTime)
 {
 	FMaterialVertexParameters Result = (FMaterialVertexParameters)0;
 	Result.WorldPosition = WorldPosition - View.PreViewTranslation.xyz;
 	Result.VertexColor = Input.Color;
 	Result.TangentToWorld = mul(TangentToLocal, GetLocalToWorld3x3()); 
 	Result.Particle.Color = Intermediates.Color;
+	Result.GameTime = GameTime;
+	Result.RealTime = RealTime;
 
 #if NUM_MATERIAL_TEXCOORDS_VERTEX
 	for(int CoordinateIndex = 0; CoordinateIndex < NUM_MATERIAL_TEXCOORDS_VERTEX; CoordinateIndex++)
diff --git a/Engine/Shaders/ParticleGPUSpriteVertexFactory.usf b/Engine/Shaders/ParticleGPUSpriteVertexFactory.usf
--- a/Engine/Shaders/ParticleGPUSpriteVertexFactory.usf
+++ b/Engine/Shaders/ParticleGPUSpriteVertexFactory.usf
@@ -407,12 +407,14 @@ FMaterialPixelParameters GetMaterialPixelParameters( FVertexFactoryInterpolantsV
  * @param TangentToLocal - The tangent basis for this vertex.
  * @returns per-vertex material parameters.
  */
-FMaterialVertexParameters GetMaterialVertexParameters( FVertexFactoryInput Input, FVertexFactoryIntermediates Intermediates, float3 WorldPosition, float3x3 TangentToLocal )
+FMaterialVertexParameters GetMaterialVertexParameters( FVertexFactoryInput Input, FVertexFactoryIntermediates Intermediates, float3 WorldPosition, float3x3 TangentToLocal, float GameTime, float RealTime)
 {
 	FMaterialVertexParameters Result = (FMaterialVertexParameters)0;
 	Result.WorldPosition = WorldPosition - View.PreViewTranslation.xyz;
 	Result.VertexColor = 1;
 	Result.TangentToWorld = mul(TangentToLocal, GetLocalToWorld3x3()); 
+	Result.GameTime = GameTime;
+	Result.RealTime = RealTime;
 
 #if USE_PARTICLE_TIME
 	Result.Particle.RelativeTime = Intermediates.RelativeTime;
diff --git a/Engine/Shaders/ParticleSpriteVertexFactory.usf b/Engine/Shaders/ParticleSpriteVertexFactory.usf
--- a/Engine/Shaders/ParticleSpriteVertexFactory.usf
+++ b/Engine/Shaders/ParticleSpriteVertexFactory.usf
@@ -227,7 +227,7 @@ FMaterialPixelParameters GetMaterialPixelParameters(FVertexFactoryInterpolantsVS
 }
 
 /** Converts from vertex factory specific input to a FMaterialVertexParameters, which is used by vertex shader material inputs. */
-FMaterialVertexParameters GetMaterialVertexParameters(FVertexFactoryInput Input, FVertexFactoryIntermediates Intermediates, float3 WorldPosition, float3x3 TangentToLocal)
+FMaterialVertexParameters GetMaterialVertexParameters(FVertexFactoryInput Input, FVertexFactoryIntermediates Intermediates, float3 WorldPosition, float3x3 TangentToLocal, float GameTime, float RealTime)
 {
 	FMaterialVertexParameters Result = (FMaterialVertexParameters)0;
 	Result.WorldPosition = WorldPosition - View.PreViewTranslation.xyz;
@@ -236,6 +236,8 @@ FMaterialVertexParameters GetMaterialVertexParameters(FVertexFactoryInput Input,
 	Result.Particle.MacroUV = SpriteVF.MacroUVParameters;
 	Result.Particle.Color = Intermediates.Color;
 	Result.Particle.MotionBlurFade = 1.0f;
+	Result.GameTime = GameTime;
+	Result.RealTime = RealTime;
 
 #if USE_DYNAMIC_PARAMETERS
 	Result.Particle.DynamicParameter = Intermediates.DynamicParameter;
diff --git a/Engine/Shaders/ShadowDepthVertexShader.usf b/Engine/Shaders/ShadowDepthVertexShader.usf
--- a/Engine/Shaders/ShadowDepthVertexShader.usf
+++ b/Engine/Shaders/ShadowDepthVertexShader.usf
@@ -188,7 +188,7 @@ void Main(
 	float4 WorldPos = VertexFactoryGetWorldPosition(Input, VFIntermediates);
 	float3x3 TangentToLocal = VertexFactoryGetTangentToLocal(Input, VFIntermediates);
 
-	FMaterialVertexParameters VertexParameters = GetMaterialVertexParameters(Input, VFIntermediates, WorldPos.xyz, TangentToLocal);
+	FMaterialVertexParameters VertexParameters = GetMaterialVertexParameters(Input, VFIntermediates, WorldPos.xyz, TangentToLocal, View.GameTime, View.RealTime);
 	WorldPos.xyz += GetMaterialWorldPositionOffset(VertexParameters);
 
 #if USING_TESSELLATION
diff --git a/Engine/Shaders/TranslucentShadowDepthShaders.usf b/Engine/Shaders/TranslucentShadowDepthShaders.usf
--- a/Engine/Shaders/TranslucentShadowDepthShaders.usf
+++ b/Engine/Shaders/TranslucentShadowDepthShaders.usf
@@ -53,7 +53,7 @@ void MainVS(
 	float4 WorldPos = VertexFactoryGetWorldPosition(Input, VFIntermediates);
 	float3x3 TangentToLocal = VertexFactoryGetTangentToLocal(Input, VFIntermediates);
 
-	FMaterialVertexParameters VertexParameters = GetMaterialVertexParameters(Input, VFIntermediates, WorldPos.xyz, TangentToLocal);
+	FMaterialVertexParameters VertexParameters = GetMaterialVertexParameters(Input, VFIntermediates, WorldPos.xyz, TangentToLocal, View.GameTime, View.RealTime);
 	WorldPos.xyz += GetMaterialWorldPositionOffset(VertexParameters);
 
 	SetShadowDepthOutputs(
diff --git a/Engine/Shaders/VectorFieldVisualizationVertexFactory.usf b/Engine/Shaders/VectorFieldVisualizationVertexFactory.usf
--- a/Engine/Shaders/VectorFieldVisualizationVertexFactory.usf
+++ b/Engine/Shaders/VectorFieldVisualizationVertexFactory.usf
@@ -106,12 +106,14 @@ FMaterialPixelParameters GetMaterialPixelParameters( FVertexFactoryInterpolantsV
  * @param TangentToLocal - The tangent basis for this vertex.
  * @returns per-vertex material parameters.
  */
-FMaterialVertexParameters GetMaterialVertexParameters( FVertexFactoryInput Input, FVertexFactoryIntermediates Intermediates, float3 WorldPosition, float3x3 TangentToLocal )
+FMaterialVertexParameters GetMaterialVertexParameters( FVertexFactoryInput Input, FVertexFactoryIntermediates Intermediates, float3 WorldPosition, float3x3 TangentToLocal, float GameTime, float RealTime)
 {
 	FMaterialVertexParameters Result = (FMaterialVertexParameters)0;
 	Result.WorldPosition = WorldPosition - View.PreViewTranslation.xyz;
 	Result.VertexColor = float4(1,1,1,1);
 	Result.TangentToWorld = mul(TangentToLocal, GetLocalToWorld3x3()); 
+	Result.GameTime = GameTime;
+	Result.RealTime = RealTime;
 	return Result;
 }
 
diff --git a/Engine/Shaders/VelocityShader.usf b/Engine/Shaders/VelocityShader.usf
--- a/Engine/Shaders/VelocityShader.usf
+++ b/Engine/Shaders/VelocityShader.usf
@@ -126,7 +126,7 @@ void MainVertexShader(
 	float4 TranslatedWorldPosition = VertexFactoryGetWorldPosition( Input, VFIntermediates );
 	float3x3 TangentToLocal = VertexFactoryGetTangentToLocal( Input, VFIntermediates );
 
-	FMaterialVertexParameters VertexParameters = GetMaterialVertexParameters(Input, VFIntermediates, TranslatedWorldPosition.xyz, TangentToLocal);
+	FMaterialVertexParameters VertexParameters = GetMaterialVertexParameters(Input, VFIntermediates, TranslatedWorldPosition.xyz, TangentToLocal, View.GameTime, View.RealTime);
 	// Isolate instructions used for world position offset on xbox 360, 
 	// As these cause the optimizer to generate different position calculating instructions in each pass, resulting in self-z-fighting.
 	// This is only necessary for shaders used in passes that have depth testing enabled.
@@ -142,7 +142,7 @@ void MainVertexShader(
 #endif // #if USING_TESSELLATION
 	float4 PrevTranslatedWorldPosition = VertexFactoryGetPreviousWorldPosition( Input, VFIntermediates );	
 
-	VertexParameters = GetMaterialVertexParameters(Input, VFIntermediates, PrevTranslatedWorldPosition.xyz, TangentToLocal);
+	VertexParameters = GetMaterialVertexParameters(Input, VFIntermediates, PrevTranslatedWorldPosition.xyz, TangentToLocal, View.PrevFrameGameTime, View.PrevFrameRealTime);
 	PrevTranslatedWorldPosition.xyz += GetMaterialWorldPositionOffset(VertexParameters);
 	
 #if USING_TESSELLATION
diff --git a/Engine/Source/Runtime/Engine/Private/Materials/HLSLMaterialTranslator.h b/Engine/Source/Runtime/Engine/Private/Materials/HLSLMaterialTranslator.h
--- a/Engine/Source/Runtime/Engine/Private/Materials/HLSLMaterialTranslator.h
+++ b/Engine/Source/Runtime/Engine/Private/Materials/HLSLMaterialTranslator.h
@@ -1515,7 +1515,14 @@ protected:
 	{
 		if (!bPeriodic)
 		{
-			return AddInlinedCodeChunk(MCT_Float, TEXT("View.GameTime"));
+			if (ShaderFrequency == SF_Vertex)
+			{
+				return AddInlinedCodeChunk(MCT_Float, TEXT("Parameters.GameTime"));
+			}
+			else
+			{
+				return AddInlinedCodeChunk(MCT_Float, TEXT("View.GameTime"));
+			}
 		}
 		else if (Period == 0.0f)
 		{
@@ -1535,7 +1542,14 @@ protected:
 	{
 		if (!bPeriodic)
 		{
-			return AddInlinedCodeChunk(MCT_Float, TEXT("View.RealTime"));
+			if (ShaderFrequency == SF_Vertex)
+			{
+				return AddInlinedCodeChunk(MCT_Float, TEXT("Parameters.RealTime"));
+			}
+			else
+			{
+				return AddInlinedCodeChunk(MCT_Float, TEXT("View.RealTime"));
+			}
 		}
 		else if (Period == 0.0f)
 		{
