Removed unused, incorrectly calculated InstancedStaticMesh per-instance inverse matrix.

diff --git a/Engine/Shaders/LocalVertexFactory.usf b/Engine/Shaders/LocalVertexFactory.usf
--- a/Engine/Shaders/LocalVertexFactory.usf
+++ b/Engine/Shaders/LocalVertexFactory.usf
@@ -35,7 +35,7 @@ float4 InstancingFadeOutParams;
 	// Required for CPU emulation of Instancing
 	float4 CPUInstanceShadowMapBias;
 	float3x4 CPUInstanceTransform;
-	float3x4 CPUInstanceInverseTransform;
+	float4 CPUInstanceLightmapUVBias;
 #endif	// USE_INSTANCING_EMULATED
 
 #else // USE_INSTANCING
@@ -115,10 +115,7 @@ struct FVertexFactoryInput
 	float4 InstanceTransform1 : ATTRIBUTE9;
 	float4 InstanceTransform2 : ATTRIBUTE10;
 	float4 InstanceTransform3 : ATTRIBUTE11;
-	float4 InstanceInverseTransform1 : ATTRIBUTE12;
-	float4 InstanceInverseTransform2 : ATTRIBUTE13;
-	float4 InstanceInverseTransform3 : ATTRIBUTE14;
-	// note, the per-instance lightmap bias is stored in the W of InverseTransform1 and 2 and per-instance random in 3.
+	float4 InstanceLightmapUVBias : ATTRIBUTE12; /* per-instance random in Z */
 #else
 
 #endif //USE_INSTANCING && !USE_INSTANCING_EMULATED
@@ -171,15 +168,6 @@ float4x4 GetInstanceTransform(FPositionOnlyVertexFactoryInput Input)
 #endif	// !USE_INSTANCING_EMULATED
 }
 
-float3x3 GetInstanceInverseTransform(FVertexFactoryInput Input)
-{
-	#if !USE_INSTANCING_EMULATED
-	return float3x3(Input.InstanceInverseTransform1.xyz, Input.InstanceInverseTransform2.xyz, Input.InstanceInverseTransform3.xyz);
-	#else
-	return float3x3(CPUInstanceInverseTransform[0].xyz, CPUInstanceInverseTransform[1].xyz, CPUInstanceInverseTransform[2].xyz);
-#endif	// !USE_INSTANCING_EMULATED
-}
-
 float4 GetInstanceShadowMapBias(FVertexFactoryInput Input)
 {
 	#if !USE_INSTANCING_EMULATED
@@ -261,7 +249,6 @@ FMaterialVertexParameters GetMaterialVertexParameters(FVertexFactoryInput Input,
 
 #if USE_INSTANCING
 	Result.InstanceLocalToWorld = mul(transpose(GetInstanceTransform(Input)), Primitive.LocalToWorld);
-	Result.InstanceWorldToLocal = mul(float3x3(Primitive.WorldToLocal[0].xyz,Primitive.WorldToLocal[1].xyz,Primitive.WorldToLocal[2].xyz), transpose(GetInstanceInverseTransform(Input)));
 	Result.InstanceLocalPosition = Input.Position.xyz;
 	Result.PerInstanceParams = Intermediates.PerInstanceParams;
 #endif	// USE_INSTANCING
@@ -463,10 +450,10 @@ FVertexFactoryIntermediates GetVertexFactoryIntermediates(FVertexFactoryInput In
 	float SelectedValue = trunc(GetInstanceShadowMapBias(Input).z * (1.0 / 256.0));
 	// InstancingFadeOutParams.z,w are RenderSelected and RenderDeselected respectively.
 	#if !USE_INSTANCING_EMULATED
-		Intermediates.PerInstanceParams.x = Input.InstanceInverseTransform3.w;
+		Intermediates.PerInstanceParams.x = Input.InstanceLightmapUVBias.z;
 		float3 InstanceLocation = TransformLocalToWorld(float3(Input.InstanceTransform1.w, Input.InstanceTransform2.w, Input.InstanceTransform3.w)).xyz;
 	#else
-		Intermediates.PerInstanceParams.x = CPUInstanceInverseTransform[2].w;
+		Intermediates.PerInstanceParams.x = CPUInstanceLightmapUVBias.z;
 		float3 InstanceLocation = TransformLocalToWorld(float3(CPUInstanceTransform[0].w, CPUInstanceTransform[1].w, CPUInstanceTransform[2].w)).xyz;
 	#endif // !USE_INSTANCING_EMULATED
 	Intermediates.PerInstanceParams.y = 1.0 - saturate((length(InstanceLocation + View.PreViewTranslation.xyz) - InstancingFadeOutParams.x) * InstancingFadeOutParams.y);
@@ -541,11 +528,11 @@ FVertexFactoryInterpolantsVSToPS VertexFactoryGetInterpolantsVSToPS(FVertexFacto
 
 	#if USE_INSTANCING
 		#if !USE_INSTANCING_EMULATED
-		float2 InstanceLightMapBias = float2(Input.InstanceInverseTransform1.w, Input.InstanceInverseTransform2.w);
+		float2 InstanceLightmapUVBias = Input.InstanceLightmapUVBias.xy;
 		#else
-		float2 InstanceLightMapBias = float2(CPUInstanceInverseTransform[1].w, CPUInstanceInverseTransform[2].w);
+		float2 InstanceLightmapUVBias = CPUInstanceLightmapUVBias.xy
 		#endif // !USE_INSTANCING_EMULATED
-		LightMapCoordinate = Input.LightMapCoordinate * LightMapCoordinateScaleBias.xy + InstanceLightMapBias;
+		LightMapCoordinate = Input.LightMapCoordinate * LightMapCoordinateScaleBias.xy + InstanceLightmapUVBias;
 	#else
 		LightMapCoordinate = Input.LightMapCoordinate * LightMapCoordinateScaleBias.xy + LightMapCoordinateScaleBias.zw;
 	#endif
diff --git a/Engine/Shaders/MaterialTemplate.usf b/Engine/Shaders/MaterialTemplate.usf
--- a/Engine/Shaders/MaterialTemplate.usf
+++ b/Engine/Shaders/MaterialTemplate.usf
@@ -231,7 +231,6 @@ struct FMaterialVertexParameters
 #if USE_INSTANCING
 	/** Per-instance properties. */
 	float4x4 InstanceLocalToWorld;
-	float3x3 InstanceWorldToLocal;
 	float3 InstanceLocalPosition;
 	float4 PerInstanceParams;
 #elif PARTICLE_MESH_FACTORY
@@ -695,7 +694,6 @@ float3 GetSpeedTreeVertexOffset(FMaterialVertexParameters Parameters, int Geomet
 
 	#if USE_INSTANCING
 		float3x3 LocalToWorld = (float3x3)Parameters.InstanceLocalToWorld;
-		float3x3 WorldToLocal = Parameters.InstanceWorldToLocal;
 		float3 LocalPosition = Parameters.InstanceLocalPosition;
 
 		// skip if this instance is hidden
@@ -705,7 +703,6 @@ float3 GetSpeedTreeVertexOffset(FMaterialVertexParameters Parameters, int Geomet
 		}
 	#else
 		float3x3 LocalToWorld = (float3x3)Primitive.LocalToWorld;
-		float3x3 WorldToLocal = (float3x3)Primitive.WorldToLocal;
 		float3 LocalPosition = mul(float4(Parameters.WorldPosition.xyz, 1), Primitive.WorldToLocal).xyz;
 	#endif
 
diff --git a/Engine/Source/Runtime/Engine/Private/InstancedStaticMesh.cpp b/Engine/Source/Runtime/Engine/Private/InstancedStaticMesh.cpp
--- a/Engine/Source/Runtime/Engine/Private/InstancedStaticMesh.cpp
+++ b/Engine/Source/Runtime/Engine/Private/InstancedStaticMesh.cpp
@@ -54,13 +54,7 @@ static void FillInstanceRenderData(const FInstancedStaticMeshInstanceData& Insta
 	RenderData[2] = FVector4(Instance.Transform.M[0][1], Instance.Transform.M[1][1], Instance.Transform.M[2][1], Instance.Transform.M[3][1]);
 	RenderData[3] = FVector4(Instance.Transform.M[0][2], Instance.Transform.M[1][2], Instance.Transform.M[2][2], Instance.Transform.M[3][2]);
 
-	// Instance -> local rotation matrix (3x3)
-	// hide the offset (bias) of the lightmap and the per-instance random id in the matrix's w
-	const FMatrix Inverse = Instance.Transform;
-
-	RenderData[4] = FVector4(Inverse.M[0][0], Inverse.M[1][0], Inverse.M[2][0], Instance.LightmapUVBias.X);
-	RenderData[5] = FVector4(Inverse.M[0][1], Inverse.M[1][1], Inverse.M[2][1], Instance.LightmapUVBias.Y);
-	RenderData[6] = FVector4(Inverse.M[0][2], Inverse.M[1][2], Inverse.M[2][2], RandomInstanceID);
+	RenderData[4] = FVector4(Instance.LightmapUVBias.X, Instance.LightmapUVBias.Y, RandomInstanceID, 0);
 }
 
 void FStaticMeshInstanceBuffer::SetupCPUAccess(UInstancedStaticMeshComponent* InComponent)
@@ -140,11 +134,8 @@ void FStaticMeshInstanceBuffer::Init(UInstancedStaticMeshComponent* InComponent,
 			RenderData[1] = FVector4(0, 0, 0, 0);
 			RenderData[2] = FVector4(0, 0, 0, 0);
 			RenderData[3] = FVector4(0, 0, 0, 0);
-			
-			// inverse is identity
-			RenderData[4] = FVector4(1, 0, 0, 0);
-			RenderData[5] = FVector4(0, 1, 0, 0);
-			RenderData[6] = FVector4(0, 0, 1, 0);
+		
+			RenderData[4] = FVector4(0, 0, 0, 0);
 		}
 	}
 }
@@ -332,9 +323,7 @@ void FInstancedStaticMeshVertexFactory::InitRHI()
 		Elements.Add(AccessStreamComponent(Data.InstancedTransformComponent[0],9));
 		Elements.Add(AccessStreamComponent(Data.InstancedTransformComponent[1],10));
 		Elements.Add(AccessStreamComponent(Data.InstancedTransformComponent[2],11));
-		Elements.Add(AccessStreamComponent(Data.InstancedInverseTransformComponent[0],12));
-		Elements.Add(AccessStreamComponent(Data.InstancedInverseTransformComponent[1],13));
-		Elements.Add(AccessStreamComponent(Data.InstancedInverseTransformComponent[2],14));
+		Elements.Add(AccessStreamComponent(Data.InstancedLightmapUVBiasComponent, 12));
 	}
 
 	// we don't need per-vertex shadow or lightmap rendering
@@ -490,17 +479,14 @@ void FInstancedStaticMeshRenderData::InitStaticMeshVertexFactories(
 				CurInstanceBufferOffset += sizeof(float) * 4;
 			}
 
-			for (int32 MatrixRow = 0; MatrixRow < 3; MatrixRow++)
-			{
-				Data.InstancedInverseTransformComponent[MatrixRow] = FVertexStreamComponent(
-					&InstancedRenderData->InstanceBuffer,
-					CurInstanceBufferOffset, 
-					InstancedRenderData->InstanceBuffer.GetStride(),
-					VET_Float4,
-					true
-					);
-				CurInstanceBufferOffset += sizeof(float) * 4;
-			}
+			Data.InstancedLightmapUVBiasComponent = FVertexStreamComponent(
+				&InstancedRenderData->InstanceBuffer,
+				CurInstanceBufferOffset, 
+				InstancedRenderData->InstanceBuffer.GetStride(),
+				VET_Float4,
+				true
+				);
+			CurInstanceBufferOffset += sizeof(float) * 4;
 		}
 
 		// Assign to the vertex factory for this LOD.
@@ -1659,7 +1645,7 @@ void FInstancedStaticMeshVertexFactoryShaderParameters::SetMesh( FRHICommandList
 			auto* InstanceStream = ((const FInstancingUserData::FInstanceStream*)InstancingData->RenderData->InstanceBuffer.GetRawData()) + BatchElement.UserIndex;
 			SetShaderValue(RHICmdList, VS, CPUInstanceShadowMapBias, InstanceStream->InstanceShadowmapUVBias);
 			SetShaderValueArray(RHICmdList, VS, CPUInstanceTransform, InstanceStream->InstanceTransform, 3);
-			SetShaderValueArray(RHICmdList, VS, CPUInstanceInverseTransform, InstanceStream->InstanceInverseTransform, 3);
+			SetShaderValue(RHICmdList, VS, CPUInstanceLightmapUVBias, InstanceStream->InstanceLightmapUVBias);
 		}
 	}
 }
diff --git a/Engine/Source/Runtime/Engine/Private/InstancedStaticMesh.h b/Engine/Source/Runtime/Engine/Private/InstancedStaticMesh.h
--- a/Engine/Source/Runtime/Engine/Private/InstancedStaticMesh.h
+++ b/Engine/Source/Runtime/Engine/Private/InstancedStaticMesh.h
@@ -146,7 +146,7 @@ struct FInstancingUserData
 	{
 		FVector4 InstanceShadowmapUVBias;
 		FVector4 InstanceTransform[3];
-		FVector4 InstanceInverseTransform[3];
+		FVector4 InstanceLightmapUVBias;
 	};
 
 	class FInstancedStaticMeshRenderData* RenderData;
@@ -174,8 +174,8 @@ public:
 		/** The stream to read the mesh transform from. */
 		FVertexStreamComponent InstancedTransformComponent[3];
 
-		/** The stream to read the inverse transform, as well as the Lightmap Bias in 0/1 */
-		FVertexStreamComponent InstancedInverseTransformComponent[3];
+		/** The stream to read the Lightmap Bias and Random instance ID from. */
+		FVertexStreamComponent InstancedLightmapUVBiasComponent;
 	};
 
 	/**
@@ -254,7 +254,7 @@ class FInstancedStaticMeshVertexFactoryShaderParameters : public FLocalVertexFac
 		InstancingWorldViewOriginOneParameter.Bind(ParameterMap, TEXT("InstancingWorldViewOriginOne"));
 		CPUInstanceShadowMapBias.Bind(ParameterMap, TEXT("CPUInstanceShadowMapBias"));
 		CPUInstanceTransform.Bind(ParameterMap, TEXT("CPUInstanceTransform"));
-		CPUInstanceInverseTransform.Bind(ParameterMap, TEXT("CPUInstanceInverseTransform"));
+		CPUInstanceLightmapUVBias.Bind(ParameterMap, TEXT("CPUInstanceLightmapUVBias"));
 	}
 
 	virtual void SetMesh(FRHICommandList& RHICmdList, FShader* VertexShader,const class FVertexFactory* VertexFactory,const class FSceneView& View,const struct FMeshBatchElement& BatchElement,uint32 DataFlags) const override;
@@ -270,7 +270,7 @@ class FInstancedStaticMeshVertexFactoryShaderParameters : public FLocalVertexFac
 		Ar << InstancingWorldViewOriginOneParameter;
 		Ar << CPUInstanceShadowMapBias;
 		Ar << CPUInstanceTransform;
-		Ar << CPUInstanceInverseTransform;
+		Ar << CPUInstanceLightmapUVBias;
 	}
 
 	virtual uint32 GetSize() const { return sizeof(*this); }
@@ -285,7 +285,7 @@ private:
 
 	FShaderParameter CPUInstanceShadowMapBias;
 	FShaderParameter CPUInstanceTransform;
-	FShaderParameter CPUInstanceInverseTransform;
+	FShaderParameter CPUInstanceLightmapUVBias;
 };
 
 /*-----------------------------------------------------------------------------
diff --git a/Engine/Source/Runtime/Engine/Public/StaticMeshResources.h b/Engine/Source/Runtime/Engine/Public/StaticMeshResources.h
--- a/Engine/Source/Runtime/Engine/Public/StaticMeshResources.h
+++ b/Engine/Source/Runtime/Engine/Public/StaticMeshResources.h
@@ -905,7 +905,7 @@ public:
 
 	enum 
 	{
-		VectorsPerInstance = 7
+		VectorsPerInstance = 5
 	};
 	typedef TResourceArray<FVector4,VERTEXBUFFER_ALIGNMENT> ArrayType;
 
