shader compile fix.

diff --git a/Engine/Shaders/LocalVertexFactory.usf b/Engine/Shaders/LocalVertexFactory.usf
--- a/Engine/Shaders/LocalVertexFactory.usf
+++ b/Engine/Shaders/LocalVertexFactory.usf
@@ -530,7 +530,7 @@ FVertexFactoryInterpolantsVSToPS VertexFactoryGetInterpolantsVSToPS(FVertexFacto
 		#if !USE_INSTANCING_EMULATED
 		float2 InstanceLightmapUVBias = Input.InstanceLightmapUVBias.xy;
 		#else
-		float2 InstanceLightmapUVBias = CPUInstanceLightmapUVBias.xy
+		float2 InstanceLightmapUVBias = CPUInstanceLightmapUVBias.xy;
 		#endif // !USE_INSTANCING_EMULATED
 		LightMapCoordinate = Input.LightMapCoordinate * LightMapCoordinateScaleBias.xy + InstanceLightmapUVBias;
 	#else
