UE4 - Fix shader warning

diff --git a/Engine/Shaders/MaterialTemplate.usf b/Engine/Shaders/MaterialTemplate.usf
--- a/Engine/Shaders/MaterialTemplate.usf
+++ b/Engine/Shaders/MaterialTemplate.usf
@@ -1373,7 +1373,7 @@ void CalcMaterialParameters(
 	Parameters.SVPosition = SvPosition;
 
 	{
-		float2 PixelPos = SvPosition - View.ViewRectMin.xy;	
+		float2 PixelPos = SvPosition.xy - View.ViewRectMin.xy;	
 		Parameters.ScreenPosition.xy = ( PixelPos / View.ViewSizeAndSceneTexelSize.xy - 0.5 ) * float2(2,-2);
 		Parameters.ScreenPosition.zw = SvPosition.zw;
 		Parameters.ScreenPosition.xyz *= Parameters.ScreenPosition.w;
