Switched diffuse shading model from Lambert to energy conserving Burley.

Lighting done in the bass pass does not properly use the new shading model yet. All dynamic lighting does: deferred direct lights, ambient cubemap, and dynamic skylight.

diff --git a/Engine/Shaders/BRDF.usf b/Engine/Shaders/BRDF.usf
--- a/Engine/Shaders/BRDF.usf
+++ b/Engine/Shaders/BRDF.usf
@@ -47,12 +47,13 @@ float3 Diffuse_Lambert( float3 DiffuseColor )
 }
 
 // [Burley 2012, "Physically-Based Shading at Disney"]
+// [Lagrade et al. 2014, "Moving Frostbite to Physically Based Rendering"]
 float3 Diffuse_Burley( float3 DiffuseColor, float Roughness, float NoV, float NoL, float VoH )
 {
-	float FD90 = 0.5 + 2 * VoH * VoH * Roughness;
+	float FD90 = ( 0.5 + 2 * VoH * VoH ) * Roughness;
 	float FdV = 1 + (FD90 - 1) * pow( 1 - NoV, 5 );
 	float FdL = 1 + (FD90 - 1) * pow( 1 - NoL, 5 );
-	return DiffuseColor * ( 1 / PI * FdV * FdL );
+	return DiffuseColor * ( 1 / PI * FdV * FdL ) * ( 1 - 0.3333 * Roughness );	// TODO premultiply DiffuseColor
 }
 
 // [Gotanda 2012, "Beyond a Simple Physically Based Blinn-Phong Model in Real-Time"]
diff --git a/Engine/Shaders/BasePassPixelShader.usf b/Engine/Shaders/BasePassPixelShader.usf
--- a/Engine/Shaders/BasePassPixelShader.usf
+++ b/Engine/Shaders/BasePassPixelShader.usf
@@ -680,6 +680,10 @@ void Main(
 	}
 #endif
 
+	// TODO Replace with proper EnvBRDF for Burley diffuse model.
+	// Any indirect diffuse done in the base pass does not properly account for the Burley diffuse model besides this crude approximation.
+	DiffuseColor *= 1 - 0.3333 * Roughness;
+
 
 	half3 Color = 0;
 	float IndirectIrradiance = 0;
diff --git a/Engine/Shaders/DeferredLightingCommon.usf b/Engine/Shaders/DeferredLightingCommon.usf
--- a/Engine/Shaders/DeferredLightingCommon.usf
+++ b/Engine/Shaders/DeferredLightingCommon.usf
@@ -100,7 +100,10 @@ float3 StandardShading( FGBufferData GBuffer, float3 LobeRoughness, float3 LobeE
 	float Vis = Vis_SmithJointApprox( LobeRoughness[1], NoV, NoL );
 	float3 F = F_Schlick( GBuffer.SpecularColor, VoH );
 
-	return Diffuse_Lambert(GBuffer.DiffuseColor) * (LobeEnergy[2] * DiffSpecMask.r) + (D * Vis * DiffSpecMask.g) * F;
+	//float3 Diffuse = Diffuse_Lambert( GBuffer.DiffuseColor );
+	float3 Diffuse = Diffuse_Burley( GBuffer.DiffuseColor, LobeRoughness[1], NoV, NoL, VoH );
+	
+	return Diffuse * (LobeEnergy[2] * DiffSpecMask.r) + (D * Vis * DiffSpecMask.g) * F;
 }
 
 float3 SimpleShading( FGBufferData GBuffer, float Roughness, float3 L, float3 V, half3 N )
@@ -150,7 +153,7 @@ float3 ClearCoatShading( FGBufferData GBuffer, float3 LobeRoughness, float3 Lobe
 	//float3 F2 = F_Schlick( GBuffer.SpecularColor, VoH );
 	float3 F2 = saturate( 50.0 * GBuffer.SpecularColor.g ) * Fc + (1 - Fc) * GBuffer.SpecularColor;
 
-	float3 Fr2 = Diffuse_Lambert( GBuffer.DiffuseColor ) * LobeEnergy[2] + (D2 * Vis2) * F2;
+	float3 Fr2 = Diffuse_Burley( GBuffer.DiffuseColor, LobeRoughness[1], NoV, NoL, VoH ) * LobeEnergy[2] + (D2 * Vis2) * F2;
 	
 	return Fr1 + Fr2 * LayerAttenuation;
 #else
diff --git a/Engine/Shaders/PostProcessAmbient.usf b/Engine/Shaders/PostProcessAmbient.usf
--- a/Engine/Shaders/PostProcessAmbient.usf
+++ b/Engine/Shaders/PostProcessAmbient.usf
@@ -14,24 +14,42 @@
 
 #define IMPORTANCE_SAMPLE		0
 
-float3 DiffuseIBL( uint2 Random, float3 DiffuseColor, float3 N )
+float3 DiffuseIBL( uint2 Random, float3 DiffuseColor, float Roughness, float3 N, float3 V )
 {
+	N = normalize( N );
+	V = normalize( V );
+
 	float3 DiffuseLighting = 0;
+	
+	float NoV = saturate( dot( N, V ) );
 
 	const uint NumSamples = 32;
 	for( uint i = 0; i < NumSamples; i++ )
 	{
 		float2 E = Hammersley( i, NumSamples, Random );
 		float3 L = TangentToWorld( CosineSampleHemisphere( E ).xyz, N );
+		float3 H = normalize(V + L);
 
 		float NoL = saturate( dot( N, L ) );
+		float NoH = saturate( dot( N, H ) );
+		float VoH = saturate( dot( V, H ) );
+
 		if( NoL > 0 )
 		{
 			float3 SampleColor = AmbientCubemap.SampleLevel( AmbientCubemapSampler, L, 0 ).rgb;
 
+			float FD90 = ( 0.5 + 2 * VoH * VoH ) * Roughness;
+			//float FD90 = 0.5 + 2 * VoH * VoH * Roughness;
+			float FdV = 1 + (FD90 - 1) * pow( 1 - NoV, 5 );
+			float FdL = 1 + (FD90 - 1) * pow( 1 - NoL, 5 );
+
+#if 1
 			// lambert = DiffuseColor * NoL / PI
 			// pdf = NoL / PI
+			DiffuseLighting += SampleColor * DiffuseColor * FdV * FdL * ( 1 - 0.3333 * Roughness );
+#else
 			DiffuseLighting += SampleColor * DiffuseColor;
+#endif
 		}
 	}
 
@@ -58,7 +76,7 @@ float3 SpecularIBL( uint2 Random, float3 SpecularColor, float Roughness, float3
 		{
 			float3 SampleColor = AmbientCubemap.SampleLevel( AmbientCubemapSampler, L, 0 ).rgb;
 
-			float Vis = Vis_Smith( Roughness, NoV, NoL );
+			float Vis = Vis_SmithJointApprox( Roughness, NoV, NoL );
 			float Fc = pow( 1 - VoH, 5 );
 			float3 F = (1 - Fc) * SpecularColor + Fc;
 
@@ -85,7 +103,9 @@ float3 StandardShading( FGBufferData GBuffer, float Roughness, float3 L, float3
 	float Vis = Vis_SmithJointApprox( Roughness, NoV, NoL );
 	float3 F = F_Schlick( GBuffer.SpecularColor, VoH );
 
-	return Diffuse_Lambert( GBuffer.DiffuseColor ) + (D * Vis) * F;
+	float3 Diffuse = Diffuse_Burley( GBuffer.DiffuseColor, Roughness, NoV, NoL, VoH );
+
+	return Diffuse + (D * Vis) * F;
 }
 
 float3 ClearCoatShading( FGBufferData GBuffer, float Roughness, float3 L, float3 V, half3 N )
@@ -397,11 +417,10 @@ float3 IntegrateBRDF( uint2 Random, float Roughness, float NoV )
 			float NoH = saturate( H.z );
 			float VoH = saturate( dot( V, H ) );
 
-			//float FD90 = 0.5 * Roughness + 2 * VoH * VoH * Roughness;
-			float FD90 = 0.5 + 2 * VoH * VoH * Roughness;
+			float FD90 = ( 0.5 + 2 * VoH * VoH ) * Roughness;
 			float FdV = 1 + (FD90 - 1) * pow( 1 - NoV, 5 );
 			float FdL = 1 + (FD90 - 1) * pow( 1 - NoL, 5 );
-			C += FdV * FdL * lerp( 1, 1.0/1.5, Roughness );
+			C += FdV * FdL * ( 1 - 0.3333 * Roughness );
 		}
 	}
 
@@ -453,21 +472,24 @@ void MainPS(in float4 UVAndScreenPos : TEXCOORD0, out float4 OutColor : SV_Targe
 	float3 NonSpecularContribution = 0;
 	float3 SpecularContribution = 0;
 
-#if 0
+#if 1
 	float3 RetroReflectionDir = lerp( N, V, saturate( ( NoV * (1.02341 * a - 1.51174) + -0.511705 * a + 0.755868 ) * a ) );
 	float3 DiffuseLookup =  TextureCubeSampleLevel(AmbientCubemap, AmbientCubemapSampler, RetroReflectionDir, AbsoluteDiffuseMip).rgb;
 
 	float3 ABC = PreIntegratedGF.SampleLevel( PreIntegratedGFSampler, float2( NoV, GBuffer.Roughness ), 0 ).rgb;
 	//float3 ABC = IntegrateBRDF( Random, GBuffer.Roughness, NoV );
-	DiffuseLookup *= ABC.z;
+	
+	NonSpecularContribution += GBuffer.DiffuseColor * DiffuseLookup * ABC.z;
+	//NonSpecularContribution += DiffuseIBL( Random, GBuffer.DiffuseColor, GBuffer.Roughness, N, V );// * ABC.z;
 #else
 	float3 DiffuseLookup =  TextureCubeSampleLevel(AmbientCubemap, AmbientCubemapSampler, N, AbsoluteDiffuseMip).rgb;
+	NonSpecularContribution += GBuffer.DiffuseColor * DiffuseLookup;
 #endif
 
 	// Diffuse
 	{
 		// we want to access the mip with the preconvolved diffuse lighting (coneangle=90 degree)
-		NonSpecularContribution += GBuffer.DiffuseColor * DiffuseLookup;
+		//NonSpecularContribution += GBuffer.DiffuseColor * DiffuseLookup;
 	}
 
 	// Specular
diff --git a/Engine/Shaders/SkyLighting.usf b/Engine/Shaders/SkyLighting.usf
--- a/Engine/Shaders/SkyLighting.usf
+++ b/Engine/Shaders/SkyLighting.usf
@@ -7,6 +7,7 @@
 #include "Common.usf"
 #include "DeferredShadingCommon.usf"
 #include "ReflectionEnvironmentShared.usf"
+#include "BRDF.usf"
 
 Texture2D BentNormalAOTexture;
 SamplerState BentNormalAOSampler;
@@ -30,13 +31,28 @@ void SkyLightDiffusePS(
 	if (ShadingModelId > 0)
 	{
 		FScreenSpaceData ScreenSpaceData = GetScreenSpaceData(UV);
-		FGBufferData InGBufferData = ScreenSpaceData.GBuffer;
+		FGBufferData GBuffer = ScreenSpaceData.GBuffer;
 
 		float SkyVisibility = 1;
 		float DotProductFactor = 1;
-		float3 SkyLightingNormal = InGBufferData.WorldNormal;
+		float3 SkyLightingNormal = GBuffer.WorldNormal;
 		float3 DiffuseIrradiance = 0;
 
+		// screen position in [-1, 1] screen space
+		float2 ScreenSpacePos = UVAndScreenPos.zw;
+		float3 ScreenVector = normalize(mul(float4(ScreenSpacePos, 1, 0), View.ScreenToWorld).xyz);
+
+		float3 N = GBuffer.WorldNormal;
+		float3 V = -ScreenVector;
+		float NoV = abs( dot(N, V) ) + 1e-5;
+
+		float a = Square( GBuffer.Roughness );
+		float3 RetroReflectionDir = lerp( N, V, saturate( ( NoV * (1.02341 * a - 1.51174) + -0.511705 * a + 0.755868 ) * a ) );
+		SkyLightingNormal = RetroReflectionDir;
+
+		float3 ABC = PreIntegratedGF.SampleLevel( PreIntegratedGFSampler, float2( NoV, GBuffer.Roughness ), 0 ).rgb;
+		GBuffer.DiffuseColor *= ABC.z;
+
 #if APPLY_SHADOWING
 		float3 BentNormal = Texture2DSampleLevel(BentNormalAOTexture, BentNormalAOSampler, UV, 0).xyz;
 
@@ -50,8 +66,8 @@ void SkyLightDiffusePS(
 			// Use more bent normal in corners
 			float BentNormalWeightFactor = SkyVisibility;
 
-			SkyLightingNormal = lerp(NormalizedBentNormal, InGBufferData.WorldNormal, BentNormalWeightFactor);
-			DotProductFactor = lerp(dot(NormalizedBentNormal, InGBufferData.WorldNormal), 1, BentNormalWeightFactor);
+			SkyLightingNormal = lerp(NormalizedBentNormal, GBuffer.WorldNormal, BentNormalWeightFactor);
+			DotProductFactor = lerp(dot(NormalizedBentNormal, GBuffer.WorldNormal), 1, BentNormalWeightFactor);
 		}
 		#else
 		{
@@ -74,12 +90,12 @@ void SkyLightDiffusePS(
 		#endif
 #endif
 
-		float ScalarFactors = lerp(SkyVisibility, 1, OcclusionTintAndMinOcclusion.w) * InGBufferData.GBufferAO * ScreenSpaceData.AmbientOcclusion;
+		float ScalarFactors = lerp(SkyVisibility, 1, OcclusionTintAndMinOcclusion.w) * GBuffer.GBufferAO * ScreenSpaceData.AmbientOcclusion;
 
-		if (InGBufferData.ShadingModelID == SHADINGMODELID_TWOSIDED_FOLIAGE)
+		if (ShadingModelId == SHADINGMODELID_TWOSIDED_FOLIAGE)
 		{
-			float3 SubsurfaceLookup = GetSkySHDiffuse(-InGBufferData.WorldNormal) * View.SkyLightColor.rgb;
-			float3 SubsurfaceColor = ExtractSubsurfaceColor(InGBufferData);
+			float3 SubsurfaceLookup = GetSkySHDiffuse(-GBuffer.WorldNormal) * View.SkyLightColor.rgb;
+			float3 SubsurfaceColor = ExtractSubsurfaceColor(GBuffer);
 			Lighting += ScalarFactors * SubsurfaceLookup * SubsurfaceColor;
 		}
 
@@ -87,9 +103,9 @@ void SkyLightDiffusePS(
 		float3 DiffuseLookup = GetSkySHDiffuse(SkyLightingNormal) * View.SkyLightColor.rgb;
 		
 		// Apply AO to the sky diffuse and account for darkening due to the geometry term
-		Lighting += (ScalarFactors * DotProductFactor) * DiffuseLookup * InGBufferData.DiffuseColor + (1 - SkyVisibility) * OcclusionTintAndMinOcclusion.xyz;
+		Lighting += (ScalarFactors * DotProductFactor) * DiffuseLookup * GBuffer.DiffuseColor + (1 - SkyVisibility) * OcclusionTintAndMinOcclusion.xyz;
 
-		Lighting += DiffuseIrradiance * InGBufferData.DiffuseColor;
+		Lighting += DiffuseIrradiance * GBuffer.DiffuseColor;
 
 		// For debugging
 		#if SUPPORT_IRRADIANCE
diff --git a/Engine/Source/Runtime/Renderer/Private/DistanceFieldSurfaceCacheLighting.cpp b/Engine/Source/Runtime/Renderer/Private/DistanceFieldSurfaceCacheLighting.cpp
--- a/Engine/Source/Runtime/Renderer/Private/DistanceFieldSurfaceCacheLighting.cpp
+++ b/Engine/Source/Runtime/Renderer/Private/DistanceFieldSurfaceCacheLighting.cpp
@@ -4826,6 +4826,8 @@ public:
 		: FGlobalShader(Initializer)
 	{
 		DeferredParameters.Bind(Initializer.ParameterMap);
+		PreIntegratedGF.Bind(Initializer.ParameterMap, TEXT("PreIntegratedGF"));
+		PreIntegratedGFSampler.Bind(Initializer.ParameterMap, TEXT("PreIntegratedGFSampler"));
 		DynamicBentNormalAOTexture.Bind(Initializer.ParameterMap, TEXT("BentNormalAOTexture"));
 		DynamicBentNormalAOSampler.Bind(Initializer.ParameterMap, TEXT("BentNormalAOSampler"));
 		DynamicIrradianceTexture.Bind(Initializer.ParameterMap, TEXT("IrradianceTexture"));
@@ -4840,6 +4842,8 @@ public:
 		FGlobalShader::SetParameters(RHICmdList, ShaderRHI, View);
 		DeferredParameters.Set(RHICmdList, ShaderRHI, View);
 
+		SetTextureParameter(RHICmdList, ShaderRHI, PreIntegratedGF, PreIntegratedGFSampler, TStaticSamplerState<SF_Bilinear,AM_Clamp,AM_Clamp,AM_Clamp>::GetRHI(), GSystemTextures.PreintegratedGF->GetRenderTargetItem().ShaderResourceTexture );
+
 		SetTextureParameter(RHICmdList, ShaderRHI, DynamicBentNormalAOTexture, DynamicBentNormalAOSampler, TStaticSamplerState<SF_Point>::GetRHI(), DynamicBentNormalAO);
 
 		if (DynamicIrradianceTexture.IsBound())
@@ -4865,6 +4869,7 @@ public:
 	{
 		bool bShaderHasOutdatedParameters = FGlobalShader::Serialize(Ar);
 		Ar << DeferredParameters;
+		Ar << PreIntegratedGF << PreIntegratedGFSampler;
 		Ar << DynamicBentNormalAOTexture;
 		Ar << DynamicBentNormalAOSampler;
 		Ar << DynamicIrradianceTexture;
@@ -4877,6 +4882,8 @@ public:
 private:
 
 	FDeferredPixelShaderParameters DeferredParameters;
+	FShaderResourceParameter PreIntegratedGF;
+	FShaderResourceParameter PreIntegratedGFSampler;
 	FShaderResourceParameter DynamicBentNormalAOTexture;
 	FShaderResourceParameter DynamicBentNormalAOSampler;
 	FShaderResourceParameter DynamicIrradianceTexture;
diff --git a/Engine/Source/Runtime/Renderer/Private/SystemTextures.cpp b/Engine/Source/Runtime/Renderer/Private/SystemTextures.cpp
--- a/Engine/Source/Runtime/Renderer/Private/SystemTextures.cpp
+++ b/Engine/Source/Runtime/Renderer/Private/SystemTextures.cpp
@@ -255,11 +255,12 @@ void FSystemTextures::InitializeTextures(FRHICommandListImmediate& RHICmdList, E
 
 			EPixelFormat Format = PF_R8G8;
 			// for low roughness we would get banding with PF_R8G8 but for low spec it could be used, for now we don't do this optimization
-			if (GPixelFormats[PF_G16R16].Supported)
+			if (GPixelFormats[PF_A16B16G16R16].Supported)
 			{
-				Format = PF_G16R16;
+				// 3rd channel required for diffuse
+				// TODO try 11:11:10
+				Format = PF_A16B16G16R16;
 			}
-			//Format = PF_A16B16G16R16;
 
 			FPooledRenderTargetDesc Desc(FPooledRenderTargetDesc::Create2DDesc(FIntPoint(128, 32), Format, TexCreate_FastVRAM, TexCreate_None, false));
 
@@ -342,10 +343,10 @@ void FSystemTextures::InitializeTextures(FRHICommandListImmediate& RHICmdList, E
 							float NoH = FMath::Max(H.Z, 0.0f);
 							float VoH = FMath::Max(V | H, 0.0f);
 
-							float FD90 = 0.5f + 2.0f * VoH * VoH * Roughness;
+							float FD90 = ( 0.5f + 2.0f * VoH * VoH ) * Roughness;
 							float FdV = 1.0f + (FD90 - 1.0f) * pow( 1.0f - NoV, 5 );
 							float FdL = 1.0f + (FD90 - 1.0f) * pow( 1.0f - NoL, 5 );
-							C += FdV * FdL;// * ( 1.0f - 0.3333f * Roughness );
+							C += FdV * FdL * ( 1.0f - 0.3333f * Roughness );
 						}
 					}
 					A /= NumSamples;
