Normalized dominant diffuse direction

diff --git a/Engine/Shaders/SkyLighting.usf b/Engine/Shaders/SkyLighting.usf
--- a/Engine/Shaders/SkyLighting.usf
+++ b/Engine/Shaders/SkyLighting.usf
@@ -47,8 +47,8 @@ void SkyLightDiffusePS(
 		float NoV = abs( dot(N, V) ) + 1e-5;
 
 		float a = Square( GBuffer.Roughness );
-		float3 RetroReflectionDir = lerp( N, V, saturate( ( NoV * (1.02341 * a - 1.51174) + -0.511705 * a + 0.755868 ) * a ) );
-		SkyLightingNormal = RetroReflectionDir;
+		float3 DominantDiffuseDir = lerp( N, V, saturate( ( NoV * (1.02341 * a - 1.51174) + -0.511705 * a + 0.755868 ) * a ) );
+		SkyLightingNormal = normalize( DominantDiffuseDir );
 
 		float3 ABC = PreIntegratedGF.SampleLevel( PreIntegratedGFSampler, float2( NoV, GBuffer.Roughness ), 0 ).rgb;
 		GBuffer.DiffuseColor *= ABC.z;
