Fixed translucent landscape materials not compiling in ES2 preview
This was due to the interpolant indexes used by BasePassCommon.usf's VertexFog and LandscapeVertexFactory.usf's ShadowMapCoordinate overlapping

diff --git a/Engine/Shaders/BasePassCommon.usf b/Engine/Shaders/BasePassCommon.usf
--- a/Engine/Shaders/BasePassCommon.usf
+++ b/Engine/Shaders/BasePassCommon.usf
@@ -16,18 +16,18 @@
 struct FSharedBasePassInterpolants
 {
 	//for texture-lightmapped translucency we can pass the vertex fog in its own interpolator
-	#if NEEDS_BASEPASS_FOGGING
-		float4 VertexFog		: TEXCOORD6;
-	#endif
-
-	#if !PC_D3D && !GL4_PROFILE
-		float4 PixelPosition	: TEXCOORD7;
-
-#if USE_WORLD_POSITION_EXCLUDING_SHADER_OFFSETS
-		float4 PixelPositionExcludingWPO : TEXCOORD8;
+#if NEEDS_BASEPASS_FOGGING
+	float4 VertexFog		: TEXCOORD7;
 #endif
 
+#if !PC_D3D && !GL4_PROFILE
+	float4 PixelPosition	: TEXCOORD8;
+
+	#if USE_WORLD_POSITION_EXCLUDING_SHADER_OFFSETS
+	float4 PixelPositionExcludingWPO : TEXCOORD9;
 	#endif
+
+#endif
 };
 
 #if PC_D3D || GL4_PROFILE
@@ -35,10 +35,10 @@ struct FSharedBasePassInterpolants
 /** Interpolants passed from the vertex shader to the pixel shader. */
 struct FBasePassInterpolantsVSToPS : FSharedBasePassInterpolants
 {
-	float4 PixelPosition	: TEXCOORD7;
+	float4 PixelPosition	: TEXCOORD8;
 
 #if USE_WORLD_POSITION_EXCLUDING_SHADER_OFFSETS
-	float4 PixelPositionExcludingWPO : TEXCOORD8;
+	float4 PixelPositionExcludingWPO : TEXCOORD9;
 #endif
 };
 
@@ -46,7 +46,7 @@ struct FBasePassInterpolantsVSToPS : FSharedBasePassInterpolants
 struct FBasePassInterpolantsVSToDS : FSharedBasePassInterpolants
 {
 #if USE_WORLD_POSITION_EXCLUDING_SHADER_OFFSETS
-	float3 WorldPositionExcludingWPO : TEXCOORD8;
+	float3 WorldPositionExcludingWPO : TEXCOORD9;
 #endif
 };
 
