UE4 - Add r.BasePassOutputsVelocity to DDC and fix the shader define so it only works with opaque & masked materials

diff --git a/Engine/Shaders/BasePassCommon.usf b/Engine/Shaders/BasePassCommon.usf
--- a/Engine/Shaders/BasePassCommon.usf
+++ b/Engine/Shaders/BasePassCommon.usf
@@ -13,6 +13,8 @@
 #define NEEDS_BASEPASS_FOGGING			(ENABLE_TRANSLUCENCY_VERTEX_FOG && ((MATERIALBLENDING_TRANSLUCENT || MATERIALBLENDING_ADDITIVE || MATERIALBLENDING_MODULATE) && !MATERIAL_USES_SCENE_COLOR_COPY))
 #define NEEDS_LIGHTMAP					(NEEDS_LIGHTMAP_COORDINATE)
 
+#define WRITES_VELOCITY_TO_GBUFFER		(OUTPUT_GBUFFER_VELOCITY && (MATERIALBLENDING_SOLID || MATERIAL_BLENDING_MASKED))
+
 struct FSharedBasePassInterpolants
 {
 	//for texture-lightmapped translucency we can pass the vertex fog in its own interpolator
@@ -28,7 +30,7 @@ struct FSharedBasePassInterpolants
 	#endif
 #endif
 
-#if OUTPUT_GBUFFER_VELOCITY
+#if WRITES_VELOCITY_TO_GBUFFER
 	// .xy is clip position, pre divide by w; .w is clip W; .z is 0 or 1 to mask out the velocity output
 	float4 PrevScreenPosition : TEXCOORD16;
 	float4 ScreenPosition : TEXCOORD17;
diff --git a/Engine/Shaders/BasePassPixelShader.usf b/Engine/Shaders/BasePassPixelShader.usf
--- a/Engine/Shaders/BasePassPixelShader.usf
+++ b/Engine/Shaders/BasePassPixelShader.usf
@@ -587,7 +587,7 @@ void Main(
 	#if ALLOW_STATIC_LIGHTING
 		,out float4 OutGBufferE : SV_Target5
 	#endif
-	#if OUTPUT_GBUFFER_VELOCITY
+	#if WRITES_VELOCITY_TO_GBUFFER
 		, out float4 OutGBufferVelocity : SV_Target6
 	#endif
 #endif
@@ -850,13 +850,10 @@ void Main(
 			float4 OutGBufferE = 0;
 		#endif
 
-		#if OUTPUT_GBUFFER_VELOCITY
+		#if WRITES_VELOCITY_TO_GBUFFER
 			{
-				float2 ScreenPosition = BasePassInterpolants.ScreenPosition.xy / BasePassInterpolants.ScreenPosition.w - View.PrevProjection[2].xy;
-				float2 PrevScreenPosition = BasePassInterpolants.PrevScreenPosition.xy / BasePassInterpolants.PrevScreenPosition.w - View.PrevProjection[2].xy;
-
 				// 2d velocity, includes camera an object motion
-				float2 Velocity = ScreenPosition - PrevScreenPosition;
+				float2 Velocity = Calculate2DVelocity(BasePassInterpolants.ScreenPosition, BasePassInterpolants.PrevScreenPosition);
 
 				// Make sure not to touch 0,0 which is clear color
 				GBuffer.Velocity = float4(EncodeVelocityToTexture(Velocity), 0, 0) * BasePassInterpolants.PrevScreenPosition.z;
diff --git a/Engine/Shaders/BasePassVertexShader.usf b/Engine/Shaders/BasePassVertexShader.usf
--- a/Engine/Shaders/BasePassVertexShader.usf
+++ b/Engine/Shaders/BasePassVertexShader.usf
@@ -6,7 +6,7 @@
 
 #include "BasePassVertexCommon.usf"
 
-#if OUTPUT_GBUFFER_VELOCITY
+#if WRITES_VELOCITY_TO_GBUFFER
 //@todo-rco: Move to pixel shader
 float SkipOutputVelocity;
 #endif
@@ -65,7 +65,7 @@ void Main(
 	#endif
 #endif
 
-#if OUTPUT_GBUFFER_VELOCITY
+#if WRITES_VELOCITY_TO_GBUFFER
 	{
 		float4 PrevTranslatedWorldPosition = VertexFactoryGetPreviousWorldPosition( Input, VFIntermediates );	
 		VertexParameters = GetMaterialVertexParameters(Input, VFIntermediates, PrevTranslatedWorldPosition.xyz, TangentToLocal);
diff --git a/Engine/Shaders/DeferredShadingCommon.usf b/Engine/Shaders/DeferredShadingCommon.usf
--- a/Engine/Shaders/DeferredShadingCommon.usf
+++ b/Engine/Shaders/DeferredShadingCommon.usf
@@ -278,7 +278,7 @@ struct FGBufferData
 	// in unreal units (linear), can be used to reconstruct world position,
 	// only valid when decoding the GBuffer as the value gets reconstructed from the Z buffer
 	float Depth;
-	// Velocity for motion blur (only used when OUTPUT_GBUFFER_VELOCITY is enabled)
+	// Velocity for motion blur (only used when WRITES_VELOCITY_TO_GBUFFER is enabled)
 	float4 Velocity;
 };
 
@@ -351,7 +351,7 @@ void EncodeGBuffer(
 
 	OutGBufferE = GBuffer.PrecomputedShadowFactors;
 
-#if OUTPUT_GBUFFER_VELOCITY
+#if WRITES_VELOCITY_TO_GBUFFER
 	OutGBufferVelocity = GBuffer.Velocity;
 #else
 	OutGBufferVelocity = 0;
@@ -538,7 +538,7 @@ SamplerState CustomDepthTextureSampler;
 			float4 GBufferE = 1;
 		#endif
 
-		#if OUTPUT_GBUFFER_VELOCITY
+		#if WRITES_VELOCITY_TO_GBUFFER
 			float4 GBufferVelocity = GBuffers.GBufferVelocityTextureNonMS.Load(int3(PixelPos, 0));
 		#else
 			float4 GBufferVelocity = 0;
@@ -613,7 +613,7 @@ FGBufferData GetGBufferData(float2 UV, bool bGetNormalizedNormal = true)
 		float4 GBufferE = 1;
 	#endif
 
-	#if OUTPUT_GBUFFER_VELOCITY
+	#if WRITES_VELOCITY_TO_GBUFFER
 		float4 GBufferVelocity = Texture2DSampleLevel(GBuffers.GBufferVelocityTexture, GBuffers.GBufferVelocityTextureSampler, UV, 0);
 	#else
 		float4 GBufferVelocity = 0;
@@ -681,7 +681,7 @@ FScreenSpaceData GetScreenSpaceData(float2 UV, bool bGetNormalizedNormal = true)
 		float DeviceZ = SceneDepthSurface.Load(IntUV, SampleIndex);
 		float SceneDepth = ConvertFromDeviceZ(DeviceZ);
 
-	#if OUTPUT_GBUFFER_VELOCITY
+	#if WRITES_VELOCITY_TO_GBUFFER
 		float4 GBufferVelocity = GBuffers.GBufferVelocityTextureMS.Load(IntUV, SampleIndex);
 	#else
 		float4 GBufferVelocity = 0;
diff --git a/Engine/Source/Runtime/ShaderCore/Private/Shader.cpp b/Engine/Source/Runtime/ShaderCore/Private/Shader.cpp
--- a/Engine/Source/Runtime/ShaderCore/Private/Shader.cpp
+++ b/Engine/Source/Runtime/ShaderCore/Private/Shader.cpp
@@ -1140,6 +1140,14 @@ void ShaderMapAppendKeyString(EShaderPlatform Platform, FString& KeyString)
 		KeyString += bValue ? TEXT("_SL") : TEXT("_NoSL");
 	}
 
+	{
+		static const auto CVar = IConsoleManager::Get().FindTConsoleVariableDataInt(TEXT("r.BasePassOutputsVelocity"));
+		if (CVar && CVar->GetValueOnGameThread() != 0)
+		{
+			KeyString += TEXT("_GV");
+		}
+	}
+
 	{
 		static const auto CVar = IConsoleManager::Get().FindTConsoleVariableDataInt(TEXT("r.GBuffer"));
 		if(CVar ? CVar->GetValueOnGameThread() == 0 : false)
