UE4 - Fix typo

diff --git a/Engine/Shaders/BasePassCommon.usf b/Engine/Shaders/BasePassCommon.usf
--- a/Engine/Shaders/BasePassCommon.usf
+++ b/Engine/Shaders/BasePassCommon.usf
@@ -13,7 +13,7 @@
 #define NEEDS_BASEPASS_FOGGING			(ENABLE_TRANSLUCENCY_VERTEX_FOG && ((MATERIALBLENDING_TRANSLUCENT || MATERIALBLENDING_ADDITIVE || MATERIALBLENDING_MODULATE) && !MATERIAL_USES_SCENE_COLOR_COPY))
 #define NEEDS_LIGHTMAP					(NEEDS_LIGHTMAP_COORDINATE)
 
-#define WRITES_VELOCITY_TO_GBUFFER		(OUTPUT_GBUFFER_VELOCITY && (MATERIALBLENDING_SOLID || MATERIAL_BLENDING_MASKED))
+#define WRITES_VELOCITY_TO_GBUFFER		(OUTPUT_GBUFFER_VELOCITY && (MATERIALBLENDING_SOLID || MATERIALBLENDING_MASKED))
 
 struct FSharedBasePassInterpolants
 {
