UE4 - Output velocities on tessellated meshes

diff --git a/Engine/Shaders/BasePassTessellationShaders.usf b/Engine/Shaders/BasePassTessellationShaders.usf
--- a/Engine/Shaders/BasePassTessellationShaders.usf
+++ b/Engine/Shaders/BasePassTessellationShaders.usf
@@ -21,6 +21,11 @@ FBasePassVSToDS PassInterpolate(FBasePassVSToDS a, float aInterp, FBasePassVSToD
 		TESSELLATION_INTERPOLATE_MEMBER(BasePassInterpolants.WorldPositionExcludingWPO);
 	#endif
 
+	#if WRITES_VELOCITY_TO_GBUFFER
+		TESSELLATION_INTERPOLATE_MEMBER(BasePassInterpolants.PrevScreenPosition);
+		TESSELLATION_INTERPOLATE_MEMBER(BasePassInterpolants.ScreenPosition);
+	#endif
+
 	return O;
 }
 
