added shader function ComputePixelPosCenter() for easier computations in the shader

diff --git a/Engine/Shaders/PostProcessCommon.usf b/Engine/Shaders/PostProcessCommon.usf
--- a/Engine/Shaders/PostProcessCommon.usf
+++ b/Engine/Shaders/PostProcessCommon.usf
@@ -60,17 +60,27 @@ float EyeAdaptationLookup()
 
 // viewport width, height, 1/width, 1/height (scaled to the current rendertarget resolution)
 float4 ViewportSize;
-// in pixels (scaled to the current rendertarget resolution)
+// in pixels (scaled to the current rendertarget resolution), float4(minx,miny,maxx,maxy)
 float4 ViewportRect;
 
 SamplerState BilinearTextureSampler0;
 SamplerState BilinearTextureSampler1;
 
 // transforms ScreenPos to the screen pixel position, not the viewport local position (*xy +wz)
-// (0,0) left top, (width-1,height-1) at right bottom, in pixels
-// 	float2 PixelPos = UVAndScreenPos.zw * ScreenPosToPixel.xy + ScreenPosToPixel.zw;
+// (0,0) left top, (width-1,height-1) at right bottom, in pixels (without offset to reach the pixel center)
+// 	use ComputePixelPosCenter(UVAndScreenPos.zw, bPixelCenter) instead of directly accessing this
 float4 ScreenPosToPixel;
 
+// uses ScreenPosToPixel which is setup in FPostProcessPassParameters
+// for integer output bPixelCenter should be set to true (to avoid rounding errors)
+// @param ScreenPos -1 .. 1, usually from UVAndScreenPos.zw
+// @param bCenter false: left top of the pixel and almost no fractional value, true: added 0.5 to be pixel center
+float2 ComputePixelPosCenter(float2 ScreenPos, bool bPixelCenter)
+{
+	float2 PixelOffset = bPixelCenter ? 0.5f : 0.0f;
+	return ScreenPos * ScreenPosToPixel.xy + ScreenPosToPixel.zw + PixelOffset;
+}
+
 // shape of the lens for vignette like effects and masking
 // @param ScreenPos -1 .. 1
 // @return 0..1 0:borders, 1:center
