Fix nans coming out of new diffuse model

diff --git a/Engine/Shaders/BRDF.usf b/Engine/Shaders/BRDF.usf
--- a/Engine/Shaders/BRDF.usf
+++ b/Engine/Shaders/BRDF.usf
@@ -51,8 +51,17 @@ float3 Diffuse_Lambert( float3 DiffuseColor )
 float3 Diffuse_Burley( float3 DiffuseColor, float Roughness, float NoV, float NoL, float VoH )
 {
 	float FD90 = ( 0.5 + 2 * VoH * VoH ) * Roughness;
-	float FdV = 1 + (FD90 - 1) * pow( 1 - NoV, 5 );
-	float FdL = 1 + (FD90 - 1) * pow( 1 - NoL, 5 );
+	
+	float InvNov = 1- NoV;
+	float NoVPow5 = InvNov * InvNov;
+	NoVPow5 = NoVPow5 * NoVPow5 * InvNov;
+
+	float InvNoL = 1 - NoL;
+	float NoLPow5 = InvNoL * InvNoL;
+	NoLPow5 = NoLPow5 * NoLPow5 * InvNoL;
+
+	float FdV = 1 + (FD90 - 1) * NoVPow5;
+	float FdL = 1 + (FD90 - 1) * NoLPow5;
 	return DiffuseColor * ( 1 / PI * FdV * FdL ) * ( 1 - 0.3333 * Roughness );	// TODO premultiply DiffuseColor
 }
 
