UE4 - fix triangle collapses when instances are culled

diff --git a/Engine/Shaders/LocalVertexFactory.usf b/Engine/Shaders/LocalVertexFactory.usf
--- a/Engine/Shaders/LocalVertexFactory.usf
+++ b/Engine/Shaders/LocalVertexFactory.usf
@@ -525,7 +525,7 @@ half3x3 VertexFactoryGetTangentToLocal( FVertexFactoryInput Input, FVertexFactor
 float4 VertexFactoryGetWorldPosition(FVertexFactoryInput Input, FVertexFactoryIntermediates Intermediates)
 {
 #if USE_INSTANCING
-	return CalcWorldPosition(Input.Position, GetInstanceTransform(Input)) * Intermediates.PerInstanceParams.z;
+	return CalcWorldPosition(Input.Position, GetInstanceTransform(Input));
 #else
 	return CalcWorldPosition(Input.Position);
 #endif	// USE_INSTANCING
diff --git a/Engine/Shaders/MaterialTemplate.usf b/Engine/Shaders/MaterialTemplate.usf
--- a/Engine/Shaders/MaterialTemplate.usf
+++ b/Engine/Shaders/MaterialTemplate.usf
@@ -1172,11 +1172,25 @@ half GetMaterialOpacity(FMaterialPixelParameters Parameters)
 
 float3 GetMaterialWorldPositionOffset(FMaterialVertexParameters Parameters)
 {
+	#if USE_INSTANCING
+		// skip if this instance is hidden
+		if (Parameters.PerInstanceParams.z < 1.f)
+		{
+			return float3(0,0,0);
+		}
+	#endif
 %s;
 }
 
 float3 GetMaterialPreviousWorldPositionOffset(FMaterialVertexParameters Parameters)
 {
+	#if USE_INSTANCING
+		// skip if this instance is hidden
+		if (Parameters.PerInstanceParams.z < 1.f)
+		{
+			return float3(0,0,0);
+		}
+	#endif
 %s;
 }
 
