improved VisualizeBloom
fixed BloomScale (wasn't screen percentage, factor off by 4 and dependent on buffersize not viewport size

diff --git a/Engine/Shaders/PostProcessBloom.usf b/Engine/Shaders/PostProcessBloom.usf
--- a/Engine/Shaders/PostProcessBloom.usf
+++ b/Engine/Shaders/PostProcessBloom.usf
@@ -33,7 +33,7 @@ void VisualizeBloomSetupPS(
 
 	int2 PixelPos = (int2)ComputePixelPosCenter(UVAndScreenPos.zw, true);
 
-	int2 ViewportCenter = (int2)(ViewportRect.xy + ViewportSize / 2);
+	int2 ViewportCenter = (int2)(ViewportRect.xy + ViewportSize.xy / 2);
 
 	// cross in x and y
 //	bool bMask = ViewportCenter.x == PixelPos.x || ViewportCenter.y == PixelPos.y;
@@ -42,8 +42,16 @@ void VisualizeBloomSetupPS(
 
 	if(bMask)
 	{
-		OutColor = 20;
+		// green is the test impulse
+		OutColor.g = 20;
 	}
+	
+	float3 HDRColor = Texture2DSample(PostprocessInput0, PostprocessInput0Sampler, UV).rgb;
+
+	// reference is in red
+	OutColor.r = dot(HDRColor, 1/3.0f);
+	// blue is unused
+	OutColor.b = 0;
 }
 
 void VisualizeBloomOverlayPS(
@@ -51,50 +59,70 @@ void VisualizeBloomOverlayPS(
 	, out float4 OutColor : SV_Target0)
 {
 	float2 UV = UVAndScreenPos.xy;
-	
-	float3 LDRColor = Texture2DSample(PostprocessInput0, PostprocessInput0Sampler, UV).rgb;
-	
+	int2 PixelPos = (int2)ComputePixelPosCenter(UVAndScreenPos.zw, true);
+
 	int2 ViewportCenter = (int2)(ViewportRect + ViewportSize / 2);
 
-	float2 UVLinePos = float2(UV.x, ViewportCenter.y * PostprocessInput0Size.w);
+	// debug: zoomin
+	if(0)
+	{
+		PixelPos.x = (PixelPos.x - ViewportCenter.x) / 10 + ViewportCenter.x;
+		UV.x = PixelPos.x * PostprocessInput0Size.z;
+	}
+
+	float2 ViewPortUV = PixelPos * ViewportSize.zw;
 
-	float3 NonBloomColor = Texture2DSampleLevel(PostprocessInput1, PostprocessInput1Sampler, UVLinePos, 0).rgb;
+	float2 UVCenter = float2(UV.x, (int)(ViewportRect.y + ViewportSize.y / 2) * PostprocessInput0Size.w);
 
-	float3 BloomColor;
+	// bottom is reference, top is cretated through the bloom system
+	bool bBottom = PixelPos.y > ViewportCenter.y;
+
+	float HDRValueTest;
 	{
-		float4 CombinedBloom = Texture2DSample(PostprocessInput2, PostprocessInput2Sampler, UVLinePos);
+		float3 NonBloomColorTop = Texture2DSampleLevel(PostprocessInput1, PostprocessInput1Sampler, UVCenter, 0).rgb;
+		float4 CombinedBloom = Texture2DSample(PostprocessInput2, PostprocessInput2Sampler, UVCenter);
 		float3 BloomDirtMaskColor = 0;
-		BloomColor = CombinedBloom.rgb * (ColorScale1.rgb + BloomDirtMaskColor); 
+
+		// only take green, in red is the reference
+		HDRValueTest = (NonBloomColorTop + CombinedBloom.rgb * (ColorScale1.rgb + BloomDirtMaskColor)).g;
 	}
 
-	float NonBloomValue = Luminance(NonBloomColor);
-	float BloomValue = Luminance(BloomColor);
+	float HDRValueReference = Texture2DSampleLevel(PostprocessInput1, PostprocessInput1Sampler, UVCenter, 0).r;
+
+	float CompareY = lerp(1.1f, -0.1f, ViewPortUV.y);
 
-	float CompareY = lerp(1.1f, -0.1f, UV.y);
+	bool bOutside = CompareY <= 0 || CompareY >= 1;
 
-	bool bNonBloom = NonBloomValue > CompareY;
-	bool bBloom = BloomValue > CompareY;
-	bool bOutside = CompareY < 0 || CompareY > 1;
+	// non linear distribution for easier tweaking?
+	float ValueCompare = exp(-20 * (1-CompareY) + 3.5);
 
-	OutColor = 0.1f;
+	bool bReferenceMask = HDRValueReference > ValueCompare;
+	bool bTestMask = HDRValueTest > ValueCompare;
+	bool bCenter = ViewportCenter.x == PixelPos.x;
 
-	FLATTEN if(bNonBloom)
+	// faint grey in the background to see where the diagram ends
+	OutColor = 0.05f;
+
+	FLATTEN if(bReferenceMask || bTestMask)
 	{
-		OutColor.rgb = 0.3f;
+		OutColor = 0;
 	}
 
-	FLATTEN if(bBloom)
+	FLATTEN if(bReferenceMask)
 	{
-		OutColor.rgb = float3(0, 1, 0);
+		OutColor.g = 1;
+	}
+
+	FLATTEN if(bTestMask)
+	{
+		OutColor.r = 1;
 	}
 
 	FLATTEN if(bOutside)
 	{
-		OutColor.rgb = float3(0, 0, 0);
+		OutColor.rgb = bCenter ? float3(1, 0, 0) : float3(0,0,0);
 	}
-	
-	// for debugging: show LDR scene color
-//	OutColor.rgb = LDRColor;
+
 }
 
 // vertex shader entry point
diff --git a/Engine/Source/Runtime/Engine/Classes/Engine/Scene.h b/Engine/Source/Runtime/Engine/Classes/Engine/Scene.h
--- a/Engine/Source/Runtime/Engine/Classes/Engine/Scene.h
+++ b/Engine/Source/Runtime/Engine/Classes/Engine/Scene.h
@@ -160,6 +160,9 @@ struct FPostProcessSettings
 	UPROPERTY(BlueprintReadWrite, Category=Overrides, meta=(PinHiddenByDefault))
 	uint32 bOverride_Bloom5Size:1;
 
+	UPROPERTY(BlueprintReadWrite, Category=Overrides, meta=(PinHiddenByDefault))
+	uint32 bOverride_BloomSizeScale:1;
+
 	UPROPERTY(BlueprintReadWrite, Category=Overrides, meta=(PinHiddenByDefault))
 	uint32 bOverride_BloomDirtMaskIntensity:1;
 
@@ -451,6 +454,12 @@ struct FPostProcessSettings
 	UPROPERTY(interp, BlueprintReadWrite, Category=Bloom, AdvancedDisplay, meta=(ClampMin = "-1.0", UIMin = "0.0", UIMax = "8.0", editcondition = "bOverride_BloomThreshold", DisplayName = "Threshold"))
 	float BloomThreshold;
 
+	/**
+	 * Scale for all bloom sizes
+	 */
+	UPROPERTY(interp, BlueprintReadWrite, Category=Bloom, AdvancedDisplay, meta=(ClampMin = "0.0", UIMax = "64.0", editcondition = "bOverride_BloomSizeScale", DisplayName = "Size scale"))
+	float BloomSizeScale;
+
 	/**
 	 * Diameter size for the Bloom1 in percent of the screen width
 	 * (is done in 1/2 resolution, larger values cost more performance, good for high frequency details)
@@ -879,6 +888,8 @@ struct FPostProcessSettings
 		BloomIntensity = 1.0f;
 		BloomThreshold = 1.0f;
 		Bloom1Tint = FLinearColor(0.5f, 0.5f, 0.5f);
+		// default is 4 to maintain old settings after fixing something that caused a factor of 4
+		BloomSizeScale = 4.0;
 		Bloom1Size = 1.0f;
 		Bloom2Tint = FLinearColor(0.5f, 0.5f, 0.5f);
 		Bloom2Size = 4.0f;
diff --git a/Engine/Source/Runtime/Engine/Private/SceneView.cpp b/Engine/Source/Runtime/Engine/Private/SceneView.cpp
--- a/Engine/Source/Runtime/Engine/Private/SceneView.cpp
+++ b/Engine/Source/Runtime/Engine/Private/SceneView.cpp
@@ -734,6 +734,7 @@ void FSceneView::OverridePostProcessSettings(const FPostProcessSettings& Src, fl
 	LERP_PP(BloomIntensity);
 	LERP_PP(BloomThreshold);
 	LERP_PP(Bloom1Tint);
+	LERP_PP(BloomSizeScale);
 	LERP_PP(Bloom1Size);
 	LERP_PP(Bloom2Tint);
 	LERP_PP(Bloom2Size);
@@ -1213,7 +1214,8 @@ void FSceneView::EndFinalPostprocessSettings(const FSceneViewInitOptions& ViewIn
 		if( !Family->EngineShowFlags.PostProcessing || !Family->EngineShowFlags.AntiAliasing || Quality <= 0
 			// Disable antialiasing in GammaLDR mode to avoid jittering.
 			|| (FeatureLevel == ERHIFeatureLevel::ES2 && MobileHDRCvar->GetValueOnGameThread() == 0)
-			|| (FeatureLevel <= ERHIFeatureLevel::ES3_1 && (MSAAValue > 1)))
+			|| (FeatureLevel <= ERHIFeatureLevel::ES3_1 && (MSAAValue > 1))
+			|| Family->EngineShowFlags.VisualizeBloom)
 		{
 			FinalPostProcessSettings.AntiAliasingMethod = AAM_None;
 		}
diff --git a/Engine/Source/Runtime/Renderer/Private/PostProcess/PostProcessBloomSetup.cpp b/Engine/Source/Runtime/Renderer/Private/PostProcess/PostProcessBloomSetup.cpp
--- a/Engine/Source/Runtime/Renderer/Private/PostProcess/PostProcessBloomSetup.cpp
+++ b/Engine/Source/Runtime/Renderer/Private/PostProcess/PostProcessBloomSetup.cpp
@@ -432,7 +432,7 @@ void FRCPassPostProcessVisualizeBloomOverlay::Process(FRenderingCompositePassCon
 	// is optimized away if possible (RT size=view size, )
 	Context.RHICmdList.Clear(true, FLinearColor(0, 0, 0, 0), false, 1.0f, false, 0, DestRect);
 
-	Context.SetViewportAndCallRHI(0, 0, 0.0f, DestSize.X, DestSize.Y, 1.0f );
+	Context.SetViewportAndCallRHI(0, 0, 0.0f, DestRect.Width(), DestRect.Height(), 1.0f );
 
 	// set the state
 	Context.RHICmdList.SetBlendState(TStaticBlendState<>::GetRHI());
@@ -456,7 +456,7 @@ void FRCPassPostProcessVisualizeBloomOverlay::Process(FRenderingCompositePassCon
 		DestRect.Width(), DestRect.Height(),
 		SrcRect.Min.X, SrcRect.Min.Y,
 		SrcRect.Width(), SrcRect.Height(),
-		DestSize,
+		DestRect.Size(),
 		SrcSize,
 		*VertexShader,
 		EDRF_UseTriangleOptimization);
diff --git a/Engine/Source/Runtime/Renderer/Private/PostProcess/PostProcessWeightedSampleSum.cpp b/Engine/Source/Runtime/Renderer/Private/PostProcess/PostProcessWeightedSampleSum.cpp
--- a/Engine/Source/Runtime/Renderer/Private/PostProcess/PostProcessWeightedSampleSum.cpp
+++ b/Engine/Source/Runtime/Renderer/Private/PostProcess/PostProcessWeightedSampleSum.cpp
@@ -263,13 +263,19 @@ void SetFilterShaders(
  */
 static float NormalDistributionUnscaled(float X,float Mean,float Variance, EFilterShape FilterShape, float CrossCenterWeight)
 {
-	float Ret = FMath::Exp(-FMath::Square(X - Mean) / (2.0f * Variance));
+	float dx = FMath::Abs(X - Mean);
 
-	// if center sample
-	if(X == 0.0f)
+	float Ret = FMath::Exp(-FMath::Square(dx) / (2.0f * Variance));
+
+	// tweak the gaussian shape e.g. "r.Bloom.Cross 3.5"
+	if(CrossCenterWeight > 1.0f)
+	{
+		Ret = FMath::Max(0.0f, 1.0f - dx / Variance);
+		Ret = FMath::Pow(Ret, CrossCenterWeight);
+	}
+	else
 	{
-		// give it extra weight
-		Ret += CrossCenterWeight;
+		Ret = FMath::Lerp(Ret, FMath::Max(0.0f, 1.0f - dx / Variance), CrossCenterWeight);
 	}
 
 	return Ret;
@@ -360,11 +366,11 @@ void FRCPassPostProcessWeightedSampleSum::Process(FRenderingCompositePassContext
 	const FSceneRenderTargetItem& DestRenderTarget = PassOutputs[0].RequestSurface(Context);
 
 	FVector2D InvSrcSize(1.0f / SrcSize.X, 1.0f / SrcSize.Y);
-	int32 SrcSizeForThisAxis = (FilterShape == EFS_Horiz) ? SrcSize.X : SrcSize.Y;
-
-	// in texel (input resolution), *2 as we use the diameter
 	// we scale by width because FOV is defined horizontally
-	float EffectiveBlurRadius = SizeScale * SrcSizeForThisAxis  * 2 / 100.0f;
+	float SrcSizeForThisAxis = View.ViewRect.Width() / (float)SrcScaleFactor.X;
+
+	// in texel (input resolution), /2 as we use the diameter, 100 as we use percent
+	float EffectiveBlurRadius = SizeScale * SrcSizeForThisAxis  / 2 / 100.0f;
 
 	FVector2D BlurOffsets[MAX_FILTER_SAMPLES];
 	FLinearColor BlurWeights[MAX_FILTER_SAMPLES];
diff --git a/Engine/Source/Runtime/Renderer/Private/PostProcess/PostProcessing.cpp b/Engine/Source/Runtime/Renderer/Private/PostProcess/PostProcessing.cpp
--- a/Engine/Source/Runtime/Renderer/Private/PostProcess/PostProcessing.cpp
+++ b/Engine/Source/Runtime/Renderer/Private/PostProcess/PostProcessing.cpp
@@ -474,6 +474,9 @@ static FRenderingCompositeOutputRef AddBloom(FPostprocessContext Context, FRende
 			const FLinearColor* Tint;
 		};
 		const FFinalPostProcessSettings& Settings = Context.View.FinalPostProcessSettings;
+
+		bool bVisualizeBloom = Context.View.Family->EngineShowFlags.VisualizeBloom;
+
 		FBloomStage BloomStages[] =
 		{
 			{ Settings.Bloom5Size, &Settings.Bloom5Tint},
@@ -490,7 +493,20 @@ static FRenderingCompositeOutputRef AddBloom(FPostprocessContext Context, FRende
 		for (uint32 i = 0, SourceIndex = NumBloomStages - 1; i < BloomStageCount; i++, SourceIndex--)
 		{
 			FBloomStage& Op = BloomStages[i];
-			BloomOutput = RenderBloom(Context, PostProcessDownsamples[SourceIndex], Op.BloomSize, (*Op.Tint) * TintScale, BloomOutput);
+
+			FLinearColor Tint = (*Op.Tint) * TintScale;
+
+			if(bVisualizeBloom)
+			{
+				float LumScale = Tint.ComputeLuminance();
+
+				// R is used to pass down the reference, G is the emulated bloom
+				Tint.R = 0;
+				Tint.G = LumScale;
+				Tint.B = 0;
+			}
+
+			BloomOutput = RenderBloom(Context, PostProcessDownsamples[SourceIndex], Op.BloomSize * Settings.BloomSizeScale, Tint, BloomOutput);
 		}
 	}
 
@@ -1076,7 +1092,8 @@ void FPostProcessing::Process(FRHICommandListImmediate& RHICmdList, FViewInfo& V
 
 				if(View.Family->EngineShowFlags.EyeAdaptation
 					&& View.FinalPostProcessSettings.AutoExposureMinBrightness < View.FinalPostProcessSettings.AutoExposureMaxBrightness
-					&& !View.bIsSceneCapture) // Eye adaption is not available for scene captures.
+					&& !View.bIsSceneCapture // Eye adaption is not available for scene captures.
+					&& !bVisualizeBloom)
 				{
 					bHistogramNeeded = true;
 				}
