UE4 - Fix landscape previous world position

diff --git a/Engine/Shaders/LandscapeVertexFactory.usf b/Engine/Shaders/LandscapeVertexFactory.usf
--- a/Engine/Shaders/LandscapeVertexFactory.usf
+++ b/Engine/Shaders/LandscapeVertexFactory.usf
@@ -146,9 +146,19 @@ struct FVertexFactoryIntermediates
 #endif
 };
 
+float3 GetLocalPosition(FVertexFactoryIntermediates Intermediates)
+{
+	return Intermediates.LocalPosition+float3(Intermediates.InputPosition.zw * LandscapeParameters.SubsectionOffsetParams.ww,0);
+}
+
 float4 VertexFactoryGetWorldPosition(FVertexFactoryInput Input, FVertexFactoryIntermediates Intermediates)
 {
-	return TransformLocalToTranslatedWorld(Intermediates.LocalPosition+float3(Intermediates.InputPosition.zw * LandscapeParameters.SubsectionOffsetParams.ww,0));
+	return TransformLocalToTranslatedWorld(GetLocalPosition(Intermediates));
+}
+
+float4 VertexFactoryGetPreviousWorldPosition(FVertexFactoryInput Input, FVertexFactoryIntermediates Intermediates)
+{
+	return TransformLocalToPrevTranslatedWorld(GetLocalPosition(Intermediates));
 }
 
 /** Calculate the texture coordinates generated by Landscape */
@@ -747,12 +757,6 @@ FVertexFactoryInterpolantsVSToDS VertexFactoryGetInterpolantsVSToDS(FVertexFacto
 	return Interpolants;
 }
 
-float4 VertexFactoryGetPreviousWorldPosition(FVertexFactoryInput Input, FVertexFactoryIntermediates Intermediates)
-{
-	// Landscape doesn't move so use the current world position.
-	return VertexFactoryGetWorldPosition(Input, Intermediates);
-}
-
 FVertexFactoryInterpolantsVSToDS VertexFactoryInterpolate(FVertexFactoryInterpolantsVSToDS a, float aInterp, FVertexFactoryInterpolantsVSToDS b, float bInterp)
 {
 	FVertexFactoryInterpolantsVSToDS O;
diff --git a/Engine/Shaders/VertexFactoryCommon.usf b/Engine/Shaders/VertexFactoryCommon.usf
--- a/Engine/Shaders/VertexFactoryCommon.usf
+++ b/Engine/Shaders/VertexFactoryCommon.usf
@@ -14,3 +14,8 @@ float4 TransformLocalToTranslatedWorld(float3 LocalPosition)
 {
 	return float4(TransformLocalToWorld(LocalPosition) + View.PreViewTranslation.xyz,1);
 }
+
+float4 TransformLocalToPrevTranslatedWorld(float3 LocalPosition)
+{
+	return float4(TransformLocalToWorld(LocalPosition) + View.PrevPreViewTranslation.xyz,1);
+}
