added value display under center pixel to visualizehdr (useful for bloom/HDR tweak)

diff --git a/Engine/Shaders/PostProcessVisualizeHDR.usf b/Engine/Shaders/PostProcessVisualizeHDR.usf
--- a/Engine/Shaders/PostProcessVisualizeHDR.usf
+++ b/Engine/Shaders/PostProcessVisualizeHDR.usf
@@ -116,6 +116,9 @@ void MainPS(float4 UVAndScreenPos : TEXCOORD0, out float4 OutColor : SV_Target0)
 	int2 PixelPos = (int2)(UVAndScreenPos.zw * ScreenPosToPixel.xy + ScreenPosToPixel.zw + 0.5f);
 	float2 ViewLocalUV = float2(UVAndScreenPos.z * 0.5f + 0.5f, 0.5f - 0.5f * UVAndScreenPos.w);
 
+	int2 ViewportCenter = (int2)(ViewportRect.xy + ViewportSize.xy / 2);
+
+
 	// background is the scene color
 	float4 SceneColor = Texture2DSample(PostprocessInput0, PostprocessInput0Sampler, UV);
 
@@ -124,6 +127,25 @@ void MainPS(float4 UVAndScreenPos : TEXCOORD0, out float4 OutColor : SV_Target0)
 	OutColor = float4(Colorize(ComputeHistogramPositionFromLuminance(LuminanceVal)), 1.0f);
 
 
+	// crosshair
+	{
+		float CrossHairMask = PixelPos.x == ViewportCenter.x || PixelPos.y == ViewportCenter.y;
+		float2 DistAbs = abs(PixelPos - ViewportCenter);
+		float Dist = max(DistAbs.x, DistAbs.y);
+		float DistMask = Dist >= 2 && Dist < 7;
+
+		OutColor.xyz = lerp(OutColor.xyz, float3(1, 1, 1), CrossHairMask * DistMask);
+	}
+
+	// value under crosshair
+	{
+		float3 CenterViewportHDRColor = Texture2DSample(PostprocessInput0, PostprocessInput0Sampler, ViewportCenter * PostprocessInput0Size.zw).rgb;
+		float Value = Luminance(CenterViewportHDRColor);
+
+		PrintFloat(PixelPos, OutColor.xyz, float3(1, 1, 1), ViewportCenter + int2(-27, 7), Value);
+	}
+
+
 
 	float2 IDAreaLocalUV = ViewLocalUV * 2 + float2(-1, 0);
 
@@ -179,7 +201,7 @@ void MainPS(float4 UVAndScreenPos : TEXCOORD0, out float4 OutColor : SV_Target0)
 	float2 InsetUV = InsetPx / HistogramSize;
 	
 //	const float3 BorderColor = float3(0.5f, 0.5f, 0.5f);
-	const float3 BorderColor = Colorize(InsetUV);
+	const float3 BorderColor = Colorize(InsetUV.x);
 
 	float BorderDistance = ComputeDistanceToRect(PixelPos, HistogramLeftTop, HistogramSize);
 	
@@ -228,7 +250,7 @@ void MainPS(float4 UVAndScreenPos : TEXCOORD0, out float4 OutColor : SV_Target0)
 	if(LocalHistogramValue >= 1 - InsetUV.y)
 	{
 		// histogram bars
-		OutColor.xyz = lerp(OutColor.xyz, Colorize(InsetUV), 0.5f);
+		OutColor.xyz = lerp(OutColor.xyz, Colorize(InsetUV.x), 0.5f);
 	}
 
 	{
@@ -251,6 +273,7 @@ void MainPS(float4 UVAndScreenPos : TEXCOORD0, out float4 OutColor : SV_Target0)
 		}
 	}
 
+	// eye adaptation
 	{
 		float Value = ComputeHistogramPositionFromLuminance(1.0f / HistogramTexture.Load(int3(0, 1, 0)).x);
 
