nicer VisualizeHDR

diff --git a/Engine/Shaders/PostProcessVisualizeHDR.usf b/Engine/Shaders/PostProcessVisualizeHDR.usf
--- a/Engine/Shaders/PostProcessVisualizeHDR.usf
+++ b/Engine/Shaders/PostProcessVisualizeHDR.usf
@@ -275,12 +275,13 @@ void MainPS(float4 UVAndScreenPos : TEXCOORD0, out float4 OutColor : SV_Target0)
 
 	// eye adaptation
 	{
-		float Value = ComputeHistogramPositionFromLuminance(1.0f / HistogramTexture.Load(int3(0, 1, 0)).x);
+		float EyeApatationValue = ComputeHistogramPositionFromLuminance(1.0f / HistogramTexture.Load(int3(0, 1, 0)).x);
 
-		PrintFloat(PixelPos, OutColor.xyz, float3(1, 1, 1), HistogramLeftTop + int2(HistogramSize.x - 7 * 8 - 2, 3), Value);
+		float ValuePx = EyeApatationValue * HistogramSize.x;
 
-		float ValuePx = Value * HistogramSize.x;
-		if(abs(InsetPx.x - ValuePx) < 2)
+		PrintFloat(PixelPos, OutColor.xyz, float3(1, 1, 1), HistogramLeftTop + int2(ValuePx + - 3 * 8 - 3, 1), EyeApatationValue);
+
+		if(abs(InsetPx.x - ValuePx) < 2 && PixelPos.y > HistogramLeftTop.y + 9)
 		{
 			// white line to show the smoothed exposure
 			OutColor = lerp(OutColor, float4(1, 1, 1, 0), 1.0f);
