UE4 - fix collapsing triangles and foliage selection

diff --git a/Engine/Shaders/LocalVertexFactory.usf b/Engine/Shaders/LocalVertexFactory.usf
--- a/Engine/Shaders/LocalVertexFactory.usf
+++ b/Engine/Shaders/LocalVertexFactory.usf
@@ -525,7 +525,7 @@ half3x3 VertexFactoryGetTangentToLocal( FVertexFactoryInput Input, FVertexFactor
 float4 VertexFactoryGetWorldPosition(FVertexFactoryInput Input, FVertexFactoryIntermediates Intermediates)
 {
 #if USE_INSTANCING
-	return CalcWorldPosition(Input.Position, GetInstanceTransform(Input));
+	return CalcWorldPosition(Input.Position, GetInstanceTransform(Input)) * Intermediates.PerInstanceParams.z;
 #else
 	return CalcWorldPosition(Input.Position);
 #endif	// USE_INSTANCING
