added RGB to VisualizeHDR

diff --git a/Engine/Shaders/PostProcessVisualizeHDR.usf b/Engine/Shaders/PostProcessVisualizeHDR.usf
--- a/Engine/Shaders/PostProcessVisualizeHDR.usf
+++ b/Engine/Shaders/PostProcessVisualizeHDR.usf
@@ -142,7 +142,23 @@ void MainPS(float4 UVAndScreenPos : TEXCOORD0, out float4 OutColor : SV_Target0)
 		float3 CenterViewportHDRColor = Texture2DSample(PostprocessInput0, PostprocessInput0Sampler, ViewportCenter * PostprocessInput0Size.zw).rgb;
 		float Value = Luminance(CenterViewportHDRColor);
 
-		PrintFloat(PixelPos, OutColor.xyz, float3(1, 1, 1), ViewportCenter + int2(-27, 7), Value);
+		int2 Cursor;
+
+		Cursor = ViewportCenter + int2(-35, 7 + 0 * 8);
+		PrintCharacter(PixelPos, OutColor.xyz, float3(1, 1, 1), Cursor, 10 + 11);	// 'L'
+		PrintFloat(PixelPos, OutColor.xyz, float3(1, 1, 1), Cursor, Value);
+	
+		Cursor = ViewportCenter + int2(-35, 7 + 1 * 8); 
+		PrintCharacter(PixelPos, OutColor.xyz, float3(1, 0, 0), Cursor, 10 + 17);	// 'R'
+		PrintFloat(PixelPos, OutColor.xyz, float3(1, 1, 1), Cursor, CenterViewportHDRColor.r);
+
+		Cursor = ViewportCenter + int2(-35, 7 + 2 * 8); 
+		PrintCharacter(PixelPos, OutColor.xyz, float3(0, 1, 0), Cursor, 10 + 6);	// 'G'
+		PrintFloat(PixelPos, OutColor.xyz, float3(1, 1, 1), Cursor, CenterViewportHDRColor.g);
+
+		Cursor = ViewportCenter + int2(-35, 7 + 3 * 8); 
+		PrintCharacter(PixelPos, OutColor.xyz, float3(0, 0, 1), Cursor, 10 + 1);	// 'B'
+		PrintFloat(PixelPos, OutColor.xyz, float3(1, 1, 1), Cursor, CenterViewportHDRColor.b);
 	}
 
 
