Surfel Global Illumination with Distance Field shadowing
* Disabled by default, need a movable skylight with r.DistanceFieldGI enabled at load time (ConsoleVariables.ini)
* Surfels (oriented disks) are placed on mesh surfaces and capture diffuse color by evaluating the material with actual UVs and world position
* Surfels are lit and shadowed by ray tracing through mesh distance fields (directional light only so far)
* Irradiance cache records are lit with nearby Virtual Point Lights created from lit surfels.  Only using camera based LOD of the surfel representation and object based surfel culling, no hierarchy yet.
* Shadowing of the VPLs is provided by distance field cone tracing.  A depth visibility function is constructed along each of the 9 cone trace directions, and used to shadow the incoming VPL lighting.
* Component reregistration and transform updating causes surfels to be generated again which can be slow.  In the future only the material needs to be re-evaluated.
* Geometry shaders with stream out are now supported, by implementing FShader::GetStreamOutElements
* Compute shaders can now evaluate materal attributes like BaseColor
* SkyLight MaxOcclusionDistance increased to 10m, which is the minimum useful amount for GI

diff --git a/Engine/Shaders/Common.usf b/Engine/Shaders/Common.usf
--- a/Engine/Shaders/Common.usf
+++ b/Engine/Shaders/Common.usf
@@ -171,23 +171,43 @@
 // for type checking of the parameters - ironically the type checking is really only needed in GLSL!
 MaterialFloat4 Texture1DSample(Texture1D Tex, SamplerState Sampler, float UV)
 {
+#if COMPUTESHADER
+	return Tex.SampleLevel(Sampler, UV, 0);
+#else
 	return Tex.Sample(Sampler, UV);
+#endif
 }
 MaterialFloat4 Texture2DSample(Texture2D Tex, SamplerState Sampler, float2 UV)
 {
+#if COMPUTESHADER
+	return Tex.SampleLevel(Sampler, UV, 0);
+#else
 	return Tex.Sample(Sampler, UV);
+#endif
 }
 MaterialFloat Texture2DSample_A8(Texture2D Tex, SamplerState Sampler, float2 UV)
 {
+#if COMPUTESHADER
+	return Tex.SampleLevel(Sampler, UV, 0) A8_SAMPLE_MASK;
+#else
 	return Tex.Sample(Sampler, UV) A8_SAMPLE_MASK;
+#endif
 }
 MaterialFloat4 Texture3DSample(Texture3D Tex, SamplerState Sampler, float3 UV)
 {
+#if COMPUTESHADER
+	return Tex.SampleLevel(Sampler, UV, 0);
+#else
 	return Tex.Sample(Sampler, UV);
+#endif
 }
 MaterialFloat4 TextureCubeSample(TextureCube Tex, SamplerState Sampler, float3 UV)
 {
+#if COMPUTESHADER
+	return Tex.SampleLevel(Sampler, UV, 0);
+#else
 	return Tex.Sample(Sampler, UV);
+#endif
 }
 MaterialFloat4 Texture1DSampleLevel(Texture1D Tex, SamplerState Sampler, float UV, MaterialFloat Mip)
 {
@@ -199,7 +219,11 @@ MaterialFloat4 Texture2DSampleLevel(Texture2D Tex, SamplerState Sampler, float2
 }
 MaterialFloat4 Texture2DSampleBias(Texture2D Tex, SamplerState Sampler, float2 UV, MaterialFloat MipBias)
 {
+#if COMPUTESHADER
+	return Tex.SampleLevel(Sampler, UV, 0);
+#else
 	return Tex.SampleBias(Sampler, UV, MipBias);
+#endif
 }
 MaterialFloat4 Texture3DSampleLevel(Texture3D Tex, SamplerState Sampler, float3 UV, MaterialFloat Mip)
 {
@@ -211,7 +235,11 @@ MaterialFloat4 TextureCubeSampleLevel(TextureCube Tex, SamplerState Sampler, flo
 }
 MaterialFloat4 TextureCubeSampleBias(TextureCube Tex, SamplerState Sampler, float3 UV, MaterialFloat MipBias)
 {
+#if COMPUTESHADER
+	return Tex.SampleLevel(Sampler, UV, 0);
+#else
 	return Tex.SampleBias(Sampler, UV, MipBias);
+#endif
 }
 
 #if FEATURE_LEVEL >= FEATURE_LEVEL_SM5 // Cubemap arrays are not supported in SM4 feature level
@@ -304,6 +332,78 @@ MaterialFloat4 ClampedPow(MaterialFloat4 X,MaterialFloat4 Y)
 	return pow(max(abs(X),MaterialFloat4(0.000001f,0.000001f,0.000001f,0.000001f)),Y);
 } 
 
+float DDX(float Input)
+{
+#if COMPUTESHADER
+	return 0;
+#else
+	return ddx(Input);
+#endif
+}
+
+float2 DDX(float2 Input)
+{
+#if COMPUTESHADER
+	return 0;
+#else
+	return ddx(Input);
+#endif
+}
+
+float3 DDX(float3 Input)
+{
+#if COMPUTESHADER
+	return 0;
+#else
+	return ddx(Input);
+#endif
+}
+
+float4 DDX(float4 Input)
+{
+#if COMPUTESHADER
+	return 0;
+#else
+	return ddx(Input);
+#endif
+}
+
+float DDY(float Input)
+{
+#if COMPUTESHADER
+	return 0;
+#else
+	return ddy(Input);
+#endif
+}
+
+float2 DDY(float2 Input)
+{
+#if COMPUTESHADER
+	return 0;
+#else
+	return ddy(Input);
+#endif
+}
+
+float3 DDY(float3 Input)
+{
+#if COMPUTESHADER
+	return 0;
+#else
+	return ddy(Input);
+#endif
+}
+
+float4 DDY(float4 Input)
+{
+#if COMPUTESHADER
+	return 0;
+#else
+	return ddy(Input);
+#endif
+}
+
 #include "Random.usf"	// used by MaterialExpressionNoise
   
 /** 
diff --git a/Engine/Shaders/ConvertToUniformMesh.usf b/Engine/Shaders/ConvertToUniformMesh.usf
new file mode 100644
--- /dev/null
+++ b/Engine/Shaders/ConvertToUniformMesh.usf
@@ -0,0 +1,74 @@
+// Copyright 1998-2014 Epic Games, Inc. All Rights Reserved.
+
+/*=============================================================================
+	ConvertToUniformMesh.usf
+=============================================================================*/
+
+#include "Common.usf"
+#include "Material.usf"
+#include "VertexFactory.usf"
+
+struct FConvertToUniformMeshVSToGS
+{
+	FVertexFactoryInterpolantsVSToPS FactoryInterpolants;
+	float3 PixelPosition : TEXCOORD6;
+};
+
+void ConvertToUniformMeshVS(
+	FVertexFactoryInput Input,
+	out FConvertToUniformMeshVSToGS Output
+	)
+{
+	FVertexFactoryIntermediates VFIntermediates = GetVertexFactoryIntermediates(Input);
+	float3 WorldPosition = VertexFactoryGetWorldPosition(Input, VFIntermediates).xyz - View.PreViewTranslation.xyz;
+	float3x3 TangentToLocal = VertexFactoryGetTangentToLocal(Input, VFIntermediates);
+
+	FMaterialVertexParameters VertexParameters = GetMaterialVertexParameters(Input, VFIntermediates, WorldPosition, TangentToLocal);
+
+	Output.FactoryInterpolants = VertexFactoryGetInterpolantsVSToPS(Input, VFIntermediates, VertexParameters);
+	Output.PixelPosition = WorldPosition;
+}
+
+// Must match GUniformMeshVertexLayout and TRIANGLE_VERTEX_DATA_STRIDE
+struct FConvertToUniformMeshVertex
+{
+	float4 Position : SV_Position;
+	float3 Tangent0 : Tangent0;
+	float3 Tangent1 : Tangent1;
+	float3 Tangent2 : Tangent2;
+	float2 UV0 : UV0;
+	float2 UV1 : UV1;
+	float4 VertexColor : VertexColor;
+};
+
+[maxvertexcount(3)]
+void ConvertToUniformMeshGS(triangle FConvertToUniformMeshVSToGS Inputs[3], inout TriangleStream<FConvertToUniformMeshVertex> OutStream)
+{ 
+	for (int i = 0; i < 3; i++)
+	{
+		FConvertToUniformMeshVSToGS Input = Inputs[i];
+		FMaterialPixelParameters MaterialParameters = GetMaterialPixelParameters(Input.FactoryInterpolants, float4(Input.PixelPosition, 1));
+
+		FConvertToUniformMeshVertex Vertex = (FConvertToUniformMeshVertex)0;
+		Vertex.Position = float4(Input.PixelPosition, 40);
+
+		Vertex.Tangent0 = MaterialParameters.TangentToWorld[0];
+		Vertex.Tangent1 = MaterialParameters.TangentToWorld[1];
+		Vertex.Tangent2 = MaterialParameters.TangentToWorld[2];
+
+		Vertex.UV0 = Vertex.UV1 = 0;
+
+		#if NUM_MATERIAL_TEXCOORDS > 0
+			Vertex.UV0 = MaterialParameters.TexCoords[0];
+		#endif
+
+		#if NUM_MATERIAL_TEXCOORDS > 1
+			Vertex.UV1 = MaterialParameters.TexCoords[1];
+		#endif
+
+		Vertex.VertexColor = MaterialParameters.VertexColor;
+
+		OutStream.Append(Vertex);
+	}
+}
+
diff --git a/Engine/Shaders/Definitions.usf b/Engine/Shaders/Definitions.usf
--- a/Engine/Shaders/Definitions.usf
+++ b/Engine/Shaders/Definitions.usf
@@ -48,8 +48,8 @@
 #define	MATERIAL_SHADINGMODEL_UNLIT						0
 #endif
 
-#ifndef MESH_MATERIAL_SHADER
-#define	MESH_MATERIAL_SHADER							0
+#ifndef HAS_PRIMITIVE_UNIFORM_BUFFER
+#define	HAS_PRIMITIVE_UNIFORM_BUFFER					0
 #endif
 
 #ifndef OUTPUT_DEPTH_TO_ALPHA
diff --git a/Engine/Shaders/DistanceFieldAOShared.usf b/Engine/Shaders/DistanceFieldAOShared.usf
--- a/Engine/Shaders/DistanceFieldAOShared.usf
+++ b/Engine/Shaders/DistanceFieldAOShared.usf
@@ -59,9 +59,7 @@ Buffer<uint4> TileHeadData;
 Buffer<uint> TileArrayData;
 uint2 TileListGroupSize;
 
-
 #define UNPACKED_TILE_HEAD 1
-#define COHERENT_OBJECT_LIST_WRITES 1
  
 uint4 GetTileHead(uint2 TileCoordinate) 
 { 
@@ -127,4 +125,5 @@ bool SphereIntersectConeWithDepthRanges(float4 SphereCenterAndRadius, float3 Con
 	}
 
 	return false;
-}
\ No newline at end of file
+}
+
diff --git a/Engine/Shaders/DistanceFieldGlobalIllumination.usf b/Engine/Shaders/DistanceFieldGlobalIllumination.usf
--- a/Engine/Shaders/DistanceFieldGlobalIllumination.usf
+++ b/Engine/Shaders/DistanceFieldGlobalIllumination.usf
@@ -10,9 +10,9 @@
 #include "DistanceFieldAOShared.usf"
 #include "DynamicLightingCommon.usf"
 
-#define USE_SHADOW_CULLING_FOR_VPL_PLACEMENT 0
+#define USE_SHADOW_CULLING_FOR_VPL_PLACEMENT 1
 
-void RayTraceThroughLightTileCulledObjects(
+void RayTraceThroughLightTileCulledObjectsFirstHit(
 	float3 RayStartPosition, 
 	float3 RayDirection, 
 	float RayLength, 
@@ -166,13 +166,8 @@ float4x4 ShadowToWorld;
 float2 InvPlacementGridSize;
 float VPLPlacementCameraRadius;
 
-// Must match C++
-#define VPL_CLUSTER_SIZE_ONE_DIM 4
-#define VPL_CLUSTER_SIZE (VPL_CLUSTER_SIZE_ONE_DIM * VPL_CLUSTER_SIZE_ONE_DIM)
-
 // In float4's, must match C++
 RWBuffer<uint> RWVPLParameterBuffer;
-RWBuffer<float4> RWVPLClusterData;
 RWBuffer<float4> RWVPLData;
 
 [numthreads(THREADGROUP_SIZEX, THREADGROUP_SIZEY, 1)]
@@ -185,105 +180,6 @@ void VPLPlacementCS(
 		// Distance for directional lights to trace
 		float TraceDistance = LightDirectionAndTraceDistance.w;
 		float3 LightDirection = LightDirectionAndTraceDistance.xyz;
-		
-// Note: have to change the dispatch dimensions in C++ to enable this
-#define GENERATE_VPL_CLUSTERS 0
-#if GENERATE_VPL_CLUSTERS
-		uint NextVPLIndex = 0;
-		float3 IntersectionPositions[VPL_CLUSTER_SIZE];
-		float3 IntersectionNormals[VPL_CLUSTER_SIZE];
-		float3 AveragePosition = 0;
-		float3 AverageNormal = 0;
-
-		for (uint ChildY = 0; ChildY < VPL_CLUSTER_SIZE_ONE_DIM; ChildY++)
-		{
-			for (uint ChildX = 0; ChildX < VPL_CLUSTER_SIZE_ONE_DIM; ChildX++)
-			{
-				uint2 ChildCoordinate = DispatchThreadId.xy * VPL_CLUSTER_SIZE_ONE_DIM + uint2(ChildX, ChildY);
-
-				if (all(ChildCoordinate * InvPlacementGridSize < 1))
-				{
-					float2 NormalizedCellPosition = ChildCoordinate * InvPlacementGridSize * 2 - 1;
-					float3 CellStartWorldPosition = mul(float4(NormalizedCellPosition.x, NormalizedCellPosition.y, 0, 1), ShadowToWorld).xyz;
-
-					uint NumIntersectingObjects = GetCulledNumObjects();
-					uint CulledDataStart = 0;
-
-				#if USE_SHADOW_CULLING_FOR_VPL_PLACEMENT
-
-					// Transform into shadow space
-					float4 HomogeneousShadowPosition = mul(float4(CellStartWorldPosition, 1), WorldToShadow);
-					float2 NormalizedShadowPosition = HomogeneousShadowPosition.xy * .5f + .5f;
-					NormalizedShadowPosition.y = 1 - NormalizedShadowPosition.y;
-					// Quantize the shadow position to get our tile position
-					uint2 TilePosition = (uint2)(NormalizedShadowPosition * ShadowTileListGroupSize);
-					// Fetch the tile head information
-					uint2 TileHead = GetShadowTileHead(TilePosition);
-					NumIntersectingObjects = TileHead.y;
-					CulledDataStart = TileHead.x;
-
-				#endif
-
-					float MinRayTime = 0;
-					float MinRayVisibility = 1;
-					uint ObjectIndex = 0;
-					RayTraceThroughLightTileCulledObjects(CellStartWorldPosition, LightDirection, TraceDistance, NumIntersectingObjects, CulledDataStart, MinRayTime, MinRayVisibility, ObjectIndex);
-
-					if (MinRayVisibility < 1)
-					{
-						float3 IntersectionPosition = CellStartWorldPosition + LightDirection * MinRayTime;
-						float3 IntersectionNormal = ComputeDistanceFieldNormal(IntersectionPosition, ObjectIndex);
-						AveragePosition += IntersectionPosition;
-						AverageNormal += IntersectionNormal;
-						IntersectionPositions[NextVPLIndex] = IntersectionPosition;
-						IntersectionNormals[NextVPLIndex] = IntersectionNormal;
-						NextVPLIndex++;
-					}
-				}
-			}
-		}
-
-		uint NumVPLs = NextVPLIndex;
-
-		if (NumVPLs > 0)
-		{
-			uint VPLArrayStartIndex;
-			InterlockedAdd(RWVPLParameterBuffer[1], NumVPLs, VPLArrayStartIndex);
-
-			float3 DiffuseColor = .5f;
-			float CellExtent = VPLPlacementCameraRadius * InvPlacementGridSize.x;
-			float CellRadius = sqrt(2) * CellExtent;
-
-			float3 TotalFlux = 0;
-
-			for (uint ChildIndex = 0; ChildIndex < NumVPLs; ChildIndex++)
-			{
-				float3 Position = IntersectionPositions[ChildIndex];
-				float3 Normal = IntersectionNormals[ChildIndex];
-				float3 Flux = DiffuseColor * LightColor.rgb * max(dot(Normal, -LightDirection), 0) * PI * CellRadius * CellRadius;
-
-				uint VPLBaseIndex = (VPLArrayStartIndex + ChildIndex) * VPL_DATA_STRIDE;
-				RWVPLData[VPLBaseIndex + 0] = float4(Position, CellRadius);
-				RWVPLData[VPLBaseIndex + 1] = float4(Normal, 0);
-				RWVPLData[VPLBaseIndex + 2] = float4(Flux, 0);
-
-				TotalFlux += Flux;
-			}
-
-			AveragePosition /= NumVPLs;
-			float AverageNormalLength = length(AverageNormal);
-			AverageNormal = AverageNormal / max(AverageNormalLength, .000001f);
-
-			float TotalChildDiskArea = NumVPLs * PI * CellRadius * CellRadius;
-			float CombinedRadius = sqrt(TotalChildDiskArea / PI);
-
-			uint ClusterArrayStartIndex;
-			InterlockedAdd(RWVPLParameterBuffer[0], 1, ClusterArrayStartIndex);
-			RWVPLClusterData[ClusterArrayStartIndex * VPL_DATA_STRIDE + 0] = float4(AveragePosition, CombinedRadius);
-			RWVPLClusterData[ClusterArrayStartIndex * VPL_DATA_STRIDE + 1] = float4(AverageNormal, NumVPLs);
-			RWVPLClusterData[ClusterArrayStartIndex * VPL_DATA_STRIDE + 2] = float4(TotalFlux, VPLArrayStartIndex);
-		}
-#else
 		uint2 ChildCoordinate = DispatchThreadId.xy;
 
 		if (all(ChildCoordinate * InvPlacementGridSize < 1))
@@ -312,7 +208,7 @@ void VPLPlacementCS(
 			float MinRayTime = 0;
 			float MinRayVisibility = 1;
 			uint ObjectIndex = 0;
-			RayTraceThroughLightTileCulledObjects(CellStartWorldPosition, LightDirection, TraceDistance, NumIntersectingObjects, CulledDataStart, MinRayTime, MinRayVisibility, ObjectIndex);
+			RayTraceThroughLightTileCulledObjectsFirstHit(CellStartWorldPosition, LightDirection, TraceDistance, NumIntersectingObjects, CulledDataStart, MinRayTime, MinRayVisibility, ObjectIndex);
 
 			if (MinRayVisibility < 1)
 			{
@@ -334,7 +230,6 @@ void VPLPlacementCS(
 				RWVPLData[VPLBaseIndex + 2] = float4(Flux, 0);
 			}
 		}
-#endif
 	}
 }
 
@@ -347,11 +242,7 @@ void SetupVPLCullndirectArgumentsCS(
 	uint3 DispatchThreadId : SV_DispatchThreadID,
     uint3 GroupThreadId : SV_GroupThreadID) 
 {
-#if GENERATE_VPL_CLUSTERS
-	uint NumClusterVPLs = VPLParameterBuffer[0];
-#else
 	uint NumClusterVPLs = VPLParameterBuffer[1];
-#endif
 
 	// One thread per record, divide and round up
 	RWDispatchParameters[0] = (NumClusterVPLs + THREADGROUP_TOTALSIZE - 1) / THREADGROUP_TOTALSIZE;
@@ -359,11 +250,21 @@ void SetupVPLCullndirectArgumentsCS(
 	RWDispatchParameters[2] = 1;
 }
 
-Buffer<float4> VPLClusterData;
+uint DebugId;
+RWBuffer<uint> RWDebugBuffer;
+
+[numthreads(1, 1, 1)]
+void TrackGPUProgressCS(
+	uint3 GroupId : SV_GroupID,
+	uint3 DispatchThreadId : SV_DispatchThreadID,
+    uint3 GroupThreadId : SV_GroupThreadID) 
+{
+	RWDebugBuffer[0] = DebugId;
+}
+
 Buffer<float4> VPLData;
 
 RWBuffer<uint> RWCulledVPLParameterBuffer;
-RWBuffer<float4> RWCulledVPLClusterData;
 RWBuffer<float4> RWCulledVPLData;
 
 [numthreads(THREADGROUP_TOTALSIZE, 1, 1)]
@@ -372,46 +273,6 @@ void CullVPLsForViewCS(
 	uint3 DispatchThreadId : SV_DispatchThreadID,
     uint3 GroupThreadId : SV_GroupThreadID) 
 {
-#if GENERATE_VPL_CLUSTERS
-	uint VPLClusterIndex = DispatchThreadId.x;
-	uint NumClusterVPLs = VPLParameterBuffer[0];
-
-	if (VPLClusterIndex < NumClusterVPLs)
-	{
-		float4 BoundingSphere = VPLClusterData[VPLClusterIndex * VPL_DATA_STRIDE + 0];
-		float DistanceToViewSq = dot(View.ViewOrigin.xyz - BoundingSphere.xyz, View.ViewOrigin.xyz - BoundingSphere.xyz);
-
-		if (DistanceToViewSq < Square(AOMaxViewDistance + BoundingSphere.w)
-			&& ViewFrustumIntersectSphere(BoundingSphere.xyz, BoundingSphere.w + AOMaxDistance))
-		{
-			float4 ClusterNormalAndNumVPLs = VPLClusterData[VPLClusterIndex * VPL_DATA_STRIDE + 1];
-			float4 ClusterFluxAndVPLStartIndex = VPLClusterData[VPLClusterIndex * VPL_DATA_STRIDE + 2];
-
-			uint DestStartVPLIndex;
-			InterlockedAdd(RWCulledVPLParameterBuffer[1], (uint)ClusterNormalAndNumVPLs.w, DestStartVPLIndex);
-			uint SourceVPLIndex = ClusterFluxAndVPLStartIndex.w;
-			uint NumChildVPLs = ClusterNormalAndNumVPLs.w;
-
-			for (uint ChildIndex = 0; ChildIndex < NumChildVPLs; ChildIndex++)
-			{
-				uint SourceIndex = (SourceVPLIndex + ChildIndex) * VPL_DATA_STRIDE;
-				uint DestIndex = (DestStartVPLIndex + ChildIndex) * VPL_DATA_STRIDE;
-				RWCulledVPLData[DestIndex + 0] = VPLData[SourceIndex + 0];
-				RWCulledVPLData[DestIndex + 1] = VPLData[SourceIndex + 1];
-				RWCulledVPLData[DestIndex + 2] = VPLData[SourceIndex + 2];
-			}
-
-			uint DestClusterIndex;
-			InterlockedAdd(RWCulledVPLParameterBuffer[0], 1, DestClusterIndex);
-			
-			RWCulledVPLClusterData[DestClusterIndex * VPL_DATA_STRIDE + 0] = BoundingSphere;
-			RWCulledVPLClusterData[DestClusterIndex * VPL_DATA_STRIDE + 1] = ClusterNormalAndNumVPLs;
-			RWCulledVPLClusterData[DestClusterIndex * VPL_DATA_STRIDE + 2] = float4(ClusterFluxAndVPLStartIndex.xyz, DestStartVPLIndex);
-		}
-	}
-
-#else
-
 	uint VPLIndex = DispatchThreadId.x;
 	uint NumVPLs = VPLParameterBuffer[1];
 
@@ -433,54 +294,298 @@ void CullVPLsForViewCS(
 			RWCulledVPLData[DestIndex + 2] = VPLData[SourceIndex + 2];
 		}
 	}
+}
 
+#ifndef LIGHT_VPLS_THREADGROUP_SIZE
+#define LIGHT_VPLS_THREADGROUP_SIZE 1
 #endif
+
+uint ObjectProcessStride;
+
+[numthreads(1, 1, 1)]
+void SetupLightVPLsIndirectArgumentsCS(
+	uint3 GroupId : SV_GroupID,
+	uint3 DispatchThreadId : SV_DispatchThreadID,
+    uint3 GroupThreadId : SV_GroupThreadID) 
+{
+	uint NumGroups = (GetCulledNumObjects() + ObjectProcessStride - 1) / ObjectProcessStride;
+
+	// One group per object
+	RWDispatchParameters[0] = NumGroups;
+	RWDispatchParameters[1] = 1;
+	RWDispatchParameters[2] = 1;
 }
 
-RWBuffer<uint> RWTileArrayData;
-RWBuffer<uint> RWTileHeadDataUnpacked;
+// From receiver to to light
+float3 LightDirection;
+float4 LightPositionAndInvRadius;
+float LightSourceRadius;
+float2 TanLightAngleAndNormalThreshold;
 
-Buffer<float4> TileConeAxisAndCos;
-Buffer<float4> TileConeDepthRanges;
-uint2 NumGroups;
+Buffer<float4> ShadowCulledObjectBounds;
+Buffer<float4> ShadowCulledObjectData;
 
-// Must match C++
-#define MAX_VPLS_PER_TILE 1024
+float RayTraceThroughLightTileCulledObjects(
+	float3 WorldRayStart, 
+	float3 WorldRayEnd, 
+	float MaxRayTime,
+	float TanLightAngle,
+	uint NumIntersectingObjects, 
+	uint CulledDataStart)
+{
+	// Keeps result from going all the way sharp
+	float MinSphereRadius = .4f;
+	// Maintain reasonable culling bounds
+	float MaxSphereRadius = 100;
 
-[numthreads(THREADGROUP_TOTALSIZE, 1, 1)]
-void VPLTileCullCS(
+	float MinRayTime = MaxRayTime;
+	float MinConeVisibility = 1;
+	float3 RayUnitDirection = normalize(WorldRayEnd - WorldRayStart);
+
+	LOOP
+	for (uint ListObjectIndex = 0; ListObjectIndex < NumIntersectingObjects && MinRayTime >= MaxRayTime; ListObjectIndex++)
+	{
+		#if USE_SHADOW_CULLING_FOR_VPL_PLACEMENT
+			uint ObjectIndex = ShadowTileArrayData.Load(ListObjectIndex * ShadowTileListGroupSize.x * ShadowTileListGroupSize.y + CulledDataStart);
+		#else
+			uint ObjectIndex = ListObjectIndex;
+		#endif
+
+		float4 SphereCenterAndRadius = LoadObjectPositionAndRadiusFromBuffer(ObjectIndex, ShadowCulledObjectBounds);
+		float ObjectCenterDistanceAlongRay = dot(SphereCenterAndRadius.xyz - WorldRayStart, RayUnitDirection);
+
+		BRANCH
+		if (ObjectCenterDistanceAlongRay > -SphereCenterAndRadius.w)
+		{
+			float3 LocalPositionExtent = LoadObjectLocalPositionExtentFromBuffer(ObjectIndex, ShadowCulledObjectData);
+			float4x4 WorldToVolume = LoadObjectWorldToVolumeFromBuffer(ObjectIndex, ShadowCulledObjectData);
+			bool bGeneratedAsTwoSided;
+			float4 UVScaleAndVolumeScale = LoadObjectUVScaleFromBuffer(ObjectIndex, ShadowCulledObjectData, bGeneratedAsTwoSided);
+			float3 UVAdd = LoadObjectUVAddFromBuffer(ObjectIndex, ShadowCulledObjectData);
+
+			float3 VolumeRayStart = mul(float4(WorldRayStart, 1), WorldToVolume).xyz;
+			float3 VolumeRayEnd = mul(float4(WorldRayEnd, 1), WorldToVolume).xyz;
+			float3 VolumeRayDirection = VolumeRayEnd - VolumeRayStart;
+			float VolumeRayLength = length(VolumeRayDirection);
+			VolumeRayDirection /= VolumeRayLength;
+			float VolumeMinSphereRadius = MinSphereRadius / UVScaleAndVolumeScale.w;
+			float VolumeMaxSphereRadius = MaxSphereRadius / UVScaleAndVolumeScale.w;
+
+			// Expand the intersection box by the radius of the cone at the distance of the object along the cone
+			float LocalConeRadiusAtObject = min(TanLightAngle * max(ObjectCenterDistanceAlongRay, 0) / UVScaleAndVolumeScale.w, VolumeMaxSphereRadius);
+
+			float2 IntersectionTimes = LineBoxIntersect(VolumeRayStart, VolumeRayEnd, -LocalPositionExtent - LocalConeRadiusAtObject, LocalPositionExtent + LocalConeRadiusAtObject);
+
+			BRANCH
+			if (IntersectionTimes.x < IntersectionTimes.y && IntersectionTimes.x < 1)
+			{
+				float SampleRayTime = IntersectionTimes.x * VolumeRayLength;
+				uint MaxSteps = 64;
+				float MinStepSize = 1.0f / (4 * MaxSteps);
+
+				float MinDistance = 1000000;
+				float3 IntersectionPosition = float3(0, 0, 0);
+
+				uint StepIndex = 0;
+
+				LOOP
+				for (; StepIndex < MaxSteps; StepIndex++)
+				{
+					float3 SampleVolumePosition = VolumeRayStart + VolumeRayDirection * SampleRayTime;
+					float3 ClampedSamplePosition = clamp(SampleVolumePosition, -LocalPositionExtent, LocalPositionExtent);
+					float DistanceToClamped = length(ClampedSamplePosition - SampleVolumePosition);
+					float3 VolumeUV = DistanceFieldVolumePositionToUV(ClampedSamplePosition, UVScaleAndVolumeScale.xyz, UVAdd);
+					float DistanceField = Texture3DSampleLevel(DistanceFieldTexture, DistanceFieldSampler, VolumeUV, 0).x + DistanceToClamped;
+
+					MinDistance = min(MinDistance, DistanceField);
+					float SphereRadius = clamp(TanLightAngle * SampleRayTime, VolumeMinSphereRadius, VolumeMaxSphereRadius);
+
+					MinConeVisibility = min(MinConeVisibility, saturate(DistanceField / SphereRadius));
+					IntersectionPosition = SampleVolumePosition;
+
+					float StepDistance = max(DistanceField, MinStepSize);
+					SampleRayTime += StepDistance;
+
+					// Terminate the trace if we reached a negative area or went past the end of the ray
+					if (DistanceField < 0 
+						|| SampleRayTime > IntersectionTimes.y * VolumeRayLength)
+					{
+						break;
+					}
+				}
+
+				if (MinDistance < 0 || StepIndex == MaxSteps)
+				{
+					MinConeVisibility = 0;
+					MinRayTime = min(MinRayTime, SampleRayTime * UVScaleAndVolumeScale.w);
+				}
+
+				// Force to shadowed as we approach max steps
+				MinConeVisibility = min(MinConeVisibility, (1 - StepIndex / (float)MaxSteps));
+			}
+		}
+	}
+
+	return MinConeVisibility;
+}
+
+float3 ComputeVPLFlux(uint SurfelIndex, float4x4 InstanceToWorld)
+{
+	// Distance for directional lights to trace
+	float TraceDistance = 10000;
+	float4 PositionAndRadius = LoadSurfelPositionAndRadius(SurfelIndex);
+	PositionAndRadius.xyz = mul(float4(PositionAndRadius.xyz, 1), InstanceToWorld).xyz;
+	float3 SurfelNormal = LoadSurfelNormal(SurfelIndex);
+	SurfelNormal = mul(SurfelNormal, (float3x3)InstanceToWorld);
+
+	float3 AccumulatedLighting = 0;
+
+	{
+		float SurfelDotLight = dot(SurfelNormal, LightDirection);
+
+		BRANCH
+		if (SurfelDotLight > 0)
+		{
+			float Visibility = 1;
+			bool bIsPointLight = false;
+
+			BRANCH
+			if (bIsPointLight || SurfelDotLight > TanLightAngleAndNormalThreshold.y)
+			{
+				// World space offset along the start of the ray to avoid incorrect self-shadowing
+				float RayStartOffset = 2;
+
+				float3 WorldRayStart;
+				float3 WorldRayEnd;
+				float MaxRayTime;
+				float TanLightAngle;
+				uint NumIntersectingObjects = GetCulledNumObjects();
+				uint CulledDataStart = 0;
+
+				if (bIsPointLight)
+				{
+					/*
+					float3 LightVector = LightPositionAndInvRadius.xyz - PositionAndRadius.xyz;
+					float LightVectorLength = length(LightVector);
+					WorldRayStart = PositionAndRadius.xyz + LightVector / LightVectorLength * RayStartOffset;
+					WorldRayEnd = LightPositionAndInvRadius.xyz;
+					MaxRayTime = LightVectorLength;
+					float MaxAngle = tan(10 * PI / 180.0f);
+					// Comparing tangents instead of angles, but tangent is always increasing in this range
+					TanLightAngle = min(LightSourceRadius / LightVectorLength, MaxAngle);
+					*/
+				}
+				else
+				{
+					WorldRayStart = PositionAndRadius.xyz + LightDirection * RayStartOffset;
+					WorldRayEnd = PositionAndRadius.xyz + LightDirection * TraceDistance;
+					MaxRayTime = TraceDistance;
+					TanLightAngle = TanLightAngleAndNormalThreshold.x;
+
+					#if USE_SHADOW_CULLING_FOR_VPL_PLACEMENT
+
+						// Transform into shadow space
+						float4 HomogeneousShadowPosition = mul(float4(WorldRayStart, 1), WorldToShadow);
+						float2 NormalizedShadowPosition = HomogeneousShadowPosition.xy * .5f + .5f;
+						NormalizedShadowPosition.y = 1 - NormalizedShadowPosition.y;
+						// Quantize the shadow position to get our tile position
+						uint2 TilePosition = (uint2)(NormalizedShadowPosition * ShadowTileListGroupSize);
+						// Fetch the tile head information
+						uint2 TileHead = GetShadowTileHead(TilePosition);
+						NumIntersectingObjects = TileHead.y;
+						CulledDataStart = TileHead.x;
+
+					#endif
+				}
+
+				Visibility = RayTraceThroughLightTileCulledObjects(WorldRayStart, WorldRayEnd, MaxRayTime, TanLightAngle, NumIntersectingObjects, CulledDataStart);
+			}
+
+			AccumulatedLighting += (max(SurfelDotLight, 0) * Visibility) * LightColor.rgb;
+		}
+	}
+
+	float3 Flux = 0;
+
+	BRANCH
+	if (any(AccumulatedLighting > 0))
+	{
+		float3 DiffuseColor = LoadSurfelDiffuseColor(SurfelIndex);
+		float3 EmissiveColor = LoadSurfelDiffuseColor(SurfelIndex);
+
+		Flux = (DiffuseColor * AccumulatedLighting + EmissiveColor) * PI * PositionAndRadius.w * PositionAndRadius.w;
+	}
+
+	return Flux;
+}
+
+uint GetViewBasedSurfelLOD(float3 ObjectPosition)
+{
+	return length(ObjectPosition - View.ViewOrigin.xyz) > .4f * AOMaxViewDistance ? 1 : 0;
+	//return 0;
+}
+
+uint2 GetSurfelOffsetAndNum(uint4 SurfelCoordinate, uint LODIndex, uniform bool bInstanced)
+{
+	uint2 OffsetAndNum = uint2(bInstanced ? SurfelCoordinate.w : SurfelCoordinate.x, SurfelCoordinate.y);
+	
+	if (LODIndex == 1)
+	{
+		OffsetAndNum.x += SurfelCoordinate.y;
+		OffsetAndNum.y = SurfelCoordinate.z - SurfelCoordinate.y;
+	}
+
+	return OffsetAndNum;
+}
+
+uint ObjectProcessStartIndex;
+RWBuffer<float4> RWVPLFlux;
+
+[numthreads(LIGHT_VPLS_THREADGROUP_SIZE, 1, 1)]
+void LightVPLsCS(
 	uint3 GroupId : SV_GroupID,
 	uint3 DispatchThreadId : SV_DispatchThreadID,
     uint3 GroupThreadId : SV_GroupThreadID) 
 {
-    uint ThreadIndex = GroupThreadId.x;
-	uint NumClusterVPLs = VPLParameterBuffer[0];
-	uint TileIndex = GroupId.y * NumGroups.x + GroupId.x;
+	uint ObjectIndex = ObjectProcessStartIndex + GroupId.x * ObjectProcessStride;
+	uint NumObjects = GetCulledNumObjects();
+	uint ThreadIndex = GroupThreadId.x;
 
-	float4 ConeAxisAndCos = TileConeAxisAndCos.Load(TileIndex);
-	float4 ConeAxisDepthRanges = TileConeDepthRanges.Load(TileIndex);
-	float3 TileConeVertex = 0;
-	float3 TileConeAxis = ConeAxisAndCos.xyz;
-	float TileConeAngleCos = ConeAxisAndCos.w;
-	float TileConeAngleSin = sqrt(1 - TileConeAngleCos * TileConeAngleCos);
-
-	LOOP
-	for (uint VPLClusterIndex = ThreadIndex; VPLClusterIndex < NumClusterVPLs; VPLClusterIndex += THREADGROUP_TOTALSIZE)
+	if (ObjectIndex < NumObjects)
 	{
-		float4 BoundingSphere = VPLClusterData[VPLClusterIndex * VPL_DATA_STRIDE + 0];
-		float3 ViewSpaceSphereCenter = mul(float4(BoundingSphere.xyz + View.PreViewTranslation.xyz, 1), View.TranslatedWorldToView).xyz;
+		uint4 SurfelCoordinate = LoadObjectSurfelCoordinate(ObjectIndex);
+		float4 ObjectPositionAndRadius = LoadObjectPositionAndRadius(ObjectIndex);
+		float4x4 ObjectInstanceToWorld = LoadObjectLocalToWorld(ObjectIndex);
+		uint ViewBasedLOD = GetViewBasedSurfelLOD(ObjectPositionAndRadius.xyz);
 		
-		BRANCH
-		if (SphereIntersectConeWithDepthRanges(float4(ViewSpaceSphereCenter, AOMaxDistance), TileConeVertex, TileConeAxis, TileConeAngleCos, TileConeAngleSin, ConeAxisDepthRanges))
 		{
-			uint CulledVPLClusterIndex;
-			InterlockedAdd(RWTileHeadDataUnpacked[TileIndex * 2 + 1], 1, CulledVPLClusterIndex);
+			uint2 SurfelOffsetAndNum = GetSurfelOffsetAndNum(SurfelCoordinate, ViewBasedLOD, false);
+			uint2 InstancedSurfelOffsetAndNum = GetSurfelOffsetAndNum(SurfelCoordinate, ViewBasedLOD, true);
+			uint NumSurfels = SurfelOffsetAndNum.y;
 
-			if (CulledVPLClusterIndex < MAX_VPLS_PER_TILE)
+			LOOP
+			for (uint SurfelIndex = ThreadIndex; SurfelIndex < NumSurfels; SurfelIndex += LIGHT_VPLS_THREADGROUP_SIZE)
 			{
-				RWTileArrayData[TileIndex * MAX_VPLS_PER_TILE + CulledVPLClusterIndex] = VPLClusterIndex;
+				// Read surfel properties from a shared location for all instances
+				float3 Flux = ComputeVPLFlux(SurfelIndex + SurfelOffsetAndNum.x, ObjectInstanceToWorld);
+				// Write flux to the instance-specific location
+				RWVPLFlux[InstancedSurfelOffsetAndNum.x + SurfelIndex] = float4(Flux, 0);
 			}
 		}
+		/*
+		// Still need to light LOD1
+		if (ViewBasedLOD == 0)
+		{
+			uint2 SurfelOffsetAndNum = GetSurfelOffsetAndNum(SurfelCoordinate, 1);
+			uint NumSurfels = SurfelOffsetAndNum.y;
+
+			LOOP
+			for (uint SurfelIndex = ThreadIndex; SurfelIndex < NumSurfels; SurfelIndex += LIGHT_VPLS_THREADGROUP_SIZE)
+			{
+				float3 Flux = ComputeVPLFlux(SurfelIndex + SurfelOffsetAndNum.x);
+				RWVPLFlux[SurfelOffsetAndNum.x + SurfelIndex] = float4(Flux, 0);
+			}
+		}*/
 	}
 }
 
@@ -497,7 +602,6 @@ float RayTraceThroughGlobalObjects(
 	float3 WorldRayEnd = RayStartPosition + RayDirection * RayLength;
 	float MaxRayTime = RayLength;
 
-	float MinRayTime = MaxRayTime;
 	float MinVisibility = 1;
 
 	LOOP
@@ -605,180 +709,277 @@ float GetVPLOcclusion(float3 BentNormalAO, float3 NormalizedVectorToVPL)
 }
 
 Buffer<float> RecordConeData;
+RWBuffer<float4> RWStepBentNormal;
+float BentNormalNormalizeFactor;
+
+/**  */
+[numthreads(FINAL_GATHER_THREADGROUP_SIZE, 1, 1)]
+void ComputeStepBentNormalCS(
+	uint3 GroupId : SV_GroupID,
+	uint3 DispatchThreadId : SV_DispatchThreadID,
+    uint3 GroupThreadId : SV_GroupThreadID) 
+{
+	uint StartIndex = SavedStartIndex[0];
+	uint NumRecords = ScatterDrawParameters[1];
+
+	uint RecordIndex = StartIndex + DispatchThreadId.x;
+
+	float3 Irradiance = 0;
+
+	if (RecordIndex < NumRecords)
+	{
+		uint RelativeRecordIndex = DispatchThreadId.x;
+		float3 WorldNormal = IrradianceCacheNormal[RecordIndex].xyz;
+
+		float3 TangentX;
+		float3 TangentY;
+		FindBestAxisVectors2(WorldNormal, TangentX, TangentY);
+
+		for (uint StepIndex = 0; StepIndex < NUM_VISIBILITY_STEPS; StepIndex++)
+		{
+			float3 UnoccludedDirection = 0;
+
+			for (uint ConeIndex = 0; ConeIndex < NUM_CONE_DIRECTIONS; ConeIndex++)
+			{
+				float3 ConeDirection = AOSamples2.SampleDirections[ConeIndex].xyz;
+				float3 RotatedConeDirection = ConeDirection.x * TangentX + ConeDirection.y * TangentY + ConeDirection.z * WorldNormal;
+
+				uint RecordConeDataIndex = (RelativeRecordIndex * NUM_CONE_DIRECTIONS + ConeIndex) * RECORD_CONE_DATA_STRIDE;
+				float ConeVisibility = RecordConeData[RecordConeDataIndex + StepIndex];
+				UnoccludedDirection += ConeVisibility * RotatedConeDirection;
+			}
+
+			float InvNumSamples = 1.0f / (float)NUM_CONE_DIRECTIONS;
+			UnoccludedDirection = UnoccludedDirection * (BentNormalNormalizeFactor * InvNumSamples);
+
+			RWStepBentNormal[RelativeRecordIndex * NUM_VISIBILITY_STEPS + StepIndex] = float4(UnoccludedDirection, 0);
+		}
+	}
+}
 
 float VPLGatherRadius;
-float ConeHalfAngle;
-float BentNormalNormalizeFactor;
 
-static float3 StepBentNormal[NUM_VISIBILITY_STEPS];
+Buffer<float4> StepBentNormalBuffer;
+
+float4 LoadVPLPositionAndRadius(uint VPLIndex, float4x4 InstanceToWorld)
+{
+#if IRRADIANCE_FROM_SURFELS
+	float4 PositionAndRadius = LoadSurfelPositionAndRadius(VPLIndex);
+	PositionAndRadius.xyz = mul(float4(PositionAndRadius.xyz, 1), InstanceToWorld).xyz;
+	return PositionAndRadius;
+#else
+	return VPLData[VPLIndex * VPL_DATA_STRIDE + 0];
+#endif
+}
+
+float3 LoadVPLNormal(uint VPLIndex, float4x4 InstanceToWorld)
+{
+#if IRRADIANCE_FROM_SURFELS
+	float3 SurfelNormal = LoadSurfelNormal(VPLIndex);
+	SurfelNormal = mul(SurfelNormal, (float3x3)InstanceToWorld);
+	return SurfelNormal;
+#else
+	return VPLData[VPLIndex * VPL_DATA_STRIDE + 1].xyz;
+#endif
+}
+
+Buffer<float4> VPLFlux;
 
-float3 ComputeVirtualPointLighting(uniform bool bCluster, uint VPLIndex, float3 WorldPosition, float3 WorldNormal)
+float3 LoadVPLFlux(uint VPLIndex)
 {
+#if IRRADIANCE_FROM_SURFELS
+	return VPLFlux[VPLIndex].xyz;
+#else
+	return VPLData[VPLIndex * VPL_DATA_STRIDE + 2].xyz;
+#endif
+}
+
+float3 ComputeVirtualPointLighting(uint VPLIndex, uint InstancedVPLIndex, uint RelativeRecordIndex, float3 WorldPosition, float3 WorldNormal, float4x4 InstanceToWorld)
+{ 
 	float3 Irradiance = 0;
 
 	float VisibilityStepSize = NUM_VISIBILITY_STEPS / AOMaxDistance;
 	float MaxGatherDistanceSq = VPLGatherRadius * VPLGatherRadius;
-	float4 VPLPositionAndRadius = bCluster ? VPLClusterData[VPLIndex * VPL_DATA_STRIDE + 0] : VPLData[VPLIndex * VPL_DATA_STRIDE + 0];
+	float4 VPLPositionAndRadius = LoadVPLPositionAndRadius(VPLIndex, InstanceToWorld);
+
 	float3 VPLToGatherPoint = WorldPosition - VPLPositionAndRadius.xyz;
 	float DistanceSq = dot(VPLToGatherPoint, VPLToGatherPoint);
 	float DirectionDot = dot(-VPLToGatherPoint, WorldNormal);
 
 	// Hack
-	//VPLPositionAndRadius.w = 100;
+	//VPLPositionAndRadius.w = 10;
 
 #define VISUALIZE_VPL_PLACEMENT 0
 #if VISUALIZE_VPL_PLACEMENT
-	Irradiance.r += 3 * (DistanceSq < VPLPositionAndRadius.w * VPLPositionAndRadius.w);
+	//Irradiance += float3(.4f, .2f, .2f) * .1f * (DistanceSq < VPLPositionAndRadius.w * VPLPositionAndRadius.w);
+
+	//float3 DebugValue = LoadVPLFlux(VPLIndex) / 10000;
+	float3 DebugValue = LoadSurfelDiffuseColor(VPLIndex);
+	Irradiance += DebugValue * .1f * (DistanceSq < VPLPositionAndRadius.w * VPLPositionAndRadius.w);
+	//Irradiance += .00001f;
 #endif
 
+#define COMPUTE_VPL_LIGHTING 1
+#define VISUALIZE_VPL_SCENE 0
+
+#if COMPUTE_VPL_LIGHTING
 	BRANCH
 	if (DistanceSq < MaxGatherDistanceSq && DirectionDot > 0)
 	{
-		float3 VPLNormal = bCluster ? VPLClusterData[VPLIndex * VPL_DATA_STRIDE + 1].xyz : VPLData[VPLIndex * VPL_DATA_STRIDE + 1].xyz;
+		float3 VPLNormal = LoadVPLNormal(VPLIndex, InstanceToWorld);
 		float VPLNormalDot = dot(VPLNormal, VPLToGatherPoint);
 
 		BRANCH
 		if (VPLNormalDot > 0)
 		{
-			float Distance = sqrt(DistanceSq);
-			float3 VPLDirection = -VPLToGatherPoint / Distance;
-
-			#define USE_INVERSE_SQUARED_DISK_APPROX 1
-			#define USE_PROJECTED_ANGLE_ATTENUATION 0
-
-			#if USE_INVERSE_SQUARED_DISK_APPROX
-				float EffectiveDiskRadius = VPLPositionAndRadius.w * 1;
-				float DistanceAttenuation = 1.0f / (PI * DistanceSq + PI * EffectiveDiskRadius * EffectiveDiskRadius);
-				//@todo - this is scale dependent
-				DistanceAttenuation = min(DistanceAttenuation, .0001f * RadialAttenuation(VPLToGatherPoint / VPLGatherRadius, 4));
-			#elif USE_PROJECTED_ANGLE_ATTENUATION
-				float TangentProjectedAngle = VPLPositionAndRadius.w / Distance;
-				float CosProjectedAngle = 1 / sqrt(1 + TangentProjectedAngle * TangentProjectedAngle);
-				float ProjectedSolidAngle = 2 * PI * (1 - CosProjectedAngle);
-				float DistanceAttenuation = ProjectedSolidAngle / (2 * PI);
-				DistanceAttenuation = min(DistanceAttenuation, RadialAttenuation(VPLToGatherPoint / VPLGatherRadius, 4));
-				DistanceAttenuation *= .0001f;
-			#else
-				float DistanceAttenuation = RadialAttenuation(VPLToGatherPoint / VPLGatherRadius, 8);
-				DistanceAttenuation *= .0001f;
-			#endif
-
-			float CosTheta = DirectionDot / Distance;
-			float SinTheta = sqrt(1 - CosTheta * CosTheta);
-			
-			#if IRRADIANCE_FROM_AO_CONES
+			float3 VPLFlux = LoadVPLFlux(InstancedVPLIndex);
 
-				float ShadowDepthBias = 0;
-				float ShadowingDistance = Distance + ShadowDepthBias;
-				float NormalizedDistance = saturate(ShadowingDistance / AOMaxDistance);
-				uint LowerStepIndex = (uint)min(NormalizedDistance * NUM_VISIBILITY_STEPS, NUM_VISIBILITY_STEPS - 1);
-				float LerpAlpha = ShadowingDistance - LowerStepIndex * VisibilityStepSize;
-				float3 InterpolatedBentNormal = lerp(StepBentNormal[LowerStepIndex], StepBentNormal[LowerStepIndex + 1], saturate(LerpAlpha));
+			BRANCH
+			if (any(VPLFlux > .01f))
+			{
+				float Distance = sqrt(DistanceSq);
+				float3 VPLDirection = -VPLToGatherPoint / Distance;
+
+				#define USE_INVERSE_SQUARED_DISK_APPROX 1
+
+				#if USE_INVERSE_SQUARED_DISK_APPROX
+					float EffectiveDiskRadius = VPLPositionAndRadius.w * 1;
+					float DistanceAttenuation = 1.0f / (PI * DistanceSq + PI * EffectiveDiskRadius * EffectiveDiskRadius);
+					//@todo - this is scale dependent
+					DistanceAttenuation = min(DistanceAttenuation, .0001f * RadialAttenuation(VPLToGatherPoint / VPLGatherRadius, 4));
+				#else
+					float DistanceAttenuation = RadialAttenuation(VPLToGatherPoint / VPLGatherRadius, 8);
+					DistanceAttenuation *= .0001f;
+				#endif
 
-				float Shadow = GetVPLOcclusion(InterpolatedBentNormal, VPLDirection);
+				float CosTheta = DirectionDot / Distance;
+				float SinTheta = sqrt(1 - CosTheta * CosTheta);
+			
+				#define IRRADIANCE_FROM_AO_CONES 1
+				#if IRRADIANCE_FROM_AO_CONES
+
+					float ShadowDepthBias = 0;
+					float ShadowingDistance = Distance + ShadowDepthBias;
+					float NormalizedDistance = saturate(ShadowingDistance / AOMaxDistance);
+					uint LowerStepIndex = (uint)min(NormalizedDistance * NUM_VISIBILITY_STEPS, NUM_VISIBILITY_STEPS - 1);
+					float LerpAlpha = ShadowingDistance - LowerStepIndex * VisibilityStepSize;
+
+					float3 InterpolatedBentNormal = lerp(
+						StepBentNormalBuffer[RelativeRecordIndex * NUM_VISIBILITY_STEPS + LowerStepIndex].xyz, 
+						StepBentNormalBuffer[RelativeRecordIndex * NUM_VISIBILITY_STEPS + LowerStepIndex + 1].xyz, 
+						saturate(LerpAlpha));
+					
+					float Shadow = GetVPLOcclusion(InterpolatedBentNormal, VPLDirection);
+
+				#else
+
+					float StartOffset = 1;
+					float EndOffset = 10;
+					float RayLength = max(Distance - StartOffset - EndOffset, 0);
+					float ConeEndDistance = max(RayLength - SinTheta * VPLPositionAndRadius.w, 0);
+					float RadiusAtConeEnd = CosTheta * VPLPositionAndRadius.w;
+					// Clamp the cone angle so that it doesn't intersect the gather point surface
+					float TanConeAngle = min(RadiusAtConeEnd / ConeEndDistance, SinTheta / CosTheta);
+					float Shadow = RayTraceThroughGlobalObjects(WorldPosition + StartOffset * VPLDirection, VPLDirection, RayLength, TanConeAngle, ConeEndDistance);
+
+					if (ConeEndDistance == 0)
+					{
+						//Shadow = 10;
+					}
 
-			#else
+				#endif
 
-				float StartOffset = 1;
-				float EndOffset = 10;
-				float RayLength = max(Distance - StartOffset - EndOffset, 0);
-				float ConeEndDistance = max(RayLength - SinTheta * VPLPositionAndRadius.w, 0);
-				float RadiusAtConeEnd = CosTheta * VPLPositionAndRadius.w;
-				// Clamp the cone angle so that it doesn't intersect the gather point surface
-				float TanConeAngle = min(RadiusAtConeEnd / ConeEndDistance, SinTheta / CosTheta);
-				float Shadow = RayTraceThroughGlobalObjects(WorldPosition + StartOffset * VPLDirection, VPLDirection, RayLength, TanConeAngle, ConeEndDistance);
+				float VPLCosineLobe = saturate(VPLNormalDot / Distance);
 
-				if (ConeEndDistance == 0)
-				{
-					//Shadow = 10;
-				}
+				Irradiance += (saturate(CosTheta) * VPLCosineLobe * DistanceAttenuation * Shadow) * VPLFlux;
+			}
+		}
+	}
+#elif VISUALIZE_VPL_SCENE
 
-			#endif
+	float DistanceWeight = (1 - saturate(DistanceSq / (VPLPositionAndRadius.w * VPLPositionAndRadius.w)));
+	float3 VPLNormal = LoadVPLNormal(VPLIndex, InstanceToWorld);
+	float DistanceBehindVPL = dot((WorldPosition - VPLPositionAndRadius.xyz), -VPLNormal);
+	float DistanceBehindMask = 1 - saturate(DistanceBehindVPL / (.5f * VPLPositionAndRadius.w));
 
-			float VPLCosineLobe = saturate(VPLNormalDot / Distance);
-			float3 VPLFlux = bCluster ? VPLClusterData[VPLIndex * VPL_DATA_STRIDE + 2].xyz : VPLData[VPLIndex * VPL_DATA_STRIDE + 2].xyz;
+	float EffectiveDiskRadius = VPLPositionAndRadius.w * 1;
+	float DistanceAttenuation = 1.0f / (PI * DistanceSq + PI * EffectiveDiskRadius * EffectiveDiskRadius);
+	float3 VPLFlux = LoadVPLFlux(InstancedVPLIndex);
+	float NormalMask = saturate(dot(VPLNormal, WorldNormal));
 
-			Irradiance += (saturate(CosTheta) * VPLCosineLobe * DistanceAttenuation * Shadow) * VPLFlux;
-		}
-	}
+	Irradiance += DistanceWeight * VPLFlux * DistanceAttenuation * DistanceBehindMask * NormalMask;
+#endif
 
 	return Irradiance;
 }
 
 float3 GatherIrradianceFromVPLs(float3 WorldPosition, float3 WorldNormal, uint RelativeRecordIndex, uint RecordIndex, uint ThreadIndex)
 {
-	float3 TangentX;
-	float3 TangentY;
-	FindBestAxisVectors2(WorldNormal, TangentX, TangentY);
-
-	for (uint StepIndex = 0; StepIndex < NUM_VISIBILITY_STEPS; StepIndex++)
-	{
-		float3 UnoccludedDirection = 0;
-
-		for (uint ConeIndex = 0; ConeIndex < NUM_CONE_DIRECTIONS; ConeIndex++)
-		{
-			float3 ConeDirection = AOSamples2.SampleDirections[ConeIndex].xyz;
-			float3 RotatedConeDirection = ConeDirection.x * TangentX + ConeDirection.y * TangentY + ConeDirection.z * WorldNormal;
-
-			uint RecordConeDataIndex = (RelativeRecordIndex * NUM_CONE_DIRECTIONS + ConeIndex) * RECORD_CONE_DATA_STRIDE;
-			float ConeVisibility = RecordConeData[RecordConeDataIndex + StepIndex];
-			UnoccludedDirection += ConeVisibility * RotatedConeDirection;
-		}
-
-		float InvNumSamples = 1.0f / (float)NUM_CONE_DIRECTIONS;
-		UnoccludedDirection = UnoccludedDirection * (BentNormalNormalizeFactor * InvNumSamples);
-		StepBentNormal[StepIndex] = UnoccludedDirection;
-	}
-
 	float3 Irradiance = 0;
 
-#define USE_VPL_CLUSTERS 0
-#if USE_VPL_CLUSTERS
+#if IRRADIANCE_FROM_SURFELS
 
-	#define USE_TILE_CULLING 0
-	#if USE_TILE_CULLING
+	uint2 TileCoordinate = IrradianceCacheTileCoordinate[RecordIndex];
+	uint4 TileHead = GetTileHead(TileCoordinate);
+	uint NumObjectsAffectingTile = TileHead.y + TileHead.z + TileHead.w;
 
-		//@todo - won't be correct across frames
-		uint2 TileCoordinate = IrradianceCacheTileCoordinate[RecordIndex];
-		uint TileIndex = TileCoordinate.y * TileListGroupSize.x + TileCoordinate.x;
-		uint NumClusterVPLs = min(TileHeadDataUnpacked[TileIndex * 2 + 1], MAX_VPLS_PER_TILE);
+	LOOP 
+	for (uint ListObjectIndex = 0; ListObjectIndex < NumObjectsAffectingTile; ListObjectIndex++)
+	{
+		uint ListIndex = 0;
+		uint ArrayIndex = ListObjectIndex;
 
-		LOOP
-		for (uint VPLListIndex = ThreadIndex; VPLListIndex < NumClusterVPLs; VPLListIndex += FINAL_GATHER_THREADGROUP_SIZE)
+		FLATTEN
+		if (ListObjectIndex >= TileHead.y + TileHead.z)
 		{
-			uint VPLClusterIndex = TileArrayData[TileIndex * MAX_VPLS_PER_TILE + VPLListIndex];
-	#else
+			ListIndex = 2;
+			ArrayIndex = ListObjectIndex - TileHead.y - TileHead.z;
+		}
+		else if (ListObjectIndex >= TileHead.y)
+		{
+			ListIndex = 1;
+			ArrayIndex = ListObjectIndex - TileHead.y;
+		}
 
-		uint NumClusterVPLs = VPLParameterBuffer[0];
+		uint ObjectIndex = TileArrayData.Load((ArrayIndex * TileListGroupSize.x * TileListGroupSize.y + TileHead.x) * NUM_CULLED_OBJECT_LISTS + ListIndex);
+		float4 ObjectPositionAndRadius = LoadObjectPositionAndRadius(ObjectIndex);
+		float ObjectDistance = length(ObjectPositionAndRadius.xyz - WorldPosition);
+		float BoundingRadius = ObjectPositionAndRadius.w + VPLGatherRadius;
 
-		LOOP
-		for (uint VPLClusterIndex = ThreadIndex; VPLClusterIndex < NumClusterVPLs; VPLClusterIndex += FINAL_GATHER_THREADGROUP_SIZE)
+		BRANCH
+		if (ObjectDistance < BoundingRadius)
 		{
+			//float ObjectDistance = length(ObjectPositionAndRadius.xyz - WorldPosition) - ObjectPositionAndRadius.w;
+			//Irradiance += .0001f * (ObjectDistance > VPLGatherRadius / 2);
 
-	#endif
-			float4 ClusterPositionAndRadius = VPLClusterData[VPLClusterIndex * VPL_DATA_STRIDE + 0];
-			float4 ClusterNormalAndNumVPLs = VPLClusterData[VPLClusterIndex * VPL_DATA_STRIDE + 1];
-			float4 ClusterFluxAndVPLStartIndex = VPLClusterData[VPLClusterIndex * VPL_DATA_STRIDE + 2];
+			uint4 SurfelCoordinate = LoadObjectSurfelCoordinate(ObjectIndex);
+			uint ViewBasedLOD = GetViewBasedSurfelLOD(ObjectPositionAndRadius.xyz);
+			//uint DistanceBasedLOD = ObjectDistance - ObjectPositionAndRadius.w > .5f * VPLGatherRadius ? 1 : 0;
+			//uint FinalLOD = max(ViewBasedLOD, DistanceBasedLOD);
+			uint2 InstancedSurfelOffsetAndNum = GetSurfelOffsetAndNum(SurfelCoordinate, ViewBasedLOD, true);
+			uint2 SurfelOffsetAndNum = GetSurfelOffsetAndNum(SurfelCoordinate, ViewBasedLOD, false);
 
-			Irradiance += ComputeVirtualPointLighting(true, VPLClusterIndex, WorldPosition, WorldNormal);
-		/*
-			uint NumChildVPLs = ClusterNormalAndNumVPLs.w;
+			float4x4 ObjectInstanceToWorld = LoadObjectLocalToWorld(ObjectIndex);
 
 			LOOP
-			for (uint ChildIndex = 0; ChildIndex < NumChildVPLs; ChildIndex++)
+			for (uint VPLIndex = ThreadIndex; VPLIndex < SurfelOffsetAndNum.y; VPLIndex += FINAL_GATHER_THREADGROUP_SIZE)
 			{
-				Irradiance += ComputeVirtualPointLighting(false, (uint)ClusterFluxAndVPLStartIndex.w + ChildIndex, WorldPosition, WorldNormal);
-			}*/
+				Irradiance += ComputeVirtualPointLighting(VPLIndex + SurfelOffsetAndNum.x, VPLIndex + InstancedSurfelOffsetAndNum.x, RelativeRecordIndex, WorldPosition, WorldNormal, ObjectInstanceToWorld);
+			}
 		}
+	}
 
 #else
+
 	uint NumVPLs = VPLParameterBuffer[1];
+	float4x4 Dummy = 0;
 
 	LOOP
 	for (uint VPLIndex = ThreadIndex; VPLIndex < NumVPLs; VPLIndex += FINAL_GATHER_THREADGROUP_SIZE)
 	{
-		Irradiance += ComputeVirtualPointLighting(false, VPLIndex, WorldPosition, WorldNormal);
+		Irradiance += ComputeVirtualPointLighting(VPLIndex, VPLIndex, RelativeRecordIndex, WorldPosition, WorldNormal, Dummy);
 	}
+
 #endif
 
 	//Irradiance = NumClusterVPLs / (float)MAX_VPLS_PER_TILE;
@@ -796,28 +997,29 @@ void ComputeIrradianceCS(
 	uint3 GroupId : SV_GroupID,
 	uint3 DispatchThreadId : SV_DispatchThreadID,
     uint3 GroupThreadId : SV_GroupThreadID) 
-{
+{ 
 	uint StartIndex = SavedStartIndex[0];
 	uint NumRecords = ScatterDrawParameters[1];
 
 	uint RecordIndex = StartIndex + GroupId.x;
+	uint ThreadIndex = GroupThreadId.x;
+
+	float3 Irradiance = 0;
 
 	if (RecordIndex < NumRecords)
 	{
 		float3 RecordWorldNormal = IrradianceCacheNormal[RecordIndex].xyz;
 
-		float3 DiffuseColor = .5f;
-
 		float4 PositionAndRadius = IrradianceCachePositionRadius[RecordIndex];
 		float3 RecordWorldPosition = PositionAndRadius.xyz;
-		float3 Irradiance = GatherIrradianceFromVPLs(RecordWorldPosition, RecordWorldNormal, GroupId.x, RecordIndex, GroupThreadId.x) * DiffuseColor;
-
-		SharedThreadIrradiance[GroupThreadId.x] = Irradiance;
+		Irradiance = GatherIrradianceFromVPLs(RecordWorldPosition, RecordWorldNormal, GroupId.x, RecordIndex, ThreadIndex);
 	}
 
+	SharedThreadIrradiance[ThreadIndex] = Irradiance;
+
 	GroupMemoryBarrierWithGroupSync();
 
-	if (GroupThreadId.x == 0)
+	if (ThreadIndex == 0)
 	{
 		float3 Irradiance = 0;
 
diff --git a/Engine/Shaders/DistanceFieldLightingShared.usf b/Engine/Shaders/DistanceFieldLightingShared.usf
--- a/Engine/Shaders/DistanceFieldLightingShared.usf
+++ b/Engine/Shaders/DistanceFieldLightingShared.usf
@@ -4,6 +4,8 @@
 	DistanceFieldLightingShared.usf
 =============================================================================*/
 
+#include "SurfelMaterialShared.usf"
+
 #ifndef THREADGROUP_SIZEX
 #define THREADGROUP_SIZEX 1
 #endif   
@@ -45,7 +47,7 @@ uint GetCulledNumObjects()
 }
 
 // In float4's.  Must match equivalent C++ variables.
-#define OBJECT_DATA_STRIDE 15
+#define OBJECT_DATA_STRIDE 16
 
 uint NumSceneObjects;
 
@@ -96,50 +98,109 @@ float3 LoadGlobalObjectUVAdd(uint ObjectIndex)
 }
 
 // In float4's.  Must match equivalent C++ variables.
-#define CULLED_OBJECT_DATA_STRIDE 11
+#define CULLED_OBJECT_DATA_STRIDE 16
 #define CULLED_OBJECT_BOX_BOUNDS_STRIDE 5
 
+float4 LoadObjectPositionAndRadiusFromBuffer(uint ObjectIndex, Buffer<float4> InBuffer)
+{
+	return InBuffer.Load(ObjectIndex);
+}
+
+float4x4 LoadObjectWorldToVolumeFromBuffer(uint ObjectIndex, Buffer<float4> InBuffer)
+{
+	float4 M0 = InBuffer.Load(ObjectIndex * CULLED_OBJECT_DATA_STRIDE + 0);
+	float4 M1 = InBuffer.Load(ObjectIndex * CULLED_OBJECT_DATA_STRIDE + 1);
+	float4 M2 = InBuffer.Load(ObjectIndex * CULLED_OBJECT_DATA_STRIDE + 2);
+	float4 M3 = InBuffer.Load(ObjectIndex * CULLED_OBJECT_DATA_STRIDE + 3);
+
+	return float4x4(M0, M1, M2, M3);
+}
+ 
+float3 LoadObjectLocalPositionExtentFromBuffer(uint ObjectIndex, Buffer<float4> InBuffer)
+{
+	return InBuffer.Load(ObjectIndex * CULLED_OBJECT_DATA_STRIDE + 4).xyz;
+}
+
+float4 LoadObjectUVScaleFromBuffer(uint ObjectIndex, Buffer<float4> InBuffer, out bool bGeneratedAsTwoSided)
+{
+	float4 Value = InBuffer.Load(ObjectIndex * CULLED_OBJECT_DATA_STRIDE + 5).xyzw;
+	bGeneratedAsTwoSided = Value.w < 0;
+	return float4(Value.xyz, abs(Value.w));
+}
+
+float3 LoadObjectUVAddFromBuffer(uint ObjectIndex, Buffer<float4> InBuffer)
+{
+	return InBuffer.Load(ObjectIndex * CULLED_OBJECT_DATA_STRIDE + 6).xyz;
+}
+
+float3x3 LoadObjectVolumeToWorldFromBuffer(uint ObjectIndex, Buffer<float4> InBuffer)
+{
+	float3 M0 = InBuffer.Load(ObjectIndex * CULLED_OBJECT_DATA_STRIDE + 8).xyz;
+	float3 M1 = InBuffer.Load(ObjectIndex * CULLED_OBJECT_DATA_STRIDE + 9).xyz;
+	float3 M2 = InBuffer.Load(ObjectIndex * CULLED_OBJECT_DATA_STRIDE + 10).xyz;
+
+	return float3x3(M0, M1, M2);
+}
+
+float4x4 LoadObjectLocalToWorldFromBuffer(uint ObjectIndex, Buffer<float4> InBuffer)
+{
+	float4 M0 = InBuffer.Load(ObjectIndex * CULLED_OBJECT_DATA_STRIDE + 11);
+	float4 M1 = InBuffer.Load(ObjectIndex * CULLED_OBJECT_DATA_STRIDE + 12);
+	float4 M2 = InBuffer.Load(ObjectIndex * CULLED_OBJECT_DATA_STRIDE + 13);
+	float4 M3 = InBuffer.Load(ObjectIndex * CULLED_OBJECT_DATA_STRIDE + 14);
+
+	return float4x4(M0, M1, M2, M3);
+}
+
 Buffer<float4> CulledObjectBounds;
 Buffer<float4> CulledObjectData;
 Buffer<float4> CulledObjectBoxBounds;
 
 float4 LoadObjectPositionAndRadius(uint ObjectIndex)
 {
-	return CulledObjectBounds.Load(ObjectIndex);
+	return LoadObjectPositionAndRadiusFromBuffer(ObjectIndex, CulledObjectBounds);
 }
 
 float4x4 LoadObjectWorldToVolume(uint ObjectIndex)
 {
-	float4 M0 = CulledObjectData.Load(ObjectIndex * CULLED_OBJECT_DATA_STRIDE + 0);
-	float4 M1 = CulledObjectData.Load(ObjectIndex * CULLED_OBJECT_DATA_STRIDE + 1);
-	float4 M2 = CulledObjectData.Load(ObjectIndex * CULLED_OBJECT_DATA_STRIDE + 2);
-	float4 M3 = CulledObjectData.Load(ObjectIndex * CULLED_OBJECT_DATA_STRIDE + 3);
-
-	return float4x4(M0, M1, M2, M3);
+	return LoadObjectWorldToVolumeFromBuffer(ObjectIndex, CulledObjectData);
 }
  
 float3 LoadObjectLocalPositionExtent(uint ObjectIndex)
 {
-	return CulledObjectData.Load(ObjectIndex * CULLED_OBJECT_DATA_STRIDE + 4).xyz;
+	return LoadObjectLocalPositionExtentFromBuffer(ObjectIndex, CulledObjectData);
 }
 
 float4 LoadObjectUVScale(uint ObjectIndex)
 {
-	return CulledObjectData.Load(ObjectIndex * CULLED_OBJECT_DATA_STRIDE + 5).xyzw;
+	bool bUnused;
+	return LoadObjectUVScaleFromBuffer(ObjectIndex, CulledObjectData, bUnused);
+}
+
+float4 LoadObjectUVScale(uint ObjectIndex, out bool bGeneratedAsTwoSided)
+{
+	return LoadObjectUVScaleFromBuffer(ObjectIndex, CulledObjectData, bGeneratedAsTwoSided);
 }
 
 float3 LoadObjectUVAdd(uint ObjectIndex)
 {
-	return CulledObjectData.Load(ObjectIndex * CULLED_OBJECT_DATA_STRIDE + 6).xyz;
+	return LoadObjectUVAddFromBuffer(ObjectIndex, CulledObjectData);
 }
 
 float3x3 LoadObjectVolumeToWorld(uint ObjectIndex)
 {
-	float3 M0 = CulledObjectData.Load(ObjectIndex * CULLED_OBJECT_DATA_STRIDE + 8).xyz;
-	float3 M1 = CulledObjectData.Load(ObjectIndex * CULLED_OBJECT_DATA_STRIDE + 9).xyz;
-	float3 M2 = CulledObjectData.Load(ObjectIndex * CULLED_OBJECT_DATA_STRIDE + 10).xyz;
+	return LoadObjectVolumeToWorldFromBuffer(ObjectIndex, CulledObjectData);
+}
 
-	return float3x3(M0, M1, M2);
+float4x4 LoadObjectLocalToWorld(uint ObjectIndex)
+{
+	return LoadObjectLocalToWorldFromBuffer(ObjectIndex, CulledObjectData);
+}
+
+// x = Offset in global buffer, y = NumLOD0, z = NumSurfels (all LODs), W = instance index
+uint4 LoadObjectSurfelCoordinate(uint ObjectIndex)
+{
+	return CulledObjectData.Load(ObjectIndex * CULLED_OBJECT_DATA_STRIDE + 15);
 }
 
 void LoadObjectViewSpaceBox(uint ObjectIndex, out float3 ObjectViewSpaceMin, out float3 ObjectViewSpaceMax)
@@ -175,10 +236,12 @@ void FindBestAxisVectors2(float3 InZAxis, out float3 OutXAxis, out float3 OutYAx
 	OutYAxis = cross( InZAxis, OutXAxis );
 }
 
+// VPLs generated by raytracing from the light
 #define VPL_DATA_STRIDE 3
 
 #define FINAL_GATHER_THREADGROUP_SIZE 64
 
+// Must match C++
 #define NUM_VISIBILITY_STEPS 10
 
 // Must match C++
@@ -211,4 +274,27 @@ bool ViewFrustumIntersectSphere(float3 SphereOrigin, float SphereRadius)
 	}
 	
 	return true;
-}
\ No newline at end of file
+}
+
+Buffer<float4> SurfelData;
+
+float4 LoadSurfelPositionAndRadius(uint SurfelIndex)
+{
+	return SurfelData[SurfelIndex * SURFEL_DATA_STRIDE + 0];
+}
+
+float3 LoadSurfelNormal(uint SurfelIndex)
+{
+	return SurfelData[SurfelIndex * SURFEL_DATA_STRIDE + 1].xyz;
+}
+
+float3 LoadSurfelDiffuseColor(uint SurfelIndex)
+{
+	return SurfelData[SurfelIndex * SURFEL_DATA_STRIDE + 2].xyz;
+}
+
+float3 LoadSurfelEmissiveColor(uint SurfelIndex)
+{
+	return SurfelData[SurfelIndex * SURFEL_DATA_STRIDE + 3].xyz;
+}
+
diff --git a/Engine/Shaders/DistanceFieldSurfaceCacheLighting.usf b/Engine/Shaders/DistanceFieldSurfaceCacheLighting.usf
--- a/Engine/Shaders/DistanceFieldSurfaceCacheLighting.usf
+++ b/Engine/Shaders/DistanceFieldSurfaceCacheLighting.usf
@@ -90,12 +90,8 @@ void ObjectCullPS(
 
 		if (ArrayIndex < MAX_OBJECTS_PER_TILE)
 		{
-#if COHERENT_OBJECT_LIST_WRITES
+			// Note: indexing so that threads are writing to RWTileArrayData coherently, has a huge impact on speed, even though the array data for one record is no longer dense
 			uint DataIndex = (ArrayIndex * (uint)(NumGroups.x * NumGroups.y + .5f) + TileIndex) * NUM_CULLED_OBJECT_LISTS + SmallestGroupIndex;
-#else
-			uint TileArrayStart = TileIndex * MAX_OBJECTS_PER_TILE;
-			uint DataIndex = TileArrayStart + ArrayIndex;
-#endif
 
 			RWTileArrayData[DataIndex] = Input.ObjectIndex;
 		}
@@ -156,7 +152,7 @@ struct FIrradianceCacheSplatVertexOutput
 {
 	float4 Position : SV_POSITION;
 	nointerpolation float4 PositionRadius : TEXCOORD0;
-	nointerpolation float3 Normal : TEXCOORD1;
+	nointerpolation float4 NormalAndFade : TEXCOORD1;
 #if FINAL_INTERPOLATION_PASS
 	nointerpolation float3 BentNormal : TEXCOORD2;
 	#if SUPPORT_IRRADIANCE
@@ -222,7 +218,7 @@ void IrradianceCacheSplatVS(
 	VirtualRadius += HackExpand * ViewToCenterLength * 4;
 #endif
 
-	float3 RecordNormal = IrradianceCacheNormal.Load(JobIndex).xyz;
+	float4 RecordNormalAndTime = IrradianceCacheNormal.Load(JobIndex);
 
 	float2 CornerUVs = InUV;
 	float3 CornerPosition = VirtualPosition + (View.ViewRight * CornerUVs.x + View.ViewUp * CornerUVs.y) * VirtualRadius;
@@ -237,7 +233,9 @@ void IrradianceCacheSplatVS(
 	Output.Position.xy += NormalizedOffsetToPixelCenter * Output.Position.w;
 
 	Output.PositionRadius = PositionAndRadius; 
-	Output.Normal = RecordNormal;
+
+	float Fade = 1;
+	Output.NormalAndFade = float4(RecordNormalAndTime.xyz, Fade);
 
 #if FINAL_INTERPOLATION_PASS
 	Output.BentNormal = IrradianceCacheBentNormal.Load(JobIndex).xyz;
@@ -293,21 +291,23 @@ void IrradianceCacheSplatPS(
 	if (DistanceError < 1 && SceneW < AOMaxViewDistance)
 	{ 
 		float3 WorldNormal = DistanceFieldNormalValue.xyz;
-		float NormalError = InterpolationAngleNormalization * sqrt(saturate(1 - dot(WorldNormal, Input.Normal)));
+		float3 RecordNormal = Input.NormalAndFade.xyz;
+		float NormalError = InterpolationAngleNormalization * sqrt(saturate(1 - dot(WorldNormal, RecordNormal)));
 
 		// Don't use a lighting record if it's in front of the query point.
 		// Query points behind the lighting record may have nearby occluders that the lighting record does not see.
 		// Offset the comparison point along the negative normal to prevent self-occlusion
-		float3 RecordToVertexVector = OpaqueWorldPosition - (Input.PositionRadius.xyz - 1 * Input.Normal.xyz);
+		float3 RecordToVertexVector = OpaqueWorldPosition - (Input.PositionRadius.xyz - 1 * RecordNormal.xyz);
 		float DistanceToVertex = length(RecordToVertexVector);
-		float PlaneDistance = dot(Input.Normal.xyz, RecordToVertexVector) / DistanceToVertex;
+		float PlaneDistance = dot(RecordNormal.xyz, RecordToVertexVector) / DistanceToVertex;
 		
 		// Setup an error metric that goes from 0 if the points are coplanar, to 1 if the point being shaded is at the angle corresponding to MinCosPointBehindPlane behind the plane
 		float PointBehindPlaneError = min(max(PlaneDistance * InvMinCosPointBehindPlane, 0.0f), DistanceToVertex / 3.0f);
 
 		float PrecisionScale = .1f;
-		Weight = saturate(PrecisionScale * (1 - max(DistanceError, max(NormalError, PointBehindPlaneError))));
-		//Weight = saturate(PrecisionScale * (1 - DistanceError));
+		Weight = saturate(PrecisionScale * Input.NormalAndFade.w * (1 - max(DistanceError, max(NormalError, PointBehindPlaneError))));
+		//Weight = saturate(PrecisionScale * Input.NormalAndFade.w * (1 - DistanceError));
+		//Weight = Input.NormalAndFade.w;
 
 		float VisualizePlacement = Distance < .1f * Input.PositionRadius.w;
 
@@ -344,6 +344,9 @@ SamplerState IrradianceSampler;
 
 Buffer<float4> DebugBuffer;
 
+// 1 / ((1 - FadeDistanceFraction) * AOMaxViewDistance)
+float DistanceFadeScale;
+
 /** Normalizes the splatted surface cache values, packs depth in alpha. */
 void AOCombinePS(
 	in float4 UVAndScreenPos : TEXCOORD0
@@ -395,7 +398,7 @@ void AOCombinePS(
 
 	float BentNormalLength = length(OutBentNormal.rgb);
 	// Fade to unoccluded in the distance
-	float FadeAlpha = saturate(.001f * (.95f * AOMaxViewDistance - SceneDepth));
+	float FadeAlpha = saturate((AOMaxViewDistance - SceneDepth) * DistanceFadeScale);
 	float3 NormalizedBentNormal = OutBentNormal.rgb / max(BentNormalLength, .0001f);
 	OutBentNormal.rgb = NormalizedBentNormal * lerp(1, BentNormalLength, FadeAlpha);
 
@@ -503,6 +506,7 @@ void FillGapsPS(
 
 	float4 Accumulation = 0;
 	float3 IrradianceAccumulation = 0;
+	float Length = 0;
 
 	for (float y = -HALF_FILL_KERNEL_SIZE; y <= HALF_FILL_KERNEL_SIZE; y++)
 	{
@@ -514,13 +518,16 @@ void FillGapsPS(
 			float SampleSceneDepth = abs(TextureValue.w);
 			float ValidMask = TextureValue.w > 0;
 
-			float DepthWeight = 1 / exp2(abs(SceneDepth - SampleSceneDepth) * .01f);
+			// Don't let depth weight go to 0 with huge depth differences
+			float DepthWeight = max(1 / exp2(abs(SceneDepth - SampleSceneDepth) * .01f), .001f);
 
 			float2 Weight2D = 1 / exp2(abs(float2(x, y) * 10.0f / HALF_FILL_KERNEL_SIZE));
 			float ScreenSpaceSpatialWeight = max(Weight2D.x, Weight2D.y);
 
 			float Weight = ValidMask * ScreenSpaceSpatialWeight * DepthWeight;
 
+			// Track length separately 
+			Length += length(TextureValue.rgb) * Weight;
 			Accumulation.rgb += TextureValue.rgb * Weight;
 			Accumulation.a += Weight;
 
@@ -532,7 +539,17 @@ void FillGapsPS(
 	}
 	
 	float InvSafeWeight = 1.0f / max(Accumulation.a, .00001f);
-	OutBentNormal = float4(Accumulation.rgb * InvSafeWeight, CenterValue.w);
+	float AverageLength = Length * InvSafeWeight;
+	float3 AverageBentNormal = Accumulation.rgb * InvSafeWeight;
+	float BentNormalLength = length(AverageBentNormal);
+
+	if (BentNormalLength < AverageLength && BentNormalLength > 0)
+	{
+		// Fixup normal shortening due to weighted average of vectors
+		AverageBentNormal = AverageBentNormal / BentNormalLength * AverageLength;
+	}
+	
+	OutBentNormal = float4(AverageBentNormal, CenterValue.w);
 
 #if SUPPORT_IRRADIANCE
 	OutIrradiance = float4(IrradianceAccumulation * InvSafeWeight, 0);
diff --git a/Engine/Shaders/DistanceFieldSurfaceCacheLightingCompute.usf b/Engine/Shaders/DistanceFieldSurfaceCacheLightingCompute.usf
--- a/Engine/Shaders/DistanceFieldSurfaceCacheLightingCompute.usf
+++ b/Engine/Shaders/DistanceFieldSurfaceCacheLightingCompute.usf
@@ -144,6 +144,7 @@ RWBuffer<float4> RWCulledObjectBoxBounds;
 uint ObjectBoundingGeometryIndexCount;
 
 groupshared uint NumGroupObjects;
+
 groupshared uint GroupBaseIndex;
 groupshared uint GroupObjectIndices[UPDATEOBJECTS_THREADGROUP_SIZE];
 
@@ -159,8 +160,10 @@ void CopyCulledObjectData(uint DestIndex, uint SourceIndex)
 	UNROLL
 	for (uint VectorIndex = 0; VectorIndex < CULLED_OBJECT_DATA_STRIDE; VectorIndex++)
 	{
+		float4 Data = FetchObjectDataFloat4(SourceIndex * OBJECT_DATA_STRIDE + VectorIndex);
+
 		// Note: only copying the first CULLED_OBJECT_DATA_STRIDE of the original object data
-		RWCulledObjectData[DestIndex * CULLED_OBJECT_DATA_STRIDE + VectorIndex] = FetchObjectDataFloat4(SourceIndex * OBJECT_DATA_STRIDE + VectorIndex);
+		RWCulledObjectData[DestIndex * CULLED_OBJECT_DATA_STRIDE + VectorIndex] = Data;
 	}
 }
 
@@ -355,11 +358,7 @@ void BuildTileConesMain(
 		RWTileConeAxisAndCos[TileIndex] = float4(TileConeAxis, TileConeAngleCos);
 		RWTileConeDepthRanges[TileIndex] = ConeAxisDepthRanges;
 
-#if COHERENT_OBJECT_LIST_WRITES
 		RWTileHeadDataUnpacked[TileIndex * 4 + 0] = TileIndex;
-#else
-		RWTileHeadDataUnpacked[TileIndex * 4 + 0] = TileIndex * MAX_OBJECTS_PER_TILE;
-#endif
 		RWTileHeadDataUnpacked[TileIndex * 4 + 1] = 0;
 		RWTileHeadDataUnpacked[TileIndex * 4 + 2] = 0;
 		RWTileHeadDataUnpacked[TileIndex * 4 + 3] = 0;
@@ -612,7 +611,7 @@ void PopulateCacheCS(
 	BRANCH
 	if (GBufferData.ShadingModelID > 0
 		&& IrradianceCacheSplat.w < .0001f 
-		&& all(DispatchThreadId.xy < AOBufferSize))
+		&& all((float2)DispatchThreadId.xy < AOBufferSize))
 	{
 		float SceneDepth = CalcSceneDepth(ScreenUV);
 
@@ -631,7 +630,7 @@ void PopulateCacheCS(
 			float3 WorldShadingPosition = OpaqueWorldPosition;
 
 			// For debugging
-			//if (all(DispatchThreadId.xy == uint2(2,1)))
+			//if (all(DispatchThreadId.xy == uint2(4,3)))
 			{
 				// Allocate a new record and store off attributes of the created record
 				int NextSampleIndex;
@@ -641,7 +640,7 @@ void PopulateCacheCS(
 				// Sign of W stores whether the sample was placed on geometry with a distance field representation or not
 				float RadiusSign = GBufferData.HasDistanceFieldRepresentation > 0 ? 1 : -1;
 				CachedIrradianceCachePositionRadius[NextSampleIndex] = float4(WorldShadingPosition, RadiusSign * SceneDepth * CurrentLevelDownsampleFactor * MaxRadiusScale);
-				CachedIrradianceCacheNormal[NextSampleIndex] = float4(WorldNormal, .01f * CurrentLevelDownsampleFactor / 2);
+				CachedIrradianceCacheNormal[NextSampleIndex] = float4(WorldNormal, 0);
 				CachedIrradianceCacheTileCoordinate[NextSampleIndex] = TileCoordinate;
 			}
 		}
@@ -661,7 +660,7 @@ void PopulateCacheCS(
 	{
 		int SampleIndex = BaseRecordIndex + LocalRecordIndex;
 		RWIrradianceCachePositionRadius[SampleIndex] = CachedIrradianceCachePositionRadius[LocalRecordIndex];
-		RWIrradianceCacheNormal[SampleIndex] = float4(CachedIrradianceCacheNormal[LocalRecordIndex].xyz, 0);
+		RWIrradianceCacheNormal[SampleIndex] = CachedIrradianceCacheNormal[LocalRecordIndex];
 		uint2 TileCoordinate = CachedIrradianceCacheTileCoordinate[LocalRecordIndex];
 		RWIrradianceCacheTileCoordinate[SampleIndex] = TileCoordinate;
 	}
@@ -702,7 +701,7 @@ groupshared float3 SharedGatheredLighting[FINAL_GATHER_THREADGROUP_SIZE];
 
 /** Computes ambient occlusion for a surface cache record by cone stepping through the nearby object distance fields. */
 [numthreads(FINAL_GATHER_THREADGROUP_SIZE, 1, 1)]
-void FinalGatherCS(
+void ConeTraceOcclusionCS(
 	uint3 GroupId : SV_GroupID,
 	uint3 DispatchThreadId : SV_DispatchThreadID,
     uint3 GroupThreadId : SV_GroupThreadID) 
@@ -783,8 +782,6 @@ void FinalGatherCS(
 				// Ignore extra threads
 				&& ObjectOffsetIndex < SIMULTANEOUSLY_TRACED_OBJECTS)
 			{
-#if COHERENT_OBJECT_LIST_WRITES
-
 				uint ListIndex = 0;
 				uint ArrayIndex = EffectiveListObjectIndex;
 
@@ -801,78 +798,87 @@ void FinalGatherCS(
 				}
 
 				uint ObjectIndex = TileArrayData.Load((ArrayIndex * TileListGroupSize.x * TileListGroupSize.y + TileHead.x) * NUM_CULLED_OBJECT_LISTS + ListIndex);
-#else
-				uint ObjectIndex = TileArrayData.Load(TileHead.x + EffectiveListObjectIndex);
-#endif
-
-				float3 LocalPositionExtent = LoadObjectLocalPositionExtent(ObjectIndex);
-				float4x4 WorldToVolume = LoadObjectWorldToVolume(ObjectIndex);
-				float4 UVScaleAndVolumeScale = LoadObjectUVScale(ObjectIndex);
-				float3 VolumeShadingPosition = mul(float4(WorldShadingPosition, 1), WorldToVolume).xyz;
 
-				float ObjectOccluderRadius = length(LocalPositionExtent) * .5f * UVScaleAndVolumeScale.w;
-				float BoxDistance = ComputeDistanceFromBoxToPoint(-LocalPositionExtent, LocalPositionExtent, VolumeShadingPosition) * UVScaleAndVolumeScale.w;
+				float4 ObjectPositionAndRadius = LoadObjectPositionAndRadius(ObjectIndex);
+				float ObjectDistanceSq = dot(ObjectPositionAndRadius.xyz - WorldShadingPosition, ObjectPositionAndRadius.xyz - WorldShadingPosition);
 
 				BRANCH
-				if (BoxDistance < AOMaxDistance)
+				// Skip tracing objects with a small projected angle 
+				if (ObjectPositionAndRadius.w * ObjectPositionAndRadius.w / ObjectDistanceSq > Square(.5f))
 				{
-					float3 UVAdd = LoadObjectUVAdd(ObjectIndex);
+					float3 LocalPositionExtent = LoadObjectLocalPositionExtent(ObjectIndex);
+					float4x4 WorldToVolume = LoadObjectWorldToVolume(ObjectIndex);
+					bool bGeneratedAsTwoSided;
+					float4 UVScaleAndVolumeScale = LoadObjectUVScale(ObjectIndex, bGeneratedAsTwoSided);
+					float3 VolumeShadingPosition = mul(float4(WorldShadingPosition, 1), WorldToVolume).xyz;
 
-					uint StartStepIndex = 0;
-			
-					FLATTEN
-					if (EffectiveListObjectIndex >= TileHead.y + TileHead.z)
-					{
-						StartStepIndex = 8;
-					}
-					else if (EffectiveListObjectIndex >= TileHead.y)
-					{
-						StartStepIndex = 5;
-					}
-
-					float WorldStepOffset = GetStepOffset(StartStepIndex);
+					float ObjectOccluderRadius = length(LocalPositionExtent) * .5f * UVScaleAndVolumeScale.w;
+					float BoxDistance = ComputeDistanceFromBoxToPoint(-LocalPositionExtent, LocalPositionExtent, VolumeShadingPosition) * UVScaleAndVolumeScale.w;
 
-					LOOP
-					for (uint StepIndex = StartStepIndex; StepIndex < NUM_CONE_STEPS && WorldStepOffset < MaxWorldStepOffset; StepIndex++)
+					BRANCH
+					if (BoxDistance < AOMaxDistance)
 					{
-						float3 WorldSamplePosition = WorldShadingPosition + RotatedConeDirection * WorldStepOffset;
-						float3 StepSamplePosition = mul(float4(WorldSamplePosition, 1), WorldToVolume).xyz;
-						float3 ClampedSamplePosition = clamp(StepSamplePosition, -LocalPositionExtent, LocalPositionExtent);
-						float DistanceToClamped = length(StepSamplePosition - ClampedSamplePosition);
-
-						float3 StepVolumeUV = DistanceFieldVolumePositionToUV(ClampedSamplePosition, UVScaleAndVolumeScale.xyz, UVAdd);
-						float DistanceToOccluder = (Texture3DSampleLevel(DistanceFieldTexture, DistanceFieldSampler, StepVolumeUV, 0).x + DistanceToClamped) * UVScaleAndVolumeScale.w;
-						float SphereRadius = WorldStepOffset * TanConeHalfAngle;
-						//@todo - have to bias away from surface further for this to work
-						float ShadingSphereRadius = SphereRadius * 1.0f;
+						float3 UVAdd = LoadObjectUVAdd(ObjectIndex);
 
-						// Derive visibility from 1d intersection
-						float Visibility = saturate(DistanceToOccluder / ShadingSphereRadius);
-
-						// Don't allow small objects to fully occlude a cone step
-						Visibility = max(Visibility, 1 - saturate(ObjectOccluderRadius / SphereRadius));
-						float OccluderDistanceFraction = (WorldStepOffset + DistanceToOccluder) / AOMaxDistance;
-
-						#if SUPPORT_IRRADIANCE
-							uint VisibilityIndex = NUM_VISIBILITY_STEPS * WorldStepOffset / AOMaxDistance;
-							// Track raw visibility before the distance fade for GI shadowing
-							MinRawVisibility[VisibilityIndex] = min(MinRawVisibility[VisibilityIndex], Visibility);
-						#endif
-
-						// Fade out occlusion based on distance to occluder to avoid a discontinuity at the max AO distance
-						Visibility = max(Visibility, saturate(OccluderDistanceFraction * OccluderDistanceFraction * .8f));
+						uint StartStepIndex = 0;
+			
+						FLATTEN
+						if (EffectiveListObjectIndex >= TileHead.y + TileHead.z)
+						{
+							StartStepIndex = 8;
+						}
+						else if (EffectiveListObjectIndex >= TileHead.y)
+						{
+							StartStepIndex = 5;
+						}
 
-						MinVisibility = min(MinVisibility, Visibility);
+						float WorldStepOffset = GetStepOffset(StartStepIndex);
 
-						if (DistanceToOccluder < .9f * SphereRadius)
+						LOOP
+						for (uint StepIndex = StartStepIndex; StepIndex < NUM_CONE_STEPS && WorldStepOffset < MaxWorldStepOffset; StepIndex++)
 						{
-							// Assuming occluder is straight forward along the cone
-							float WorldDistanceToOccluder = WorldStepOffset + DistanceToOccluder;
-							MinWorldDistanceToOccluder = min(MinWorldDistanceToOccluder, WorldDistanceToOccluder);
-						}
+							float3 WorldSamplePosition = WorldShadingPosition + RotatedConeDirection * WorldStepOffset;
+							float3 StepSamplePosition = mul(float4(WorldSamplePosition, 1), WorldToVolume).xyz;
+							float3 ClampedSamplePosition = clamp(StepSamplePosition, -LocalPositionExtent, LocalPositionExtent);
+							float DistanceToClamped = length(StepSamplePosition - ClampedSamplePosition);
+
+							float3 StepVolumeUV = DistanceFieldVolumePositionToUV(ClampedSamplePosition, UVScaleAndVolumeScale.xyz, UVAdd);
+							float DistanceToOccluder = (Texture3DSampleLevel(DistanceFieldTexture, DistanceFieldSampler, StepVolumeUV, 0).x + DistanceToClamped) * UVScaleAndVolumeScale.w;
+							float SphereRadius = WorldStepOffset * TanConeHalfAngle;
+							//@todo - have to bias away from surface further for this to work
+							float ShadingSphereRadius = SphereRadius * 1.0f;
+
+							// Derive visibility from 1d intersection
+							float Visibility = saturate(DistanceToOccluder / ShadingSphereRadius);
+
+							// Don't allow small objects to fully occlude a cone step
+							Visibility = max(Visibility, 1 - saturate(ObjectOccluderRadius / SphereRadius));
+							float OccluderDistanceFraction = (WorldStepOffset + DistanceToOccluder) / AOMaxDistance;
+
+							#if SUPPORT_IRRADIANCE
+								uint VisibilityIndex = NUM_VISIBILITY_STEPS * WorldStepOffset / AOMaxDistance;
+								// Less GI occlusion for two sided meshes, which can't separate self-occlusion
+								//@todo - expose
+								float TwoSidedVisibilityScale = bGeneratedAsTwoSided ? 100 : 1;
+								// Track raw visibility before the distance fade for GI shadowing
+								MinRawVisibility[VisibilityIndex] = min(MinRawVisibility[VisibilityIndex], TwoSidedVisibilityScale * Visibility);
+							#endif
+
+							// Fade out occlusion based on distance to occluder to avoid a discontinuity at the max AO distance
+							Visibility = max(Visibility, saturate(OccluderDistanceFraction * OccluderDistanceFraction * .8f));
+
+							MinVisibility = min(MinVisibility, Visibility);
+
+							if (DistanceToOccluder < .9f * SphereRadius)
+							{
+								// Assuming occluder is straight forward along the cone
+								float WorldDistanceToOccluder = WorldStepOffset + DistanceToOccluder;
+								MinWorldDistanceToOccluder = min(MinWorldDistanceToOccluder, WorldDistanceToOccluder);
+							}
 						
-						float MinStepSize = .6f * (GetStepOffset(StepIndex + 1) - GetStepOffset(StepIndex));
-						WorldStepOffset += max(DistanceToOccluder, MinStepSize);
+							float MinStepSize = .6f * (GetStepOffset(StepIndex + 1) - GetStepOffset(StepIndex));
+							WorldStepOffset += max(DistanceToOccluder, MinStepSize);
+						}
 					}
 				}
 			}
@@ -921,7 +927,8 @@ void FinalGatherCS(
 			RWRecordConeVisibility[GroupId.x * NUM_CONE_DIRECTIONS + ConeIndex] = ConeVisibility;
 		}
 
-		RWOccluderRadius[RecordIndex] = RecordRadiusScale * asfloat(SharedMinOcclusionDistance);
+		float RecordRadius = RecordRadiusScale * asfloat(SharedMinOcclusionDistance);
+		RWOccluderRadius[RecordIndex] = RecordRadius;
 	}
 }
 
@@ -933,9 +940,13 @@ Buffer<float4> HeightfieldDescriptions;
 // In float4's, must match C++
 #define HEIGHTFIELD_DATA_STRIDE 10
 
-float4 GetHeightfieldScaleBias(uint Index)
+float4 GetHeightfieldScaleBias(uint Index, out bool bHasSubsections)
 {
-	return HeightfieldDescriptions[Index * HEIGHTFIELD_DATA_STRIDE + 0];
+	float4 Vector = HeightfieldDescriptions[Index * HEIGHTFIELD_DATA_STRIDE + 0];
+	// Sign bit stored presence of subsections
+	bHasSubsections = Vector < 0;
+	Vector.x = abs(Vector.x);
+	return Vector;
 }
 
 float4 GetHeightfieldMinMaxUV(uint Index)
@@ -990,10 +1001,20 @@ void CalculateHeightfieldOcclusionCS(
 		for (uint HeightfieldIndex = 0; HeightfieldIndex < NumHeightfields; HeightfieldIndex++)
 		{
 			float3 LocalPosition = mul(float4(WorldPosition, 1), GetWorldToLocal(HeightfieldIndex)).xyz;
-			float4 HeightfieldScaleBias = GetHeightfieldScaleBias(HeightfieldIndex);
+			bool bHasSubsections;
+			float4 HeightfieldScaleBias = GetHeightfieldScaleBias(HeightfieldIndex, bHasSubsections);
 			float2 HeightfieldUV = LocalPosition.xy * HeightfieldScaleBias.xy + HeightfieldScaleBias.zw;
 			float4 MinMaxHeightfieldUV = GetHeightfieldMinMaxUV(HeightfieldIndex);
 
+			if (bHasSubsections)
+			{
+				float2 CenterUV = .5f * (MinMaxHeightfieldUV.xy + MinMaxHeightfieldUV.zw);
+				// One row of heights in the heightfield are duplicated between subsections
+				// The only number of subsections supported is 2
+				// If we are in subsection 0, no adjustment is needed, but if we are in subsection 1 we need to shift up by one texel
+				HeightfieldUV += HeightfieldUV + .5f * HeightfieldScaleBias.xy > CenterUV ? HeightfieldScaleBias.xy : 0;
+			}
+
 			BRANCH
 			if (all(HeightfieldUV > MinMaxHeightfieldUV.xy) 
 				&& all(HeightfieldUV < MinMaxHeightfieldUV.zw))
@@ -1003,6 +1024,7 @@ void CalculateHeightfieldOcclusionCS(
 
 				float2 SampleNormal = float2(SampleValue.b, SampleValue.a) * float2(2.0,2.0) - float2(1.0,1.0);
 				// @todo - use the normal to compute occlusion from an arbitrary plane, instead of assuming upward facing plane
+				// This would require an averaged normal over the gather region
 				float3 HeightfieldNormal = float3(SampleNormal, sqrt(max(1.0 - dot(SampleNormal, SampleNormal), 0.0)));
 
 				float3 HeightfieldShadingPosition = float3(LocalPosition.xy, HeightfieldHeight);
@@ -1017,7 +1039,7 @@ void CalculateHeightfieldOcclusionCS(
 
 					float ExistingConeVisibility = RWRecordConeVisibility[DispatchThreadId.x * NUM_CONE_DIRECTIONS + ConeIndex];
 				
-					float ZBias = 10;
+					float ZBias = 20;
 					float DistanceAlongConeToOccluder = clamp((WorldPosition.z + ZBias - WorldHeightfieldShadingPosition.z) / max(-RotatedConeDirection.z, .0001f), .00001f, AOMaxDistance);
 					float DistanceFromConeEndToOccluder = max(WorldPosition.z + ZBias + RotatedConeDirection.z * DistanceAlongConeToOccluder - WorldHeightfieldShadingPosition.z, 0);
 
@@ -1131,7 +1153,7 @@ void SetupCopyIndirectArgumentsCS(
 {
 	uint NumRecords = (1 - TrimFraction) * DrawParameters[1];
 
-	// Make sure we spawn at least one group so RWDispatchParameters gets written to
+	// Make sure we spawn at least one group so RWScatterDrawParameters gets written to in the next pass
 	RWDispatchParameters[0] = max((NumRecords + COPY_THREADGROUP_SIZE - 1) / COPY_THREADGROUP_SIZE, 1);
 	RWDispatchParameters[1] = 1;
 	RWDispatchParameters[2] = 1;
@@ -1153,6 +1175,7 @@ void CopyIrradianceCacheSamplesCS(
     uint3 GroupThreadId : SV_GroupThreadID) 
 {
 	uint NumRecords = DrawParameters[1];
+
 	// Always trim at least one to handle dynamic scene changes not accounted for with irradiance cache misses
 	uint StartIndex = max(NumRecords * TrimFraction, 1);
 	StartIndex = min(StartIndex, NumRecords);
diff --git a/Engine/Shaders/EvaluateSurfelMaterial.usf b/Engine/Shaders/EvaluateSurfelMaterial.usf
new file mode 100644
--- /dev/null
+++ b/Engine/Shaders/EvaluateSurfelMaterial.usf
@@ -0,0 +1,76 @@
+// Copyright 1998-2014 Epic Games, Inc. All Rights Reserved.
+
+/*=============================================================================
+	ConvertToUniformMesh.usf
+=============================================================================*/
+
+#include "Common.usf"
+#include "Material.usf"
+#include "SurfelMaterialShared.usf"
+
+uint SurfelStartIndex;
+uint NumSurfelsToGenerate;
+Buffer<float4> InterpolatedVertexData;
+RWBuffer<float4> RWSurfelData;
+
+void AssembleMaterialParameters(out FMaterialPixelParameters MaterialParameters, uint SurfelIndex)
+{
+	MaterialParameters = (FMaterialPixelParameters)0;
+
+	float3 Tangent0 = InterpolatedVertexData[(SurfelStartIndex + SurfelIndex) * INTERPOLATED_VERTEX_DATA_STRIDE + 1].xyz;
+	float3 Tangent1 = InterpolatedVertexData[(SurfelStartIndex + SurfelIndex) * INTERPOLATED_VERTEX_DATA_STRIDE + 2].xyz;
+	float3 Tangent2 = InterpolatedVertexData[(SurfelStartIndex + SurfelIndex) * INTERPOLATED_VERTEX_DATA_STRIDE + 3].xyz;
+
+	MaterialParameters.TangentToWorld = float3x3(Tangent0, Tangent1, Tangent2);
+
+	float4 PackedUVs = InterpolatedVertexData[(SurfelStartIndex + SurfelIndex) * INTERPOLATED_VERTEX_DATA_STRIDE + 4];
+
+	#if NUM_MATERIAL_TEXCOORDS > 0
+		MaterialParameters.TexCoords[0] = PackedUVs.xy;
+	#endif
+
+	#if NUM_MATERIAL_TEXCOORDS > 1
+		MaterialParameters.TexCoords[1] = PackedUVs.zw;
+	#endif
+
+	float4 VertexColor = InterpolatedVertexData[(SurfelStartIndex + SurfelIndex) * INTERPOLATED_VERTEX_DATA_STRIDE + 5];
+
+	MaterialParameters.VertexColor = VertexColor;
+}
+
+float4x4 Instance0InverseTransform;
+
+[numthreads(EVALUATE_SURFEL_MATERIAL_GROUP_SIZE, 1, 1)]
+void EvaluateSurfelMaterialCS(
+	uint3 GroupId : SV_GroupID,
+	uint3 DispatchThreadId : SV_DispatchThreadID,
+    uint3 GroupThreadId : SV_GroupThreadID) 
+{
+	uint SurfelIndex = DispatchThreadId.x;
+
+	if (SurfelIndex < NumSurfelsToGenerate)
+	{
+		float4 PositionAndRadius = InterpolatedVertexData[(SurfelStartIndex + SurfelIndex) * INTERPOLATED_VERTEX_DATA_STRIDE + 0];
+
+		FMaterialPixelParameters MaterialParameters;
+		AssembleMaterialParameters(MaterialParameters, SurfelIndex);
+
+		float4 PixelPosition = float4(PositionAndRadius.xyz + View.PreViewTranslation.xyz, 0);
+		CalcMaterialParameters(MaterialParameters, float4(0, 0, 0, 0), true, PixelPosition);
+
+		float3 BaseColor = GetMaterialBaseColor(MaterialParameters);	
+		float Metallic = GetMaterialMetallic(MaterialParameters);
+		float3 DiffuseColor = BaseColor - BaseColor * Metallic;
+		float3 EmissiveColor = GetMaterialEmissive(MaterialParameters);	
+
+		float3 InstanceLocalPosition = mul(float4(PositionAndRadius.xyz, 1), Instance0InverseTransform).xyz;
+		//@todo - use a proper normal transform
+		float3 InstanceLocalNormal = mul(MaterialParameters.TangentToWorld[2], (float3x3)Instance0InverseTransform);
+
+		//@todo - take instance scale into account for the radius
+		RWSurfelData[(SurfelStartIndex + SurfelIndex) * SURFEL_DATA_STRIDE + 0] = float4(InstanceLocalPosition, PositionAndRadius.w);
+		RWSurfelData[(SurfelStartIndex + SurfelIndex) * SURFEL_DATA_STRIDE + 1] = float4(InstanceLocalNormal, 0);
+		RWSurfelData[(SurfelStartIndex + SurfelIndex) * SURFEL_DATA_STRIDE + 2] = float4(DiffuseColor, 0);
+		RWSurfelData[(SurfelStartIndex + SurfelIndex) * SURFEL_DATA_STRIDE + 3] = float4(EmissiveColor, 0);
+	}
+}
\ No newline at end of file
diff --git a/Engine/Shaders/MaterialTemplate.usf b/Engine/Shaders/MaterialTemplate.usf
--- a/Engine/Shaders/MaterialTemplate.usf
+++ b/Engine/Shaders/MaterialTemplate.usf
@@ -250,7 +250,7 @@ struct FMaterialVertexParameters
 	FMaterialParticleParameters Particle;
 };
 
-#if MESH_MATERIAL_SHADER
+#if HAS_PRIMITIVE_UNIFORM_BUFFER
 
 /** Transforms a vector from tangent space to world space, prescaling by an amount calculated previously */
 MaterialFloat3 TransformTangentVectorToWorld_PreScaled(FMaterialTessellationParameters Parameters, MaterialFloat3 InTangentVector)
@@ -410,7 +410,7 @@ MaterialFloat GetDistanceCullFade()
 	return saturate( View.RealTime * PrimitiveFade.FadeTimeScaleBias.x + PrimitiveFade.FadeTimeScaleBias.y );
 }
 
-#endif // #if MESH_MATERIAL_SHADER
+#endif // #if HAS_PRIMITIVE_UNIFORM_BUFFER
 
 /** Transforms a vector from view space to world space */
 MaterialFloat3 TransformViewVectorToWorld(MaterialFloat3 InViewVector)
@@ -626,14 +626,14 @@ float GetSphericalParticleOpacity(FMaterialPixelParameters Parameters, float Den
 {
 	float Opacity = 0;
 
-#if PARTICLE_FACTORY || MESH_MATERIAL_SHADER
+#if PARTICLE_FACTORY || HAS_PRIMITIVE_UNIFORM_BUFFER
 
 #if PARTICLE_FACTORY
 
 	float3 ParticlePosition = Parameters.Particle.PositionAndSize.xyz;
 	float ParticleRadius = max(0.000001f, Parameters.Particle.PositionAndSize.w);
 
-#elif MESH_MATERIAL_SHADER
+#elif HAS_PRIMITIVE_UNIFORM_BUFFER
 
 	// Substitute object attributes if the mesh is not a particle
 	// This is mostly useful for previewing materials using spherical opacity in the material editor
@@ -1400,7 +1400,7 @@ void CalcMaterialParameters(
 
 	Parameters.TwoSidedSign = 1.0f;
 
-#if MATERIAL_TWOSIDED && MESH_MATERIAL_SHADER
+#if MATERIAL_TWOSIDED && HAS_PRIMITIVE_UNIFORM_BUFFER
 	Parameters.TwoSidedSign *= View.CullingSign * Primitive.LocalToWorldDeterminantSign;
 
 	#if !MATERIAL_TWOSIDED_SEPARATE_PASS
diff --git a/Engine/Shaders/SurfelMaterialShared.usf b/Engine/Shaders/SurfelMaterialShared.usf
new file mode 100644
--- /dev/null
+++ b/Engine/Shaders/SurfelMaterialShared.usf
@@ -0,0 +1,9 @@
+// Copyright 1998-2015 Epic Games, Inc. All Rights Reserved.
+
+/*=============================================================================
+	SurfelMaterialShared.usf
+=============================================================================*/
+
+// In float4's, must match C++
+#define SURFEL_DATA_STRIDE 4
+#define INTERPOLATED_VERTEX_DATA_STRIDE 6
diff --git a/Engine/Shaders/SurfelTree.usf b/Engine/Shaders/SurfelTree.usf
new file mode 100644
--- /dev/null
+++ b/Engine/Shaders/SurfelTree.usf
@@ -0,0 +1,235 @@
+// Copyright 1998-2014 Epic Games, Inc. All Rights Reserved.
+
+/*=============================================================================
+	SurfelTree.usf
+=============================================================================*/
+
+#include "Common.usf"
+#include "DeferredShadingCommon.usf"
+#include "DistanceFieldLightingShared.usf"
+#include "DistanceFieldAOShared.usf"
+#include "MonteCarlo.usf"
+
+uint NumSurfels;
+RWBuffer<float4> RWCopyInterpolatedVertexData;
+RWBuffer<float4> RWCopySurfelData;
+Buffer<float4> InterpolatedVertexData;
+
+[numthreads(UPDATEOBJECTS_THREADGROUP_SIZE, 1, 1)]
+void CopySurfelBufferCS(
+	uint3 GroupId : SV_GroupID,
+	uint3 DispatchThreadId : SV_DispatchThreadID,
+    uint3 GroupThreadId : SV_GroupThreadID) 
+{
+	uint CopyOperationIndex = DispatchThreadId.x;
+
+	if (CopyOperationIndex < NumSurfels)
+	{
+		uint SourceIndex = CopyOperationIndex;
+		uint DestIndex = CopyOperationIndex;
+
+		UNROLL
+		for (uint VectorIndex = 0; VectorIndex < INTERPOLATED_VERTEX_DATA_STRIDE; VectorIndex++)
+		{
+			RWCopyInterpolatedVertexData[DestIndex * INTERPOLATED_VERTEX_DATA_STRIDE + VectorIndex] = InterpolatedVertexData[SourceIndex * INTERPOLATED_VERTEX_DATA_STRIDE + VectorIndex];
+		}
+
+		UNROLL
+		for (uint VectorIndex = 0; VectorIndex < SURFEL_DATA_STRIDE; VectorIndex++)
+		{
+			RWCopySurfelData[DestIndex * SURFEL_DATA_STRIDE + VectorIndex] = SurfelData[SourceIndex * SURFEL_DATA_STRIDE + VectorIndex];
+		}
+	}
+}
+
+Buffer<float4> VPLFlux;
+RWBuffer<float4> RWCopyVPLFlux;
+
+[numthreads(UPDATEOBJECTS_THREADGROUP_SIZE, 1, 1)]
+void CopyVPLFluxBufferCS(
+	uint3 GroupId : SV_GroupID,
+	uint3 DispatchThreadId : SV_DispatchThreadID,
+    uint3 GroupThreadId : SV_GroupThreadID) 
+{
+	uint CopyOperationIndex = DispatchThreadId.x;
+
+	if (CopyOperationIndex < NumSurfels)
+	{
+		uint SourceIndex = CopyOperationIndex;
+		uint DestIndex = CopyOperationIndex;
+
+		RWCopyVPLFlux[DestIndex] = VPLFlux[SourceIndex];
+	}
+}
+
+// In float's, must match FConvertToUniformMeshVertex and ComputeUniformVertexStride()
+#define TRIANGLE_VERTEX_DATA_STRIDE 24
+
+uint NumTriangles;
+Buffer<float> TriangleVertexData;
+
+float LoadVertexAttribute(uint TriangleIndex, uint VertexIndex, uint AttributeIndex)
+{
+	return TriangleVertexData[(TriangleIndex * 3 + VertexIndex) * TRIANGLE_VERTEX_DATA_STRIDE + AttributeIndex];
+}
+
+RWBuffer<float> RWTriangleAreas;
+
+[numthreads(THREADGROUP_TOTALSIZE, 1, 1)]
+void ComputeTriangleAreasCS(
+	uint3 GroupId : SV_GroupID,
+	uint3 DispatchThreadId : SV_DispatchThreadID,
+    uint3 GroupThreadId : SV_GroupThreadID) 
+{
+	uint TriangleIndex = DispatchThreadId.x;
+
+	if (TriangleIndex < NumTriangles)
+	{
+		float3 Position0 = float3(LoadVertexAttribute(TriangleIndex, 0, 0), LoadVertexAttribute(TriangleIndex, 0, 1), LoadVertexAttribute(TriangleIndex, 0, 2));
+		float3 Position1 = float3(LoadVertexAttribute(TriangleIndex, 1, 0), LoadVertexAttribute(TriangleIndex, 1, 1), LoadVertexAttribute(TriangleIndex, 1, 2));
+		float3 Position2 = float3(LoadVertexAttribute(TriangleIndex, 2, 0), LoadVertexAttribute(TriangleIndex, 2, 1), LoadVertexAttribute(TriangleIndex, 2, 2));
+
+		float3 TriangleNormal = cross(Position2 - Position0, Position1 - Position0);
+		float TriangleArea = .5f * length(TriangleNormal);
+		RWTriangleAreas[TriangleIndex] = TriangleArea;
+	}
+}
+
+Buffer<float> TriangleAreas;
+RWBuffer<float> RWTriangleCDFs;
+
+[numthreads(THREADGROUP_TOTALSIZE, 1, 1)]
+void ComputeTriangleCDFsCS(
+	uint3 GroupId : SV_GroupID,
+	uint3 DispatchThreadId : SV_DispatchThreadID,
+    uint3 GroupThreadId : SV_GroupThreadID) 
+{
+	uint TriangleIndex = DispatchThreadId.x;
+
+	if (TriangleIndex < NumTriangles)
+	{
+		float Integral = 0;
+
+		for (int i = 0; i <= TriangleIndex; i++)
+		{
+			Integral += TriangleAreas[i];
+		}
+
+		RWTriangleCDFs[TriangleIndex] = Integral;
+	}
+}
+
+float InterpolateAttribute(uint TriangleIndex, float3 BarycentricCoordinates, uint AttributeIndex)
+{
+	float Attribute = BarycentricCoordinates.x * LoadVertexAttribute(TriangleIndex, 0, AttributeIndex) 
+		+ BarycentricCoordinates.y * LoadVertexAttribute(TriangleIndex, 1, AttributeIndex) 
+		+ BarycentricCoordinates.z * LoadVertexAttribute(TriangleIndex, 2, AttributeIndex);
+
+	return Attribute;
+}
+
+float4 InterpolateFloat4Attribute(uint TriangleIndex, float3 BarycentricCoordinates, uint StartAttributeIndex)
+{
+	float4 SampleValues;
+	SampleValues.x = InterpolateAttribute(TriangleIndex, BarycentricCoordinates, StartAttributeIndex + 0);
+	SampleValues.y = InterpolateAttribute(TriangleIndex, BarycentricCoordinates, StartAttributeIndex + 1);
+	SampleValues.z = InterpolateAttribute(TriangleIndex, BarycentricCoordinates, StartAttributeIndex + 2);
+	SampleValues.w = InterpolateAttribute(TriangleIndex, BarycentricCoordinates, StartAttributeIndex + 3);
+	return SampleValues;
+}
+
+float3 InterpolateFloat3Attribute(uint TriangleIndex, float3 BarycentricCoordinates, uint StartAttributeIndex)
+{
+	float3 SampleValues;
+	SampleValues.x = InterpolateAttribute(TriangleIndex, BarycentricCoordinates, StartAttributeIndex + 0);
+	SampleValues.y = InterpolateAttribute(TriangleIndex, BarycentricCoordinates, StartAttributeIndex + 1);
+	SampleValues.z = InterpolateAttribute(TriangleIndex, BarycentricCoordinates, StartAttributeIndex + 2);
+	return SampleValues;
+}
+
+float2 InterpolateFloat2Attribute(uint TriangleIndex, float3 BarycentricCoordinates, uint StartAttributeIndex)
+{
+	float2 SampleValues;
+	SampleValues.x = InterpolateAttribute(TriangleIndex, BarycentricCoordinates, StartAttributeIndex + 0);
+	SampleValues.y = InterpolateAttribute(TriangleIndex, BarycentricCoordinates, StartAttributeIndex + 1);
+	return SampleValues;
+}
+
+Buffer<float> TriangleCDFs;
+
+float SampleTrianglesProportionalToArea(float RandomFraction, float TotalArea)
+{
+	float Sample = .999f;
+	int ElementIndex = -1;
+
+	// Find the index of where the step function becomes greater or equal to the generated number
+	//@todo - CDFArray is monotonically increasing so we can do better than a linear time search
+	for (int i = 0; i < NumTriangles; i++)
+	{
+		if (RandomFraction * TotalArea <= TriangleCDFs[i])
+		{
+			ElementIndex = i;
+			break;
+		}
+	}
+
+	if (ElementIndex >= 0)
+	{
+		Sample = ElementIndex / (float)NumTriangles;
+	}
+
+	return Sample;
+}
+
+uint SurfelStartIndex;
+uint NumSurfelsToGenerate;
+RWBuffer<float4> RWInterpolatedVertexData;
+
+[numthreads(THREADGROUP_TOTALSIZE, 1, 1)]
+void SampleTrianglesCS(
+	uint3 GroupId : SV_GroupID,
+	uint3 DispatchThreadId : SV_DispatchThreadID,
+    uint3 GroupThreadId : SV_GroupThreadID) 
+{
+	uint SurfelIndex = DispatchThreadId.x;
+
+	if (SurfelIndex < NumSurfelsToGenerate)
+	{
+		//@todo - 1d uniform
+		float Uniform = SurfelIndex / (float)NumSurfelsToGenerate;
+		float TotalArea = TriangleCDFs[NumTriangles - 1];
+
+		//uint TriangleIndex = (uint)(Uniform * NumTriangles);
+		uint TriangleIndex = (uint)(SampleTrianglesProportionalToArea(Uniform, TotalArea) * NumTriangles);
+
+		float2 Uniforms = frac(Square((SurfelIndex + 15) / (float)NumSurfelsToGenerate + 10) * float2(1361.456345, 2333.578));
+		//float2 Uniforms = Hammersley(SurfelIndex, NumSurfelsToGenerate, uint2(0, 0));
+
+		float SqrtUniform0 = sqrt(Uniforms.x);
+		float3 BarycentricCoordinates = float3(1 - SqrtUniform0, SqrtUniform0 * Uniforms.y, SqrtUniform0 * (1 - Uniforms.y));
+		//BarycentricCoordinates = float3(1, 0, 0);
+
+		float DiskArea = TotalArea / NumSurfelsToGenerate;
+		float DiskRadius = sqrt(DiskArea / PI);
+		//float DiskRadius = 30;
+		float4 PositionAndRadius = float4(InterpolateFloat3Attribute(TriangleIndex, BarycentricCoordinates, 0), DiskRadius);
+		RWInterpolatedVertexData[(SurfelStartIndex + SurfelIndex) * INTERPOLATED_VERTEX_DATA_STRIDE + 0] = PositionAndRadius;
+
+		float3 Tangent0 = InterpolateFloat3Attribute(TriangleIndex, BarycentricCoordinates, 4);
+		float3 Tangent1 = InterpolateFloat3Attribute(TriangleIndex, BarycentricCoordinates, 7);
+		float3 Tangent2 = InterpolateFloat3Attribute(TriangleIndex, BarycentricCoordinates, 10);
+
+		RWInterpolatedVertexData[(SurfelStartIndex + SurfelIndex) * INTERPOLATED_VERTEX_DATA_STRIDE + 1] = float4(Tangent0, 0);
+		RWInterpolatedVertexData[(SurfelStartIndex + SurfelIndex) * INTERPOLATED_VERTEX_DATA_STRIDE + 2] = float4(Tangent1, 0);
+		RWInterpolatedVertexData[(SurfelStartIndex + SurfelIndex) * INTERPOLATED_VERTEX_DATA_STRIDE + 3] = float4(Tangent2, 0);
+
+		float2 UV0 = InterpolateFloat2Attribute(TriangleIndex, BarycentricCoordinates, 13);
+		float2 UV1 = InterpolateFloat2Attribute(TriangleIndex, BarycentricCoordinates, 15);
+
+		RWInterpolatedVertexData[(SurfelStartIndex + SurfelIndex) * INTERPOLATED_VERTEX_DATA_STRIDE + 4] = float4(UV0, UV1);
+
+		float4 VertexColor = InterpolateFloat4Attribute(TriangleIndex, BarycentricCoordinates, 17);
+
+		RWInterpolatedVertexData[(SurfelStartIndex + SurfelIndex) * INTERPOLATED_VERTEX_DATA_STRIDE + 5] = VertexColor;
+	}
+}
\ No newline at end of file
diff --git a/Engine/Source/Runtime/Engine/Private/DirectionalLightComponent.cpp b/Engine/Source/Runtime/Engine/Private/DirectionalLightComponent.cpp
--- a/Engine/Source/Runtime/Engine/Private/DirectionalLightComponent.cpp
+++ b/Engine/Source/Runtime/Engine/Private/DirectionalLightComponent.cpp
@@ -511,18 +511,17 @@ private:
 		}
 	}
 
-	virtual FSphere GetShadowSplitBoundsDepthRange(const FSceneView& View, float SplitNear, float SplitFar, FShadowCascadeSettings* OutCascadeSettings) const
+	virtual FSphere GetShadowSplitBoundsDepthRange(const FSceneView& View, FVector ViewOrigin, float SplitNear, float SplitFar, FShadowCascadeSettings* OutCascadeSettings) const
 	{
 		const FMatrix ViewMatrix = View.ShadowViewMatrices.ViewMatrix;
 		const FMatrix ProjectionMatrix = View.ShadowViewMatrices.ProjMatrix;
-		const FVector4 ViewOrigin = View.ShadowViewMatrices.ViewOrigin;
 
 		const FVector CameraDirection = ViewMatrix.GetColumn(2);
 		const FVector LightDirection = -GetDirection();
 
 		// Get FOV and AspectRatio from the view's projection matrix.
 		float AspectRatio = ProjectionMatrix.M[1][1] / ProjectionMatrix.M[0][0];
-		float HalfFOV = (ViewOrigin.W > 0.0f) ? FMath::Atan(1.0f / ProjectionMatrix.M[0][0]) : PI/4.0f;
+		float HalfFOV = View.ShadowViewMatrices.IsPerspectiveProjection() ? FMath::Atan(1.0f / ProjectionMatrix.M[0][0]) : PI/4.0f;
 
 		// Build the camera frustum for this cascade
 		const float StartHorizontalLength = SplitNear * FMath::Tan(HalfFOV);
@@ -638,7 +637,7 @@ private:
 			OutCascadeSettings->FadePlaneLength = SplitFar - FadePlane;
 		}
 
-		const FSphere CascadeSphere = FDirectionalLightSceneProxy::GetShadowSplitBoundsDepthRange(View, SplitNear, SplitFar, OutCascadeSettings);
+		const FSphere CascadeSphere = FDirectionalLightSceneProxy::GetShadowSplitBoundsDepthRange(View, View.ViewMatrices.ViewOrigin, SplitNear, SplitFar, OutCascadeSettings);
 
 		return CascadeSphere;
 	}
diff --git a/Engine/Source/Runtime/Engine/Private/DistanceFieldAtlas.cpp b/Engine/Source/Runtime/Engine/Private/DistanceFieldAtlas.cpp
--- a/Engine/Source/Runtime/Engine/Private/DistanceFieldAtlas.cpp
+++ b/Engine/Source/Runtime/Engine/Private/DistanceFieldAtlas.cpp
@@ -23,7 +23,21 @@ FDistanceFieldVolumeTextureAtlas::FDistanceFieldVolumeTextureAtlas(EPixelFormat
 	Generation = 0;
 	Format = InFormat;
 }
- 
+
+FString FDistanceFieldVolumeTextureAtlas::GetSizeString() const
+{
+	if (VolumeTextureRHI)
+	{
+		const int32 FormatSize = GPixelFormats[Format].BlockBytes;
+		float MemorySize = VolumeTextureRHI->GetSizeX() * VolumeTextureRHI->GetSizeY() * VolumeTextureRHI->GetSizeZ() * FormatSize / 1024.0f / 1024.0f;
+		return FString::Printf(TEXT("Allocated %ux%ux%u distance field atlas = %.1fMb"), VolumeTextureRHI->GetSizeX(), VolumeTextureRHI->GetSizeY(), VolumeTextureRHI->GetSizeZ(), MemorySize);
+	}
+	else
+	{
+		return TEXT("");
+	}
+}
+
 void FDistanceFieldVolumeTextureAtlas::AddAllocation(FDistanceFieldVolumeTexture* Texture)
 {
 	PendingAllocations.AddUnique(Texture);
@@ -115,9 +129,7 @@ void FDistanceFieldVolumeTextureAtlas::UpdateAllocations()
 				TexCreate_ShaderResource,
 				CreateInfo);
 
-			const int32 FormatSize = GPixelFormats[Format].BlockBytes;
-			float MemorySize = VolumeTextureRHI->GetSizeX() * VolumeTextureRHI->GetSizeY() * VolumeTextureRHI->GetSizeZ() * FormatSize / 1024.0f / 1024.0f;
-			UE_LOG(LogStaticMesh,Log,TEXT("Allocated %ux%ux%u distance field atlas = %.1fMb"), VolumeTextureRHI->GetSizeX(), VolumeTextureRHI->GetSizeY(), VolumeTextureRHI->GetSizeZ(), MemorySize);
+			UE_LOG(LogStaticMesh,Log,TEXT("Allocated %s"), *GetSizeString());
 		}
 
 		for (int32 AllocationIndex = 0; AllocationIndex < PendingAllocations.Num(); AllocationIndex++)
diff --git a/Engine/Source/Runtime/Engine/Private/HierarchicalInstancedStaticMesh.cpp b/Engine/Source/Runtime/Engine/Private/HierarchicalInstancedStaticMesh.cpp
--- a/Engine/Source/Runtime/Engine/Private/HierarchicalInstancedStaticMesh.cpp
+++ b/Engine/Source/Runtime/Engine/Private/HierarchicalInstancedStaticMesh.cpp
@@ -642,35 +642,6 @@ public:
 	{
 	}
 
-
-	virtual int32 GetNumMeshBatches() const override
-	{
-		check(0);
-		return 0;
-	}
-
-	/** Sets up a shadow FMeshBatch for a specific LOD. */
-	virtual bool GetShadowMeshElement(int32 LODIndex, int32 BatchIndex, uint8 InDepthPriorityGroup, FMeshBatch& OutMeshBatch) const override
-	{
-		check(0);
-		return false;
-	}
-
-	/** Sets up a FMeshBatch for a specific LOD and element. */
-	virtual bool GetMeshElement(int32 LODIndex, int32 BatchIndex, int32 ElementIndex, uint8 InDepthPriorityGroup, const bool bUseSelectedMaterial, const bool bUseHoveredMaterial, FMeshBatch& OutMeshBatch) const override
-	{
-		check(0);
-		return false;
-	}
-
-
-	/** Sets up a wireframe FMeshBatch for a specific LOD. */
-	virtual bool GetWireframeMeshElement(int32 LODIndex, int32 BatchIndex, const FMaterialRenderProxy* WireframeRenderProxy, uint8 InDepthPriorityGroup, FMeshBatch& OutMeshBatch) const override
-	{
-		check(0);
-		return false;
-	}
-
 	void FillDynamicMeshElements(FMeshElementCollector& Collector, const FFoliageElementParams& ElementParams, const FFoliageRenderInstanceParams& Instances) const;
 
 	template<bool TUseVector>
@@ -1091,6 +1062,12 @@ void FHierarchicalStaticMeshSceneProxy::FillDynamicMeshElements(FMeshElementColl
 
 void FHierarchicalStaticMeshSceneProxy::GetDynamicMeshElements(const TArray<const FSceneView*>& Views, const FSceneViewFamily& ViewFamily, uint32 VisibilityMap, FMeshElementCollector& Collector) const
 {
+	if (Views[0]->bRenderFirstInstanceOnly)
+	{
+		FInstancedStaticMeshSceneProxy::GetDynamicMeshElements(Views, ViewFamily, VisibilityMap, Collector);
+		return;
+	}
+
 	QUICK_SCOPE_CYCLE_COUNTER(STAT_HierarchicalInstancedStaticMeshSceneProxy_GetMeshElements);
 
 	bool bMultipleSections = ALLOW_DITHERED_LOD_FOR_INSTANCED_STATIC_MESHES && CVarDitheredLOD.GetValueOnRenderThread() > 0;
diff --git a/Engine/Source/Runtime/Engine/Private/InstancedStaticMesh.cpp b/Engine/Source/Runtime/Engine/Private/InstancedStaticMesh.cpp
--- a/Engine/Source/Runtime/Engine/Private/InstancedStaticMesh.cpp
+++ b/Engine/Source/Runtime/Engine/Private/InstancedStaticMesh.cpp
@@ -522,6 +522,14 @@ void FInstancedStaticMeshSceneProxy::GetDynamicMeshElements(const TArray<const F
 							MeshElement.bUseSelectionOutline = BatchRenderSelection[SelectionGroupIndex];
 							MeshElement.bUseWireframeSelectionColoring = BatchRenderSelection[SelectionGroupIndex];
 
+							if (View->bRenderFirstInstanceOnly)
+							{
+								for (int32 ElementIndex = 0; ElementIndex < MeshElement.Elements.Num(); ElementIndex++)
+								{
+									MeshElement.Elements[ElementIndex].NumInstances = FMath::Min<uint32>(MeshElement.Elements[ElementIndex].NumInstances, 1);
+								}
+							}
+
 							Collector.AddMesh(ViewIndex, MeshElement);
 							INC_DWORD_STAT_BY(STAT_StaticMeshTriangles, MeshElement.GetNumPrimitives());
 						}
@@ -632,7 +640,7 @@ void FInstancedStaticMeshSceneProxy::GetDistancefieldAtlasData(FBox& LocalVolume
 	ObjectLocalToWorldTransforms.Reset();
 
 	const FInstanceStream* InstanceStream = InstancedRenderData.PerInstanceRenderData->InstanceBuffer.GetData();
-	FVector4 FourthVector(0, 0, 0, 1);
+	const FVector4 FourthVector(0, 0, 0, 1);
 
 	const uint32 NumInstances = InstancedRenderData.PerInstanceRenderData->InstanceBuffer.GetNumInstances();
 	for (uint32 InstanceIndex = 0; InstanceIndex < NumInstances; InstanceIndex++)
@@ -659,6 +667,43 @@ void FInstancedStaticMeshSceneProxy::GetDistancefieldAtlasData(FBox& LocalVolume
 	}
 }
 
+void FInstancedStaticMeshSceneProxy::GetDistanceFieldInstanceInfo(int32& NumInstances, float& BoundsSurfaceArea) const
+{
+	NumInstances = DistanceFieldData ? InstancedRenderData.PerInstanceRenderData->InstanceBuffer.GetNumInstances() : 0;
+
+	if (NumInstances > 0)
+	{
+		const FInstanceStream* InstanceStream = InstancedRenderData.PerInstanceRenderData->InstanceBuffer.GetData();
+		const FVector4 FourthVector(0, 0, 0, 1);
+		const int32 InstanceIndex = 0;
+		const FMatrix TransposedInstanceToLocal(
+			FPlane(
+				InstanceStream[InstanceIndex].InstanceTransform1[0].GetFloat(),
+				InstanceStream[InstanceIndex].InstanceTransform1[1].GetFloat(),
+				InstanceStream[InstanceIndex].InstanceTransform1[2].GetFloat(),
+				InstanceStream[InstanceIndex].InstanceOrigin.X),
+			FPlane(
+				InstanceStream[InstanceIndex].InstanceTransform2[0].GetFloat(),
+				InstanceStream[InstanceIndex].InstanceTransform2[1].GetFloat(),
+				InstanceStream[InstanceIndex].InstanceTransform2[2].GetFloat(),
+				InstanceStream[InstanceIndex].InstanceOrigin.Y),
+			FPlane(
+				InstanceStream[InstanceIndex].InstanceTransform3[0].GetFloat(),
+				InstanceStream[InstanceIndex].InstanceTransform3[1].GetFloat(),
+				InstanceStream[InstanceIndex].InstanceTransform3[2].GetFloat(),
+				InstanceStream[InstanceIndex].InstanceOrigin.Z),
+			(const FPlane&)FourthVector);
+
+		const FMatrix InstanceTransform = TransposedInstanceToLocal.GetTransposed() * GetLocalToWorld();
+		const FVector AxisScales = InstanceTransform.GetScaleVector();
+		const FVector BoxDimensions = RenderData->Bounds.BoxExtent * AxisScales * 2;
+
+		BoundsSurfaceArea = 2 * BoxDimensions.X * BoxDimensions.Y
+			+ 2 * BoxDimensions.Z * BoxDimensions.Y
+			+ 2 * BoxDimensions.X * BoxDimensions.Z;
+	}
+}
+
 /*-----------------------------------------------------------------------------
 	UInstancedStaticMeshComponent
 -----------------------------------------------------------------------------*/
