removed usage of GeneralPurposeTweak

diff --git a/Engine/Shaders/PostProcessAmbientOcclusion.usf b/Engine/Shaders/PostProcessAmbientOcclusion.usf
--- a/Engine/Shaders/PostProcessAmbientOcclusion.usf
+++ b/Engine/Shaders/PostProcessAmbientOcclusion.usf
@@ -537,8 +537,7 @@ void MainPS(in float4 UVAndScreenPos : TEXCOORD0, out float4 OutColor : SV_Targe
 		float3 ScaledViewSpaceNormal = ViewSpaceNormal;
 
 #if OPTIMIZATION_O1
-//		ScaledViewSpaceNormal *= 100.0f * View.GeneralPurposeTweak;
-		ScaledViewSpaceNormal *= 0.08f * lerp(SceneDepth, 1000, ScaleRadiusInWorldSpace) * View.GeneralPurposeTweak;
+		ScaledViewSpaceNormal *= 0.08f * lerp(SceneDepth, 1000, ScaleRadiusInWorldSpace);
 #endif
 
 		UNROLL for(int i = 0; i < SAMPLESET_ARRAY_SIZE; ++i)
