UE4 - Remove extra interpolator during velocity output on base pass

diff --git a/Engine/Shaders/BasePassCommon.usf b/Engine/Shaders/BasePassCommon.usf
--- a/Engine/Shaders/BasePassCommon.usf
+++ b/Engine/Shaders/BasePassCommon.usf
@@ -14,6 +14,8 @@
 #define NEEDS_LIGHTMAP					(NEEDS_LIGHTMAP_COORDINATE)
 
 #define WRITES_VELOCITY_TO_GBUFFER		(OUTPUT_GBUFFER_VELOCITY && (MATERIALBLENDING_SOLID || MATERIALBLENDING_MASKED))
+// Use an extra interpolator for position
+#define WRITES_VELOCITY_TO_GBUFFER_USE_POS_INTERPOLATOR		(WRITES_VELOCITY_TO_GBUFFER && USING_TESSELLATION)
 
 struct FSharedBasePassInterpolants
 {
@@ -32,8 +34,10 @@ struct FSharedBasePassInterpolants
 
 #if WRITES_VELOCITY_TO_GBUFFER
 	// .xy is clip position, pre divide by w; .w is clip W; .z is 0 or 1 to mask out the velocity output
-	float4 PrevScreenPosition : TEXCOORD16;
-	float4 ScreenPosition : TEXCOORD17;
+	float4 VelocityPrevScreenPosition : VELOCITY_PREV_POS;
+	#if WRITES_VELOCITY_TO_GBUFFER_USE_POS_INTERPOLATOR
+		float4 VelocityScreenPosition : VELOCITY_POS;
+	#endif
 #endif
 };
 
@@ -44,17 +48,17 @@ struct FBasePassInterpolantsVSToPS : FSharedBasePassInterpolants
 {
 	float4 PixelPosition	: TEXCOORD8;
 
-#if USE_WORLD_POSITION_EXCLUDING_SHADER_OFFSETS
+	#if USE_WORLD_POSITION_EXCLUDING_SHADER_OFFSETS
 	float4 PixelPositionExcludingWPO : TEXCOORD9;
-#endif
+	#endif
 };
 
 /** Interpolants passed from the vertex shader to the domain shader. */
 struct FBasePassInterpolantsVSToDS : FSharedBasePassInterpolants
 {
-#if USE_WORLD_POSITION_EXCLUDING_SHADER_OFFSETS
+	#if USE_WORLD_POSITION_EXCLUDING_SHADER_OFFSETS
 	float3 WorldPositionExcludingWPO : TEXCOORD9;
-#endif
+	#endif
 };
 
 #else
diff --git a/Engine/Shaders/BasePassPixelShader.usf b/Engine/Shaders/BasePassPixelShader.usf
--- a/Engine/Shaders/BasePassPixelShader.usf
+++ b/Engine/Shaders/BasePassPixelShader.usf
@@ -859,10 +859,14 @@ void Main(
 		#if WRITES_VELOCITY_TO_GBUFFER
 			{
 				// 2d velocity, includes camera an object motion
-				float2 Velocity = Calculate2DVelocity(BasePassInterpolants.ScreenPosition, BasePassInterpolants.PrevScreenPosition);
+			#if WRITES_VELOCITY_TO_GBUFFER_USE_POS_INTERPOLATOR
+				float2 Velocity = Calculate2DVelocity(BasePassInterpolants.VelocityScreenPosition, BasePassInterpolants.VelocityPrevScreenPosition);
+			#else
+				float2 Velocity = Calculate2DVelocity(MaterialParameters.ScreenPosition, BasePassInterpolants.VelocityPrevScreenPosition);
+			#endif
 
 				// Make sure not to touch 0,0 which is clear color
-				GBuffer.Velocity = float4(EncodeVelocityToTexture(Velocity), 0, 0) * BasePassInterpolants.PrevScreenPosition.z;
+				GBuffer.Velocity = float4(EncodeVelocityToTexture(Velocity), 0, 0) * BasePassInterpolants.VelocityPrevScreenPosition.z;
 			}
 		#else
 			GBuffer.Velocity = 0;
diff --git a/Engine/Shaders/BasePassTessellationShaders.usf b/Engine/Shaders/BasePassTessellationShaders.usf
--- a/Engine/Shaders/BasePassTessellationShaders.usf
+++ b/Engine/Shaders/BasePassTessellationShaders.usf
@@ -22,8 +22,10 @@ FBasePassVSToDS PassInterpolate(FBasePassVSToDS a, float aInterp, FBasePassVSToD
 	#endif
 
 	#if WRITES_VELOCITY_TO_GBUFFER
-		TESSELLATION_INTERPOLATE_MEMBER(BasePassInterpolants.PrevScreenPosition);
-		TESSELLATION_INTERPOLATE_MEMBER(BasePassInterpolants.ScreenPosition);
+		TESSELLATION_INTERPOLATE_MEMBER(BasePassInterpolants.VelocityPrevScreenPosition);
+		#if WRITES_VELOCITY_TO_GBUFFER_USE_POS_INTERPOLATOR
+		TESSELLATION_INTERPOLATE_MEMBER(BasePassInterpolants.VelocityScreenPosition);
+		#endif
 	#endif
 
 	return O;
@@ -43,11 +45,13 @@ FBasePassVSToPS PassFinalizeTessellationOutput(FBasePassVSToDS Interpolants, flo
 	{
 		O.Position = mul(WorldPosition, View.TranslatedWorldToClip);
 #if WRITES_VELOCITY_TO_GBUFFER
-		float SkipVelocity = O.BasePassInterpolants.PrevScreenPosition.w;
-		O.BasePassInterpolants.PrevScreenPosition = mul(float4(Interpolants.BasePassInterpolants.PrevScreenPosition.xyz, 1), View.PrevTranslatedWorldToClip);
+		float SkipVelocity = O.BasePassInterpolants.VelocityPrevScreenPosition.w;
+		O.BasePassInterpolants.VelocityPrevScreenPosition = mul(float4(Interpolants.BasePassInterpolants.VelocityPrevScreenPosition.xyz, 1), View.PrevTranslatedWorldToClip);
 		// Move the mask to .z when as we don't need the clip space Z
-		O.BasePassInterpolants.PrevScreenPosition.z = SkipVelocity;
-		O.BasePassInterpolants.ScreenPosition = mul(float4(Interpolants.BasePassInterpolants.ScreenPosition.xyz, 1), View.TranslatedWorldToClip);
+		O.BasePassInterpolants.VelocityPrevScreenPosition.z = SkipVelocity;
+		#if WRITES_VELOCITY_TO_GBUFFER_USE_POS_INTERPOLATOR
+		O.BasePassInterpolants.VelocityScreenPosition = mul(float4(Interpolants.BasePassInterpolants.VelocityScreenPosition.xyz, 1), View.TranslatedWorldToClip);
+		#endif
 #endif
 	}
 	
diff --git a/Engine/Shaders/BasePassVertexShader.usf b/Engine/Shaders/BasePassVertexShader.usf
--- a/Engine/Shaders/BasePassVertexShader.usf
+++ b/Engine/Shaders/BasePassVertexShader.usf
@@ -73,16 +73,20 @@ void Main(
 
 		#if USING_TESSELLATION
 			// We let the Domain Shader convert to post projection when tessellating
-			Output.BasePassInterpolants.PrevScreenPosition = PrevTranslatedWorldPosition;
+			Output.BasePassInterpolants.VelocityPrevScreenPosition = PrevTranslatedWorldPosition;
 			// Store the mask in .w when doing tessellation as we need the world space Z
-			Output.BasePassInterpolants.PrevScreenPosition.w = 1 - SkipOutputVelocity;
-			Output.BasePassInterpolants.ScreenPosition = WorldPosition;
+			Output.BasePassInterpolants.VelocityPrevScreenPosition.w = 1 - SkipOutputVelocity;
+			#if WRITES_VELOCITY_TO_GBUFFER_USE_POS_INTERPOLATOR
+			Output.BasePassInterpolants.VelocityScreenPosition = WorldPosition;
+			#endif
 		#else
 			// compute the old screen pos with the old world position and the old camera matrix
-			Output.BasePassInterpolants.PrevScreenPosition = mul(float4(PrevTranslatedWorldPosition.xyz, 1), View.PrevTranslatedWorldToClip);		
+			Output.BasePassInterpolants.VelocityPrevScreenPosition = mul(float4(PrevTranslatedWorldPosition.xyz, 1), View.PrevTranslatedWorldToClip);		
 			// Store the mask in .z when as we don't need the clip space Z
-			Output.BasePassInterpolants.PrevScreenPosition.z = 1 - SkipOutputVelocity;
-			Output.BasePassInterpolants.ScreenPosition = Output.Position;
+			Output.BasePassInterpolants.VelocityPrevScreenPosition.z = 1 - SkipOutputVelocity;
+			#if WRITES_VELOCITY_TO_GBUFFER_USE_POS_INTERPOLATOR
+			Output.BasePassInterpolants.VelocityScreenPosition = Output.Position;
+			#endif
 		#endif	// USING_TESSELLATION
 	}
 #endif	// WRITES_VELOCITY_TO_GBUFFER
