fixed shader compile warning

diff --git a/Engine/Shaders/MaterialTemplate.usf b/Engine/Shaders/MaterialTemplate.usf
--- a/Engine/Shaders/MaterialTemplate.usf
+++ b/Engine/Shaders/MaterialTemplate.usf
@@ -716,7 +716,7 @@ float3 GetSpeedTreeVertexOffset(FMaterialVertexParameters Parameters, int Geomet
 		const float ScreenMultiple = 0.5 * max(View.ViewSizeAndSceneTexelSize.x * View.ViewToClip[0][0],
 												View.ViewSizeAndSceneTexelSize.y * View.ViewToClip[1][1]);
 		const float ScreenRadius = ScreenMultiple * Primitive.ObjectWorldPositionAndRadius.w / 
-									max(1.0, dot(TreePos - View.ViewOrigin, View.ViewForward));
+									max(1.0, dot(TreePos - View.ViewOrigin.xyz, View.ViewForward));
 		const float ScreenArea = PI * ScreenRadius * ScreenRadius / (View.ViewSizeAndSceneTexelSize.x * View.ViewSizeAndSceneTexelSize.y);
 		LodInterp = saturate((ScreenArea - SpeedTreeLODInfo.x) / SpeedTreeLODInfo.z);
 	}
