#android Fix for normals on Android, making them behave the same way as most other platforms (deriving the third component in the shader). This change also makes it possible to use two-component textures on Android, like ASTC.

diff --git a/Engine/Shaders/Common.usf b/Engine/Shaders/Common.usf
--- a/Engine/Shaders/Common.usf
+++ b/Engine/Shaders/Common.usf
@@ -818,7 +818,7 @@ MaterialFloat2 GetViewportCoordinates(MaterialFloat2 InFragmentCoordinates)
  */
 MaterialFloat4 UnpackNormalMap( MaterialFloat4 TextureSample )
 {
-#if COMPILER_GLSL_ES2
+#if COMPILER_GLSL_ES2 && IOS
 	return MaterialFloat4(TextureSample.rgb * 2 - 1, 1);
 #else
 	#if DXT5_NORMALMAPS
