fixed engery conservation for 2 sided folige, backscatter get slightly darker

diff --git a/Engine/Shaders/DeferredLightingCommon.usf b/Engine/Shaders/DeferredLightingCommon.usf
--- a/Engine/Shaders/DeferredLightingCommon.usf
+++ b/Engine/Shaders/DeferredLightingCommon.usf
@@ -251,9 +251,11 @@ float3 SubsurfaceShadingSubsurface( FGBufferData GBuffer, float3 L, float3 V, ha
 float3 TwoSidedFoliageShadingSubsurface( float3 SubsurfaceColor, float3 L, float3 V, half3 N )
 {
 	float VoL = saturate(dot(V, -L));
+	// here we raise to the power of 4
 	float VoLLobe = VoL * VoL;
 	VoLLobe *= VoLLobe;
-	float Exponent = 8;
+	// so for energy conservation we need to do the same here
+	const float Exponent = 4;
 	// Modified phong energy conservation
 	VoLLobe = VoLLobe * (Exponent + 1) / (2.0f * PI);
 
