UE4 - Speedtree now calculates previous world pos (incurs extra vertex cost)

diff --git a/Engine/Shaders/MaterialTemplate.usf b/Engine/Shaders/MaterialTemplate.usf
--- a/Engine/Shaders/MaterialTemplate.usf
+++ b/Engine/Shaders/MaterialTemplate.usf
@@ -686,7 +686,7 @@ float GetSphericalParticleOpacity(FMaterialPixelParameters Parameters, float Den
 #if USES_SPEEDTREE
 
 /** Vertex offset for SpeedTree wind and LOD */
-float3 GetSpeedTreeVertexOffset(FMaterialVertexParameters Parameters, int GeometryType, int WindType, int LODType, float BillboardThreshold) 
+float3 GetSpeedTreeVertexOffsetInner(FMaterialVertexParameters Parameters, int GeometryType, int WindType, int LODType, float BillboardThreshold, FSpeedTreeData STData) 
 {
 	#if (NUM_MATERIAL_TEXCOORDS_VERTEX < 6) || PARTICLE_MESH_FACTORY
 		return float4(0,0,0);
@@ -773,7 +773,7 @@ float3 GetSpeedTreeVertexOffset(FMaterialVertexParameters Parameters, int Geomet
 			{
 				float2 TexCoords = Parameters.TexCoords[0];
 				float4 WindExtra = float4(Parameters.TexCoords[5].x, Parameters.TexCoords[5].y, Parameters.TexCoords[6].x, 0.0);
-				FinalPosition = RippleFrond(FinalPosition, Normal, TexCoords.x, TexCoords.y, WindExtra.x, WindExtra.y, WindExtra.z);
+				FinalPosition = RippleFrond(STData, FinalPosition, Normal, TexCoords.x, TexCoords.y, WindExtra.x, WindExtra.y, WindExtra.z);
 			}
 		}
 
@@ -819,7 +819,7 @@ float3 GetSpeedTreeVertexOffset(FMaterialVertexParameters Parameters, int Geomet
 			{
 				float4 WindExtra = float4(Parameters.TexCoords[6].x, Parameters.TexCoords[6].y, Parameters.TexCoords[7].x, Parameters.TexCoords[7].y);
 				float LeafWindTrigOffset = Anchor.x + Anchor.y;
-				FinalPosition = LeafWind(WindExtra.w > 0.0, FinalPosition, Normal, WindExtra.x, float3(0,0,0), WindExtra.y, WindExtra.z, LeafWindTrigOffset, WindType);
+				FinalPosition = LeafWind(STData, WindExtra.w > 0.0, FinalPosition, Normal, WindExtra.x, float3(0,0,0), WindExtra.y, WindExtra.z, LeafWindTrigOffset, WindType);
 			}
 				
 			// move leaf back to anchor
@@ -830,20 +830,30 @@ float3 GetSpeedTreeVertexOffset(FMaterialVertexParameters Parameters, int Geomet
 		{
 			// branch wind (applies to all geometry)
 			float2 VertBranchWind = Parameters.TexCoords[2];
-			FinalPosition = BranchWind(FinalPosition, TreePos, float4(VertBranchWind, 0, 0), WindType);
+			FinalPosition = BranchWind(STData, FinalPosition, TreePos, float4(VertBranchWind, 0, 0), WindType);
 		}	
 	}
 
 	// global wind can apply to the whole tree, even billboards
 	if (WindType != SPEEDTREE_WIND_TYPE_NONE)
 	{
-		FinalPosition = GlobalWind(FinalPosition, TreePos, true);
+		FinalPosition = GlobalWind(STData, FinalPosition, TreePos, true);
 	}
 
 	// convert into a world space offset
 	return (FinalPosition - OriginalPosition) * TreeScale;
 }
 
+/** Vertex offset for SpeedTree wind and LOD */
+float3 GetSpeedTreeVertexOffset(FMaterialVertexParameters Parameters, int GeometryType, int WindType, int LODType, float BillboardThreshold, bool bUsePreviousFrame) 
+{
+	if (bUsePreviousFrame)
+	{
+		return GetSpeedTreeVertexOffsetInner(Parameters, GeometryType, WindType, LODType, BillboardThreshold, GetPreviousSpeedTreeData());
+	}
+	return GetSpeedTreeVertexOffsetInner(Parameters, GeometryType, WindType, LODType, BillboardThreshold, GetCurrentSpeedTreeData());
+}
+
 #endif
 
 MaterialFloat2 GetLightmapUVs(FMaterialPixelParameters Parameters)
diff --git a/Engine/Shaders/ShaderVersion.usf b/Engine/Shaders/ShaderVersion.usf
--- a/Engine/Shaders/ShaderVersion.usf
+++ b/Engine/Shaders/ShaderVersion.usf
@@ -2,4 +2,4 @@
 // Each time the console command is executed it generates a new GUID. As this file is included
 // in common.usf (which should be included in any shader) it allows to invalidate the shader DDC.
 // 
-// GUID = {2490AEE3-37E5-4684-93D8-F170F910138B}
+// GUID = {2490AEE3-37E5-4684-93D8-F170F910138C}
diff --git a/Engine/Shaders/SpeedTreeCommon.usf b/Engine/Shaders/SpeedTreeCommon.usf
--- a/Engine/Shaders/SpeedTreeCommon.usf
+++ b/Engine/Shaders/SpeedTreeCommon.usf
@@ -6,7 +6,83 @@
 
  #include "UniformBuffers/SpeedTreeData.usf"
 
- float3 SpeedTreeLODInfo;
+float3 SpeedTreeLODInfo;
+
+// Make sure this matches the SpeedTreeData uniform buffer
+struct FSpeedTreeData
+{
+	float4 WindVector;
+	float4 WindGlobal;
+	float4 WindBranch;
+	float4 WindBranchTwitch;
+	float4 WindBranchWhip;
+	float4 WindBranchAnchor;
+	float4 WindBranchAdherences;
+	float4 WindTurbulences;
+	float4 WindLeaf1Ripple;
+	float4 WindLeaf1Tumble;
+	float4 WindLeaf1Twitch;
+	float4 WindLeaf2Ripple;
+	float4 WindLeaf2Tumble;
+	float4 WindLeaf2Twitch;
+	float4 WindFrondRipple;
+	float4 WindRollingBranch;
+	float4 WindRollingLeafAndDirection;
+	float4 WindRollingNoise;
+	float4 WindAnimation;
+};
+
+FSpeedTreeData GetCurrentSpeedTreeData()
+{
+	// Make sure these matches all members in FSpeedTreeData
+	FSpeedTreeData STData;
+	STData.WindVector = SpeedTreeData.WindVector;
+	STData.WindGlobal = SpeedTreeData.WindGlobal;
+	STData.WindBranch = SpeedTreeData.WindBranch;
+	STData.WindBranchTwitch = SpeedTreeData.WindBranchTwitch;
+	STData.WindBranchWhip = SpeedTreeData.WindBranchWhip;
+	STData.WindBranchAnchor = SpeedTreeData.WindBranchAnchor;
+	STData.WindBranchAdherences = SpeedTreeData.WindBranchAdherences;
+	STData.WindTurbulences = SpeedTreeData.WindTurbulences;
+	STData.WindLeaf1Ripple = SpeedTreeData.WindLeaf1Ripple;
+	STData.WindLeaf1Tumble = SpeedTreeData.WindLeaf1Tumble;
+	STData.WindLeaf1Twitch = SpeedTreeData.WindLeaf1Twitch;
+	STData.WindLeaf2Ripple = SpeedTreeData.WindLeaf2Ripple;
+	STData.WindLeaf2Tumble = SpeedTreeData.WindLeaf2Tumble;
+	STData.WindLeaf2Twitch = SpeedTreeData.WindLeaf2Twitch;
+	STData.WindFrondRipple = SpeedTreeData.WindFrondRipple;
+	STData.WindRollingBranch = SpeedTreeData.WindRollingBranch;
+	STData.WindRollingLeafAndDirection = SpeedTreeData.WindRollingLeafAndDirection;
+	STData.WindRollingNoise = SpeedTreeData.WindRollingNoise;
+	STData.WindAnimation = SpeedTreeData.WindAnimation;
+	return STData;
+}
+
+FSpeedTreeData GetPreviousSpeedTreeData()
+{
+	// Same as GetCurrentSpeedTreeData(), but using the 'Prev' members
+	FSpeedTreeData STData;
+	STData.WindVector = SpeedTreeData.PrevWindVector;
+	STData.WindGlobal = SpeedTreeData.PrevWindGlobal;
+	STData.WindBranch = SpeedTreeData.PrevWindBranch;
+	STData.WindBranchTwitch = SpeedTreeData.PrevWindBranchTwitch;
+	STData.WindBranchWhip = SpeedTreeData.PrevWindBranchWhip;
+	STData.WindBranchAnchor = SpeedTreeData.PrevWindBranchAnchor;
+	STData.WindBranchAdherences = SpeedTreeData.PrevWindBranchAdherences;
+	STData.WindTurbulences = SpeedTreeData.PrevWindTurbulences;
+	STData.WindLeaf1Ripple = SpeedTreeData.PrevWindLeaf1Ripple;
+	STData.WindLeaf1Tumble = SpeedTreeData.PrevWindLeaf1Tumble;
+	STData.WindLeaf1Twitch = SpeedTreeData.PrevWindLeaf1Twitch;
+	STData.WindLeaf2Ripple = SpeedTreeData.PrevWindLeaf2Ripple;
+	STData.WindLeaf2Tumble = SpeedTreeData.PrevWindLeaf2Tumble;
+	STData.WindLeaf2Twitch = SpeedTreeData.PrevWindLeaf2Twitch;
+	STData.WindFrondRipple = SpeedTreeData.PrevWindFrondRipple;
+	STData.WindRollingBranch = SpeedTreeData.PrevWindRollingBranch;
+	STData.WindRollingLeafAndDirection = SpeedTreeData.PrevWindRollingLeafAndDirection;
+	STData.WindRollingNoise = SpeedTreeData.PrevWindRollingNoise;
+	STData.WindAnimation = SpeedTreeData.PrevWindAnimation;
+	return STData;
+}
 
 ///////////////////////////////////////////////////////////////////////  
 //  SpeedTree v7 Wind
@@ -143,15 +219,16 @@ float3 mul_float3x3_float3(float3x3 mMatrix, float3 vVector)
 ///////////////////////////////////////////////////////////////////////  
 //  Roll
 
-float Roll(float fCurrent,
-		   float fMaxScale, 
-		   float fMinScale, 
-		   float fSpeed, 
-		   float fRipple, 
-		   float3 vPos, 
-		   float fTime)
+float Roll(FSpeedTreeData STData,
+	float fCurrent,
+	float fMaxScale,
+	float fMinScale,
+	float fSpeed,
+	float fRipple,
+	float3 vPos,
+	float fTime)
 {
-	float fWindAngle = dot(vPos, -SpeedTreeData.WindVector.xyz) * fRipple;
+	float fWindAngle = dot(vPos, -STData.WindVector.xyz) * fRipple;
 	float fAdjust = TrigApproximate(float4(fWindAngle + fTime * fSpeed, 0.0, 0.0, 0.0)).x;
 	fAdjust = (fAdjust + 1.0) * 0.5;
 	
@@ -182,17 +259,18 @@ float Twitch(float3 vPos, float fAmount, float fSharpness, float fTime)
 //  Whip and oscillation are combined like this to minimize calls to
 //  TrigApproximate( ) when possible.
 
-float Oscillate(float3 vPos, 
-				float fTime, 
-				float fOffset, 
-				float fWeight, 
-				float fWhip, 
-				bool bWhip, 
-				bool bRoll, 
-				bool bComplex, 
-				float fTwitch, 
-				float fTwitchFreqScale,
-				inout float4 vOscillations)
+float Oscillate(FSpeedTreeData STData,
+	float3 vPos, 
+	float fTime, 
+	float fOffset, 
+	float fWeight, 
+	float fWhip, 
+	bool bWhip, 
+	bool bRoll, 
+	bool bComplex, 
+	float fTwitch, 
+	float fTwitchFreqScale,
+	inout float4 vOscillations)
 {
 	float fOscillation = 1.0;
 	if (bComplex)
@@ -216,7 +294,7 @@ float Oscillate(float3 vPos,
 		fBroadDetail = lerp(fBroadDetail, fTarget, fAmount);
 		fBroadDetail = lerp(fBroadDetail, fTarget, fAmount);
 
-		fOscillation = fBroadDetail * fTwitch * (1.0 - SpeedTreeData.WindVector.w) + fFineDetail * (1.0 - fTwitch);
+		fOscillation = fBroadDetail * fTwitch * (1.0 - STData.WindVector.w) + fFineDetail * (1.0 - fTwitch);
 
 		if (bWhip)
 			fOscillation *= 1.0 + (vOscillations.w * fWhip);
@@ -236,7 +314,7 @@ float Oscillate(float3 vPos,
 
 	//if (bRoll)
 	//{
-	// 	fOscillation = Roll(fOscillation, SpeedTreeData.WindRollingBranches.x, SpeedTreeData.WindRollingBranches.y, SpeedTreeData.WindRollingBranches.z, SpeedTreeData.WindRollingBranches.w, vPos.xyz, fTime + fOffset);
+	// 	fOscillation = Roll(fOscillation, STData.WindRollingBranches.x, STData.WindRollingBranches.y, STData.WindRollingBranches.z, STData.WindRollingBranches.w, vPos.xyz, fTime + fOffset);
 	//}
 
 	return fOscillation;
@@ -262,7 +340,7 @@ float Turbulence(float fTime, float fOffset, float fGlobalTime, float fTurbulenc
 //  This function positions any tree geometry based on their untransformed 
 //	position and 4 wind floats.
 
-float3 GlobalWind(float3 vPos, float3 vInstancePos, bool bPreserveShape)
+float3 GlobalWind(FSpeedTreeData STData, float3 vPos, float3 vInstancePos, bool bPreserveShape)
 {
 	// WIND_LOD_GLOBAL may be on, but if the global wind effect (WIND_EFFECT_GLOBAL_WIND)
 	// was disabled for the tree in the Modeler, we should skip it
@@ -273,22 +351,22 @@ float3 GlobalWind(float3 vPos, float3 vInstancePos, bool bPreserveShape)
 
 	// compute how much the height contributes
 	#ifdef SPEEDTREE_Z_UP
-		float fAdjust = max(vPos.z - (1.0 / SpeedTreeData.WindGlobal.z) * 0.25, 0.0) * SpeedTreeData.WindGlobal.z;
+		float fAdjust = max(vPos.z - (1.0 / STData.WindGlobal.z) * 0.25, 0.0) * STData.WindGlobal.z;
 	#else
-		float fAdjust = max(vPos.y - (1.0 / SpeedTreeData.WindGlobal.z) * 0.25, 0.0) * SpeedTreeData.WindGlobal.z;
+		float fAdjust = max(vPos.y - (1.0 / STData.WindGlobal.z) * 0.25, 0.0) * STData.WindGlobal.z;
 	#endif
 	if (fAdjust != 0.0)
-		fAdjust = pow(fAdjust, SpeedTreeData.WindGlobal.w);
+		fAdjust = pow(fAdjust, STData.WindGlobal.w);
 
 	// primary oscillation
-	float4 vOscillations = TrigApproximate(float4(vInstancePos.x + SpeedTreeData.WindGlobal.x, vInstancePos.y + SpeedTreeData.WindGlobal.x * 0.8, 0.0, 0.0));
+	float4 vOscillations = TrigApproximate(float4(vInstancePos.x + STData.WindGlobal.x, vInstancePos.y + STData.WindGlobal.x * 0.8, 0.0, 0.0));
 	float fOsc = vOscillations.x + (vOscillations.y * vOscillations.y);
-	float fMoveAmount = SpeedTreeData.WindGlobal.y * fOsc;
+	float fMoveAmount = STData.WindGlobal.y * fOsc;
 
 	// move a minimum amount based on direction adherence
-	float fMinMove = SpeedTreeData.WindBranchAdherences.x;
-	if (SpeedTreeData.WindGlobal.z != 0.0) // not sure why I need this check
-		fMinMove /= SpeedTreeData.WindGlobal.z;
+	float fMinMove = STData.WindBranchAdherences.x;
+	if (STData.WindGlobal.z != 0.0) // not sure why I need this check
+		fMinMove /= STData.WindGlobal.z;
 	fMoveAmount += fMinMove;
 
 	// adjust based on how high up the tree this vertex is
@@ -296,9 +374,9 @@ float3 GlobalWind(float3 vPos, float3 vInstancePos, bool bPreserveShape)
 
 	// xy component
 	#ifdef SPEEDTREE_Z_UP
-		vPos.xy += SpeedTreeData.WindVector.xy * fMoveAmount;
+		vPos.xy += STData.WindVector.xy * fMoveAmount;
 	#else
-		vPos.xz += SpeedTreeData.WindVector.xz * fMoveAmount;
+		vPos.xz += STData.WindVector.xz * fMoveAmount;
 	#endif
 
 	if (bPreserveShape)
@@ -312,18 +390,20 @@ float3 GlobalWind(float3 vPos, float3 vInstancePos, bool bPreserveShape)
 ///////////////////////////////////////////////////////////////////////  
 //  SimpleBranchWind
 
-float3 SimpleBranchWind(float3 vPos,
-						float3 vInstancePos,
-						float fWeight, 
-						float fOffset, 
-						float fTime, 
-						float fDistance, 
-						float fTwitch, 
-						float fTwitchScale, 
-						float fWhip, 
-						bool bWhip, 
-						bool bRoll, 
-						bool bComplex)
+float3 SimpleBranchWind(
+	FSpeedTreeData STData,
+	float3 vPos,
+	float3 vInstancePos,
+	float fWeight, 
+	float fOffset, 
+	float fTime, 
+	float fDistance, 
+	float fTwitch, 
+	float fTwitchScale, 
+	float fWhip, 
+	bool bWhip, 
+	bool bRoll, 
+	bool bComplex)
 {
 	// turn the offset back into a nearly normalized vector
 	float3 vWindVector = UnpackNormalFromFloat(fOffset);
@@ -334,7 +414,7 @@ float3 SimpleBranchWind(float3 vPos,
 
 	// oscillate
 	float4 vOscillations;
-	float fOsc = Oscillate(vPos, fTime, fOffset, fWeight, fWhip, bWhip, bRoll, bComplex, fTwitch, fTwitchScale, vOscillations);
+	float fOsc = Oscillate(STData, vPos, fTime, fOffset, fWeight, fWhip, bWhip, bRoll, bComplex, fTwitch, fTwitchScale, vOscillations);
 
 	vPos.xyz += vWindVector * fOsc * fDistance;
 
@@ -345,21 +425,22 @@ float3 SimpleBranchWind(float3 vPos,
 ///////////////////////////////////////////////////////////////////////  
 //  DirectionalBranchWind
 
-float3 DirectionalBranchWind(float3 vPos, 
-							 float3 vInstancePos,
-							 float fWeight, 
-							 float fOffset, 
-							 float fTime, 
-							 float fDistance, 
-							 float fTurbulence, 
-							 float fAdherence, 
-							 float fTwitch, 
-							 float fTwitchScale, 
-							 float fWhip, 
-							 bool bWhip, 
-							 bool bRoll, 
-							 bool bComplex, 
-							 bool bTurbulence)
+float3 DirectionalBranchWind(FSpeedTreeData STData,
+	float3 vPos, 
+	float3 vInstancePos,
+	float fWeight, 
+	float fOffset, 
+	float fTime, 
+	float fDistance, 
+	float fTurbulence, 
+	float fAdherence, 
+	float fTwitch, 
+	float fTwitchScale, 
+	float fWhip, 
+	bool bWhip, 
+	bool bRoll, 
+	bool bComplex, 
+	bool bTurbulence)
 {
 	// turn the offset back into a nearly normalized vector
 	float3 vWindVector = UnpackNormalFromFloat(fOffset);
@@ -370,22 +451,22 @@ float3 DirectionalBranchWind(float3 vPos,
 	
 	// oscillate
 	float4 vOscillations;
-	float fOsc = Oscillate(vPos, fTime, fOffset, fWeight, fWhip, bWhip, false, bComplex, fTwitch, fTwitchScale, vOscillations);
+	float fOsc = Oscillate(STData, vPos, fTime, fOffset, fWeight, fWhip, bWhip, false, bComplex, fTwitch, fTwitchScale, vOscillations);
 
 	vPos.xyz += vWindVector * fOsc * fDistance;
 
 	// add in the direction, accounting for turbulence
 	float fAdherenceScale = 1.0;
 	if (bTurbulence)
-		fAdherenceScale = Turbulence(fTime, fOffset, SpeedTreeData.WindAnimation.x, fTurbulence);
+		fAdherenceScale = Turbulence(fTime, fOffset, STData.WindAnimation.x, fTurbulence);
 
 	if (bWhip)
-		fAdherenceScale += vOscillations.w * SpeedTreeData.WindVector.w * fWhip;
+		fAdherenceScale += vOscillations.w * STData.WindVector.w * fWhip;
 
 	//if (bRoll)
-	//	fAdherenceScale = Roll(fAdherenceScale, SpeedTreeData.WindRollingBranches.x, SpeedTreeData.WindRollingBranches.y, SpeedTreeData.WindRollingBranches.z, SpeedTreeData.WindRollingBranches.w, vPos.xyz, fTime + fOffset);
+	//	fAdherenceScale = Roll(fAdherenceScale, STData.WindRollingBranches.x, STData.WindRollingBranches.y, STData.WindRollingBranches.z, STData.WindRollingBranches.w, vPos.xyz, fTime + fOffset);
 
-	vPos.xyz += SpeedTreeData.WindVector.xyz * fAdherence * fAdherenceScale * fWeight;
+	vPos.xyz += STData.WindVector.xyz * fAdherence * fAdherenceScale * fWeight;
 
 	return vPos;
 }
@@ -394,21 +475,22 @@ float3 DirectionalBranchWind(float3 vPos,
 ///////////////////////////////////////////////////////////////////////  
 //  DirectionalBranchWindFrondStyle
 
-float3 DirectionalBranchWindFrondStyle(float3 vPos, 
-									   float3 vInstancePos,
-									   float fWeight, 
-									   float fOffset, 
-									   float fTime, 
-									   float fDistance, 
-									   float fTurbulence, 
-									   float fAdherence, 
-									   float fTwitch, 
-									   float fTwitchScale, 
-									   float fWhip, 
-									   bool bWhip, 
-									   bool bRoll, 
-									   bool bComplex, 
-									   bool bTurbulence)
+float3 DirectionalBranchWindFrondStyle(FSpeedTreeData STData,
+	float3 vPos, 
+	float3 vInstancePos,
+	float fWeight, 
+	float fOffset, 
+	float fTime, 
+	float fDistance, 
+	float fTurbulence, 
+	float fAdherence, 
+	float fTwitch, 
+	float fTwitchScale, 
+	float fWhip, 
+	bool bWhip, 
+	bool bRoll, 
+	bool bComplex, 
+	bool bTurbulence)
 {
 	// turn the offset back into a nearly normalized vector
 	float3 vWindVector = UnpackNormalFromFloat(fOffset);
@@ -419,22 +501,22 @@ float3 DirectionalBranchWindFrondStyle(float3 vPos,
 	
 	// oscillate
 	float4 vOscillations;
-	float fOsc = Oscillate(vPos, fTime, fOffset, fWeight, fWhip, bWhip, false, bComplex, fTwitch, fTwitchScale, vOscillations);
+	float fOsc = Oscillate(STData, vPos, fTime, fOffset, fWeight, fWhip, bWhip, false, bComplex, fTwitch, fTwitchScale, vOscillations);
 
 	vPos.xyz += vWindVector * fOsc * fDistance;
 
 	// add in the direction, accounting for turbulence
 	float fAdherenceScale = 1.0;
 	if (bTurbulence)
-		fAdherenceScale = Turbulence(fTime, fOffset, SpeedTreeData.WindAnimation.x, fTurbulence);
+		fAdherenceScale = Turbulence(fTime, fOffset, STData.WindAnimation.x, fTurbulence);
 
 	//if (bRoll)
-	//	fAdherenceScale = Roll(fAdherenceScale, SpeedTreeData.WindRollingBranches.x, SpeedTreeData.WindRollingBranches.y, SpeedTreeData.WindRollingBranches.z, SpeedTreeData.WindRollingBranches.w, vPos.xyz, fTime + fOffset);
+	//	fAdherenceScale = Roll(fAdherenceScale, STData.WindRollingBranches.x, STData.WindRollingBranches.y, STData.WindRollingBranches.z, STData.WindRollingBranches.w, vPos.xyz, fTime + fOffset);
 
 	if (bWhip)
-		fAdherenceScale += vOscillations.w * SpeedTreeData.WindVector.w * fWhip;
+		fAdherenceScale += vOscillations.w * STData.WindVector.w * fWhip;
 
-	float3 vWindAdherenceVector = SpeedTreeData.WindBranchAnchor.xyz - vPos.xyz;
+	float3 vWindAdherenceVector = STData.WindBranchAnchor.xyz - vPos.xyz;
 	vPos.xyz += vWindAdherenceVector * fAdherence * fAdherenceScale * fWeight;
 
 	return vPos;
@@ -444,15 +526,15 @@ float3 DirectionalBranchWindFrondStyle(float3 vPos,
 ///////////////////////////////////////////////////////////////////////  
 //  BranchWind
 
-float3 BranchWind(float3 vPos, float3 vInstancePos, float4 vWindData, int WindType)
+float3 BranchWind(FSpeedTreeData STData, float3 vPos, float3 vInstancePos, float4 vWindData, int WindType)
 {
 	if (WindType == SPEEDTREE_WIND_TYPE_PALM)
 	{
-		vPos = DirectionalBranchWindFrondStyle(vPos, vInstancePos, vWindData.x, vWindData.y, SpeedTreeData.WindBranch.x, SpeedTreeData.WindBranch.y, SpeedTreeData.WindTurbulences.x, SpeedTreeData.WindBranchAdherences.y, SpeedTreeData.WindBranchTwitch.x, SpeedTreeData.WindBranchTwitch.y, SpeedTreeData.WindBranchWhip.x, true, false, true, true);
+		vPos = DirectionalBranchWindFrondStyle(STData, vPos, vInstancePos, vWindData.x, vWindData.y, STData.WindBranch.x, STData.WindBranch.y, STData.WindTurbulences.x, STData.WindBranchAdherences.y, STData.WindBranchTwitch.x, STData.WindBranchTwitch.y, STData.WindBranchWhip.x, true, false, true, true);
 	}
 	else
 	{
-		vPos = SimpleBranchWind(vPos, vInstancePos, vWindData.x, vWindData.y, SpeedTreeData.WindBranch.x, SpeedTreeData.WindBranch.y, SpeedTreeData.WindBranchTwitch.x, SpeedTreeData.WindBranchTwitch.y, SpeedTreeData.WindBranchWhip.x, false, false, true);
+		vPos = SimpleBranchWind(STData, vPos, vInstancePos, vWindData.x, vWindData.y, STData.WindBranch.x, STData.WindBranch.y, STData.WindBranchTwitch.x, STData.WindBranchTwitch.y, STData.WindBranchWhip.x, false, false, true);
 	}
 
 	return vPos;
@@ -492,20 +574,21 @@ float3 LeafRipple(float3 vPos,
 ///////////////////////////////////////////////////////////////////////  
 //  LeafTumble
 
-float3 LeafTumble(float3 vPos, 
-				  inout float3 vDirection, 
-				  float fScale, 
-				  float3 vAnchor, 
-				  float3 vGrowthDir, 
-				  float fTrigOffset, 
-				  float fTime, 
-				  float fFlip, 
-				  float fTwist, 
-				  float fAdherence, 
-				  float3 vTwitch, 
-				  float4 vRoll, 
-				  bool bTwitch, 
-				  bool bRoll)
+float3 LeafTumble(FSpeedTreeData STData,
+	float3 vPos, 
+	inout float3 vDirection, 
+	float fScale, 
+	float3 vAnchor, 
+	float3 vGrowthDir, 
+	float fTrigOffset, 
+	float fTime, 
+	float fFlip, 
+	float fTwist, 
+	float fAdherence, 
+	float3 vTwitch, 
+	float4 vRoll, 
+	bool bTwitch, 
+	bool bRoll)
 {
 	// compute all oscillations up front
  	float3 vFracs = frac((vAnchor + fTrigOffset) * 30.3);
@@ -521,8 +604,8 @@ float3 LeafTumble(float3 vPos,
  	float3x3 matTumble = RotationMatrix(vGrowthDir, fScale * fTwist * fOsc);
 
 	// with wind
-	float3 vAxis = wind_cross(vGrowthDir, SpeedTreeData.WindVector.xyz);
-	float fDot = clamp(dot(SpeedTreeData.WindVector.xyz, vGrowthDir), -1.0, 1.0);
+	float3 vAxis = wind_cross(vGrowthDir, STData.WindVector.xyz);
+	float fDot = clamp(dot(STData.WindVector.xyz, vGrowthDir), -1.0, 1.0);
 	#ifdef SPEEDTREE_Z_UP
 		vAxis.z += fDot;
 	#else
@@ -557,14 +640,15 @@ float3 LeafTumble(float3 vPos,
 ///////////////////////////////////////////////////////////////////////  
 //  RippleFrondOneSided
 
-float3 RippleFrondOneSided(float3 vPos, 
-						   inout float3 vDirection, 
-						   float fU, 
-						   float fV, 
-						   float fRippleScale
+float3 RippleFrondOneSided(FSpeedTreeData STData,
+	float3 vPos, 
+	inout float3 vDirection, 
+	float fU, 
+	float fV, 
+	float fRippleScale
 #ifdef WIND_EFFECT_FROND_RIPPLE_ADJUST_LIGHTING
-						   , float3 vBinormal
-						   , float3 vTangent
+	, float3 vBinormal
+	, float3 vTangent
 #endif
 						   )
 {
@@ -572,14 +656,14 @@ float3 RippleFrondOneSided(float3 vPos,
 	if (fU < 0.5)
 		fOffset = 0.75;
 
-	float4 vOscillations = TrigApproximate(float4((SpeedTreeData.WindFrondRipple.x + fV) * SpeedTreeData.WindFrondRipple.z + fOffset, 0.0, 0.0, 0.0));
+	float4 vOscillations = TrigApproximate(float4((STData.WindFrondRipple.x + fV) * STData.WindFrondRipple.z + fOffset, 0.0, 0.0, 0.0));
 
-	float fAmount = fRippleScale * vOscillations.x * SpeedTreeData.WindFrondRipple.y;
+	float fAmount = fRippleScale * vOscillations.x * STData.WindFrondRipple.y;
 	float3 vOffset = fAmount * vDirection;
 	vPos.xyz += vOffset;
 
 	#ifdef WIND_EFFECT_FROND_RIPPLE_ADJUST_LIGHTING
-		vTangent.xyz = normalize(vTangent.xyz + vOffset * SpeedTreeData.WindFrondRipple.w);
+		vTangent.xyz = normalize(vTangent.xyz + vOffset * STData.WindFrondRipple.w);
 		float3 vNewNormal = normalize(wind_cross(vBinormal.xyz, vTangent.xyz));
 		if (dot(vNewNormal, vDirection.xyz) < 0.0)
 			vNewNormal = -vNewNormal;
@@ -592,29 +676,30 @@ float3 RippleFrondOneSided(float3 vPos,
 ///////////////////////////////////////////////////////////////////////  
 //  RippleFrondTwoSided
 
-float3 RippleFrondTwoSided(float3 vPos, 
-						   inout float3 vDirection, 
-						   float fU, 
-						   float fLengthPercent, 
-						   float fPackedRippleDir, 
-						   float fRippleScale
+float3 RippleFrondTwoSided(FSpeedTreeData STData,
+	float3 vPos, 
+	inout float3 vDirection, 
+	float fU, 
+	float fLengthPercent, 
+	float fPackedRippleDir, 
+	float fRippleScale
 #ifdef WIND_EFFECT_FROND_RIPPLE_ADJUST_LIGHTING
-						   , float3 vBinormal
-						   , float3 vTangent
+	, float3 vBinormal
+	, float3 vTangent
 #endif
 						   )
 {
-	float4 vOscillations = TrigApproximate(float4(SpeedTreeData.WindFrondRipple.x * fLengthPercent * SpeedTreeData.WindFrondRipple.z, 0.0, 0.0, 0.0));
+	float4 vOscillations = TrigApproximate(float4(STData.WindFrondRipple.x * fLengthPercent * STData.WindFrondRipple.z, 0.0, 0.0, 0.0));
 
 	float3 vRippleDir = UnpackNormalFromFloat(fPackedRippleDir);
 
-	float fAmount = fRippleScale * vOscillations.x * SpeedTreeData.WindFrondRipple.y;
+	float fAmount = fRippleScale * vOscillations.x * STData.WindFrondRipple.y;
 	float3 vOffset = fAmount * vRippleDir;
 
 	vPos.xyz += vOffset;
 
 	#ifdef WIND_EFFECT_FROND_RIPPLE_ADJUST_LIGHTING
-		vTangent.xyz = normalize(vTangent.xyz + vOffset * SpeedTreeData.WindFrondRipple.w);
+		vTangent.xyz = normalize(vTangent.xyz + vOffset * STData.WindFrondRipple.w);
 		float3 vNewNormal = normalize(wind_cross(vBinormal.xyz, vTangent.xyz));
 		if (dot(vNewNormal, vDirection.xyz) < 0.0)
 			vNewNormal = -vNewNormal;
@@ -628,28 +713,29 @@ float3 RippleFrondTwoSided(float3 vPos,
 ///////////////////////////////////////////////////////////////////////  
 //  RippleFrond
 
-float3 RippleFrond(float3 vPos, 
-				   inout float3 vDirection,
-				   float fU, 
-				   float fV, 
-				   float fPackedRippleDir, 
-				   float fRippleScale, 
-				   float fLenghtPercent
-				#ifdef WIND_EFFECT_FROND_RIPPLE_ADJUST_LIGHTING
-				   , float3 vBinormal
-				   , float3 vTangent
-				#endif
+float3 RippleFrond(FSpeedTreeData STData,
+	float3 vPos, 
+	inout float3 vDirection,
+	float fU, 
+	float fV, 
+	float fPackedRippleDir, 
+	float fRippleScale, 
+	float fLenghtPercent
+#ifdef WIND_EFFECT_FROND_RIPPLE_ADJUST_LIGHTING
+	, float3 vBinormal
+	, float3 vTangent
+#endif
 				   )
 {
-	return RippleFrondOneSided(vPos, 
+	return RippleFrondOneSided(STData, vPos, 
 								vDirection, 
 								fU, 
 								fV, 
 								fRippleScale
-							#ifdef WIND_EFFECT_FROND_RIPPLE_ADJUST_LIGHTING
+#ifdef WIND_EFFECT_FROND_RIPPLE_ADJUST_LIGHTING
 								, vBinormal
 								, vTangent
-							#endif
+#endif
 								);
 }
 
@@ -660,31 +746,32 @@ float3 RippleFrond(float3 vPos,
 //  LeafWind
 //  Optimized (for instruction count) version for UE4. Assumes leaf 1 and 2 have the same options
 
-float3 LeafWind(bool bLeaf2,
-				float3 vPos, 
-				inout float3 vDirection, 
-				float fScale, 
-				float3 vAnchor, 
-				float fPackedGrowthDir, 
-				float fPackedRippleDir,
-				float fRippleTrigOffset,
-				int WindType)
+float3 LeafWind(FSpeedTreeData STData,
+	bool bLeaf2,
+	float3 vPos, 
+	inout float3 vDirection, 
+	float fScale, 
+	float3 vAnchor, 
+	float fPackedGrowthDir, 
+	float fPackedRippleDir,
+	float fRippleTrigOffset,
+	int WindType)
 {
 
 	vPos = LeafRipple(vPos, vDirection, fScale, fPackedRippleDir, 
-							(bLeaf2 ? SpeedTreeData.WindLeaf2Ripple.x : SpeedTreeData.WindLeaf1Ripple.x),
-							(bLeaf2 ? SpeedTreeData.WindLeaf2Ripple.y : SpeedTreeData.WindLeaf1Ripple.y),
+							(bLeaf2 ? STData.WindLeaf2Ripple.x : STData.WindLeaf1Ripple.x),
+							(bLeaf2 ? STData.WindLeaf2Ripple.y : STData.WindLeaf1Ripple.y),
 							false, fRippleTrigOffset);
 
 	if (WindType == SPEEDTREE_WIND_TYPE_BEST)
 	{
 		float3 vGrowthDir = UnpackNormalFromFloat(fPackedGrowthDir);
-		vPos = LeafTumble(vPos, vDirection, fScale, vAnchor, vGrowthDir, fPackedGrowthDir,
-						  (bLeaf2 ? SpeedTreeData.WindLeaf2Tumble.x : SpeedTreeData.WindLeaf1Tumble.x), 
-						  (bLeaf2 ? SpeedTreeData.WindLeaf2Tumble.y : SpeedTreeData.WindLeaf1Tumble.y), 
-						  (bLeaf2 ? SpeedTreeData.WindLeaf2Tumble.z : SpeedTreeData.WindLeaf1Tumble.z), 
-						  (bLeaf2 ? SpeedTreeData.WindLeaf2Tumble.w : SpeedTreeData.WindLeaf1Tumble.w),
-						  (bLeaf2 ? SpeedTreeData.WindLeaf2Twitch.xyz : SpeedTreeData.WindLeaf1Twitch.xyz), 
+		vPos = LeafTumble(STData, vPos, vDirection, fScale, vAnchor, vGrowthDir, fPackedGrowthDir,
+						  (bLeaf2 ? STData.WindLeaf2Tumble.x : STData.WindLeaf1Tumble.x), 
+						  (bLeaf2 ? STData.WindLeaf2Tumble.y : STData.WindLeaf1Tumble.y), 
+						  (bLeaf2 ? STData.WindLeaf2Tumble.z : STData.WindLeaf1Tumble.z), 
+						  (bLeaf2 ? STData.WindLeaf2Tumble.w : STData.WindLeaf1Tumble.w),
+						  (bLeaf2 ? STData.WindLeaf2Twitch.xyz : STData.WindLeaf1Twitch.xyz), 
 						  0.0f,
 						  (bLeaf2 ? true : true), 
 						  (bLeaf2 ? true : true));
diff --git a/Engine/Source/Runtime/Engine/Classes/Materials/MaterialExpressionSpeedTree.h b/Engine/Source/Runtime/Engine/Classes/Materials/MaterialExpressionSpeedTree.h
--- a/Engine/Source/Runtime/Engine/Classes/Materials/MaterialExpressionSpeedTree.h
+++ b/Engine/Source/Runtime/Engine/Classes/Materials/MaterialExpressionSpeedTree.h
@@ -50,6 +50,9 @@ class UMaterialExpressionSpeedTree : public UMaterialExpression
 	UPROPERTY(EditAnywhere, Category=MaterialExpressionSpeedTree, meta=(DisplayName = "Billboard Threshold", ToolTip="The threshold for triangles to be removed from the bilboard mesh when not facing the camera (0 = none pass, 1 = all pass).", ClampMin = "0.0", ClampMax = "1.0", UIMin = "0.0", UIMax = "1.0"))
 	float BillboardThreshold;
 
+	UPROPERTY(EditAnywhere, Category=MaterialExpressionSpeedTree, AdvancedDisplay, meta=(DisplayName = "Support accurate velocities from wind"))
+	bool bOutputPreviousPosition;
+
 	// Begin UObject Interface
 	virtual void Serialize(FArchive& Ar) override;
 #if WITH_EDITOR
diff --git a/Engine/Source/Runtime/Engine/Private/Materials/HLSLMaterialTranslator.h b/Engine/Source/Runtime/Engine/Private/Materials/HLSLMaterialTranslator.h
--- a/Engine/Source/Runtime/Engine/Private/Materials/HLSLMaterialTranslator.h
+++ b/Engine/Source/Runtime/Engine/Private/Materials/HLSLMaterialTranslator.h
@@ -199,6 +199,8 @@ protected:
 	// True if the current property requires last frame's time
 	uint32 bUsesPreviousFrameTime : 1;
 	uint32 bUsesPixelDepthOffset : 1;
+	/** True if material will output accurate velocities during base pass rendering. */
+	uint32 bOutputsBasePassVelocities : 1;
 	/** Tracks the number of texture coordinates used by this material. */
 	uint32 NumUserTexCoords;
 	/** Tracks the number of texture coordinates used by the vertex shader in this material. */
@@ -242,6 +244,7 @@ public:
 	,	bUsesTransformVector(false)
 	,	bUsesPreviousFrameTime(false)
 	,	bUsesPixelDepthOffset(false)
+	,	bOutputsBasePassVelocities(true)
 	,	NumUserTexCoords(0)
 	,	NumUserVertexTexCoords(0)
 	{}
@@ -3755,7 +3758,7 @@ protected:
 	*
 	* @return	Code index
 	*/
-	virtual int32 SpeedTree(ESpeedTreeGeometryType GeometryType, ESpeedTreeWindType WindType, ESpeedTreeLODType LODType, float BillboardThreshold) override 
+	virtual int32 SpeedTree(ESpeedTreeGeometryType GeometryType, ESpeedTreeWindType WindType, ESpeedTreeLODType LODType, float BillboardThreshold, bool bOutputPreviousPosition) override 
 	{ 
 		if (ErrorUnlessFeatureLevelSupported(ERHIFeatureLevel::SM4) == INDEX_NONE)
 		{
@@ -3771,7 +3774,9 @@ protected:
 			bUsesSpeedTree = true;
 
 			NumUserVertexTexCoords = FMath::Max<uint32>(NumUserVertexTexCoords, 8);
-			return AddCodeChunk(MCT_Float3, TEXT("GetSpeedTreeVertexOffset(Parameters, %d, %d, %d, %g)"), GeometryType, WindType, LODType, BillboardThreshold);
+			// Only generate previous frame's computations if required and opted-in
+			const bool bEnablePreviousFrameInformation = bUsesPreviousFrameTime && bOutputPreviousPosition;
+			return AddCodeChunk(MCT_Float3, TEXT("GetSpeedTreeVertexOffset(Parameters, %d, %d, %d, %g, %s)"), GeometryType, WindType, LODType, BillboardThreshold, bEnablePreviousFrameInformation ? TEXT("true") : TEXT("false"));
 		}
 	}
 
diff --git a/Engine/Source/Runtime/Engine/Private/Materials/MaterialExpressions.cpp b/Engine/Source/Runtime/Engine/Private/Materials/MaterialExpressions.cpp
--- a/Engine/Source/Runtime/Engine/Private/Materials/MaterialExpressions.cpp
+++ b/Engine/Source/Runtime/Engine/Private/Materials/MaterialExpressions.cpp
@@ -8845,13 +8845,14 @@ UMaterialExpressionSpeedTree::UMaterialExpressionSpeedTree(const FObjectInitiali
 	WindType = STW_None;
 	LODType = STLOD_Pop;
 	BillboardThreshold = 0.9f;
+	bOutputPreviousPosition = true;
 
 	MenuCategories.Add(ConstructorStatics.NAME_SpeedTree);
 }
 
 int32 UMaterialExpressionSpeedTree::Compile(class FMaterialCompiler* Compiler, int32 OutputIndex, int32 MultiplexIndex)
 {
-	return Compiler->SpeedTree(GeometryType, WindType, LODType, BillboardThreshold);
+	return Compiler->SpeedTree(GeometryType, WindType, LODType, BillboardThreshold, bOutputPreviousPosition);
 }
 
 void UMaterialExpressionSpeedTree::GetCaption(TArray<FString>& OutCaptions) const
diff --git a/Engine/Source/Runtime/Engine/Private/SpeedTreeWind.cpp b/Engine/Source/Runtime/Engine/Private/SpeedTreeWind.cpp
--- a/Engine/Source/Runtime/Engine/Private/SpeedTreeWind.cpp
+++ b/Engine/Source/Runtime/Engine/Private/SpeedTreeWind.cpp
@@ -163,6 +163,11 @@ void FSpeedTreeWind::Advance(bool bEnabled, double fTime)
 		m_fElapsedTime = fTime - m_fLastTime;
 	m_fLastTime = fTime;
 
+	{
+		// Copy values to previous frame's area of the buffer
+		FMemory::Memcpy(m_afShaderTable + NUM_SHADER_VALUES, m_afShaderTable, NUM_SHADER_VALUES * sizeof(m_afShaderTable[0]));
+	}
+
 	if (bEnabled)
 	{
 		if (m_bGustingEnabled)
diff --git a/Engine/Source/Runtime/Engine/Public/MaterialCompiler.h b/Engine/Source/Runtime/Engine/Public/MaterialCompiler.h
--- a/Engine/Source/Runtime/Engine/Public/MaterialCompiler.h
+++ b/Engine/Source/Runtime/Engine/Public/MaterialCompiler.h
@@ -172,7 +172,7 @@ public:
 	virtual int32 BlackBody( int32 Temp ) = 0;
 	virtual int32 DepthOfFieldFunction(int32 Depth, int32 FunctionValueIndex) = 0;
 	virtual int32 AtmosphericFogColor(int32 WorldPosition) = 0;
-	virtual int32 SpeedTree(ESpeedTreeGeometryType GeometryType, ESpeedTreeWindType WindType, ESpeedTreeLODType LODType, float BillboardThreshold) = 0;
+	virtual int32 SpeedTree(ESpeedTreeGeometryType GeometryType, ESpeedTreeWindType WindType, ESpeedTreeLODType LODType, float BillboardThreshold, bool bOutputPreviousPosition) = 0;
 	virtual int32 TextureCoordinateOffset() = 0;
 	virtual int32 EyeAdaptation() = 0;
 };
@@ -335,9 +335,9 @@ public:
 		return Compiler->DepthOfFieldFunction(Depth, FunctionValueIndex);
 	}
 
-	virtual int32 SpeedTree(ESpeedTreeGeometryType GeometryType, ESpeedTreeWindType WindType, ESpeedTreeLODType LODType, float BillboardThreshold) override 
+	virtual int32 SpeedTree(ESpeedTreeGeometryType GeometryType, ESpeedTreeWindType WindType, ESpeedTreeLODType LODType, float BillboardThreshold, bool bOutputPreviousPosition) override 
 	{ 
-		return Compiler->SpeedTree(GeometryType, WindType, LODType, BillboardThreshold); 
+		return Compiler->SpeedTree(GeometryType, WindType, LODType, BillboardThreshold, bOutputPreviousPosition); 
 	}
 
 	virtual int32 AtmosphericFogColor(int32 WorldPosition) override
diff --git a/Engine/Source/Runtime/Engine/Public/SpeedTreeWind.h b/Engine/Source/Runtime/Engine/Public/SpeedTreeWind.h
--- a/Engine/Source/Runtime/Engine/Public/SpeedTreeWind.h
+++ b/Engine/Source/Runtime/Engine/Public/SpeedTreeWind.h
@@ -315,7 +315,8 @@ protected:
 
 					bool		m_bNeedsReload;
 
-	MS_ALIGN(16)	float		m_afShaderTable[NUM_SHADER_VALUES] GCC_ALIGN(16);
+					// Includes Previous frame's values after current
+	MS_ALIGN(16)	float		m_afShaderTable[NUM_SHADER_VALUES * 2] GCC_ALIGN(16);
 };
 
 
@@ -343,6 +344,24 @@ BEGIN_UNIFORM_BUFFER_STRUCT(FSpeedTreeUniformParameters, ENGINE_API)
 	DECLARE_UNIFORM_BUFFER_STRUCT_MEMBER(FVector4, WindRollingLeafAndDirection)
 	DECLARE_UNIFORM_BUFFER_STRUCT_MEMBER(FVector4, WindRollingNoise)
 	DECLARE_UNIFORM_BUFFER_STRUCT_MEMBER(FVector4, WindAnimation)
+	// Straight copy of the previous members for last frame's values
+	DECLARE_UNIFORM_BUFFER_STRUCT_MEMBER(FVector4, PrevWindVector)
+	DECLARE_UNIFORM_BUFFER_STRUCT_MEMBER(FVector4, PrevWindGlobal)
+	DECLARE_UNIFORM_BUFFER_STRUCT_MEMBER(FVector4, PrevWindBranch)
+	DECLARE_UNIFORM_BUFFER_STRUCT_MEMBER(FVector4, PrevWindBranchTwitch)
+	DECLARE_UNIFORM_BUFFER_STRUCT_MEMBER(FVector4, PrevWindBranchWhip)
+	DECLARE_UNIFORM_BUFFER_STRUCT_MEMBER(FVector4, PrevWindBranchAnchor)
+	DECLARE_UNIFORM_BUFFER_STRUCT_MEMBER(FVector4, PrevWindBranchAdherences)
+	DECLARE_UNIFORM_BUFFER_STRUCT_MEMBER(FVector4, PrevWindTurbulences)
+	DECLARE_UNIFORM_BUFFER_STRUCT_MEMBER(FVector4, PrevWindLeaf1Ripple)
+	DECLARE_UNIFORM_BUFFER_STRUCT_MEMBER(FVector4, PrevWindLeaf1Tumble)
+	DECLARE_UNIFORM_BUFFER_STRUCT_MEMBER(FVector4, PrevWindLeaf1Twitch)
+	DECLARE_UNIFORM_BUFFER_STRUCT_MEMBER(FVector4, PrevWindLeaf2Ripple)
+	DECLARE_UNIFORM_BUFFER_STRUCT_MEMBER(FVector4, PrevWindLeaf2Tumble)
+	DECLARE_UNIFORM_BUFFER_STRUCT_MEMBER(FVector4, PrevWindLeaf2Twitch)
+	DECLARE_UNIFORM_BUFFER_STRUCT_MEMBER(FVector4, PrevWindFrondRipple)
+	DECLARE_UNIFORM_BUFFER_STRUCT_MEMBER(FVector4, PrevWindRollingBranch)
+	DECLARE_UNIFORM_BUFFER_STRUCT_MEMBER(FVector4, PrevWindRollingLeafAndDirection)
+	DECLARE_UNIFORM_BUFFER_STRUCT_MEMBER(FVector4, PrevWindRollingNoise)
+	DECLARE_UNIFORM_BUFFER_STRUCT_MEMBER(FVector4, PrevWindAnimation)
 END_UNIFORM_BUFFER_STRUCT(FSpeedTreeUniformParameters)
-typedef TUniformBufferRef<FSpeedTreeUniformParameters> FSpeedTreeUniformBufferRef;
-
diff --git a/Engine/Source/Runtime/Renderer/Private/RendererScene.cpp b/Engine/Source/Runtime/Renderer/Private/RendererScene.cpp
--- a/Engine/Source/Runtime/Renderer/Private/RendererScene.cpp
+++ b/Engine/Source/Runtime/Renderer/Private/RendererScene.cpp
@@ -14,6 +14,7 @@
 #include "PrecomputedLightVolume.h"
 #include "FXSystem.h"
 #include "DistanceFieldLightingShared.h"
+#include "SpeedTreeWind.h"
 
 // Enable this define to do slow checks for components being added to the wrong
 // world's scene, when using PIE. This can happen if a PIE component is reattached
@@ -28,6 +29,26 @@ TGlobalResource< FGlobalDistanceCullFadeUniformBuffer > GDistanceCullFadedInUnif
 
 SIZE_T FStaticMeshDrawListBase::TotalBytesUsed = 0;
 
+/**
+ * Holds the info to update SpeedTree wind per unique tree object in the scene, instead of per instance
+ */
+struct FSpeedTreeWindComputation
+{
+	explicit FSpeedTreeWindComputation() :
+		ReferenceCount(1)
+	{
+	}
+
+	/** SpeedTree wind object */
+	FSpeedTreeWind Wind;
+
+	/** Uniform buffer shared between trees of the same type. */
+	TUniformBuffer<FSpeedTreeUniformParameters> UniformBuffer;
+
+	int32 ReferenceCount;
+};
+
+
 /** Default constructor. */
 FSceneViewState::FSceneViewState()
 	: OcclusionQueryPool(RQT_Occlusion)
diff --git a/Engine/Source/Runtime/Renderer/Private/ScenePrivate.h b/Engine/Source/Runtime/Renderer/Private/ScenePrivate.h
--- a/Engine/Source/Runtime/Renderer/Private/ScenePrivate.h
+++ b/Engine/Source/Runtime/Renderer/Private/ScenePrivate.h
@@ -68,7 +68,6 @@ bool IsMobileHDR32bpp();
 #include "TextureLayout3d.h"
 #include "ScopedPointer.h"
 #include "ClearQuad.h"
-#include "SpeedTreeWind.h"
 #include "AtmosphereRendering.h"
 
 #if WITH_SLI || PLATFORM_SHOULD_BUFFER_QUERIES
@@ -1231,25 +1230,6 @@ namespace EOcclusionFlags
 	};
 };
 
-/**
- * Holds the info to update SpeedTree wind per unique tree object in the scene, instead of per instance
- */
-struct FSpeedTreeWindComputation
-{
-	explicit FSpeedTreeWindComputation() :
-		ReferenceCount(1)
-	{
-	}
-
-	/** SpeedTree wind object */
-	FSpeedTreeWind Wind;
-
-	/** Uniform buffer shared between trees of the same type. */
-	TUniformBuffer<FSpeedTreeUniformParameters> UniformBuffer;
-
-	int32 ReferenceCount;
-};
-
 /** Information about the primitives that are attached together. */
 class FAttachmentGroupSceneInfo
 {
@@ -1495,7 +1475,7 @@ public:
 	TArray<class FWindSourceSceneProxy*> WindSources;
 
 	/** SpeedTree wind objects in the scene. FLocalVertexFactoryShaderParameters needs to lookup by FVertexFactory, but wind objects are per tree (i.e. per UStaticMesh)*/
-	TMap<const UStaticMesh*, FSpeedTreeWindComputation*> SpeedTreeWindComputationMap;
+	TMap<const UStaticMesh*, struct FSpeedTreeWindComputation*> SpeedTreeWindComputationMap;
 	TMap<FVertexFactory*, const UStaticMesh*> SpeedTreeVertexFactoryMap;
 
 	/** The attachment groups in the scene.  The map key is the attachment group's root primitive. */
