fixed Speedtree shader compiler warnings

diff --git a/Engine/Shaders/MaterialTemplate.usf b/Engine/Shaders/MaterialTemplate.usf
--- a/Engine/Shaders/MaterialTemplate.usf
+++ b/Engine/Shaders/MaterialTemplate.usf
@@ -699,7 +699,7 @@ float3 GetSpeedTreeVertexOffsetInner(FMaterialVertexParameters Parameters, int G
 		// skip if this instance is hidden
 		if (Parameters.PerInstanceParams.z < 1.f)
 		{
-			return float4(0,0,0,0);
+			return float3(0,0,0);
 		}
 	#else
 		float3x3 LocalToWorld = (float3x3)Primitive.LocalToWorld;
@@ -716,7 +716,7 @@ float3 GetSpeedTreeVertexOffsetInner(FMaterialVertexParameters Parameters, int G
 		const float ScreenMultiple = 0.5 * max(View.ViewSizeAndSceneTexelSize.x * View.ViewToClip[0][0],
 												View.ViewSizeAndSceneTexelSize.y * View.ViewToClip[1][1]);
 		const float ScreenRadius = ScreenMultiple * Primitive.ObjectWorldPositionAndRadius.w / 
-									max(1.0, dot(TreePos - View.ViewOrigin.xyz, View.ViewForward));
+									max(1.0, dot(TreePos - View.ViewOrigin.xyz, View.ViewForward.xyz));
 		const float ScreenArea = PI * ScreenRadius * ScreenRadius / (View.ViewSizeAndSceneTexelSize.x * View.ViewSizeAndSceneTexelSize.y);
 		LodInterp = saturate((ScreenArea - SpeedTreeLODInfo.x) / SpeedTreeLODInfo.z);
 	}
