experimental ability to output a different opacity mask for shadow rendering (free as decided at compile time, minimal code complexity, currently hidden as it requires Custom node)

diff --git a/Engine/Shaders/Common.usf b/Engine/Shaders/Common.usf
--- a/Engine/Shaders/Common.usf
+++ b/Engine/Shaders/Common.usf
@@ -1101,3 +1101,16 @@ SafeSaturate_Def(float)
 SafeSaturate_Def(float2)
 SafeSaturate_Def(float3)
 SafeSaturate_Def(float4)
+
+// Experimental way to allow adjusting the OpacityMask for shadow map rendering of masked materials.
+// Can be accessed with a Custom material node. If this turns out to be very useful we can expose as MaterialFunction
+// and potentially expose other queries as well (e.g. SkeletalMesh, HitProxy, ).
+// @return 0:no, 1:yes
+float IsShadowPixelShader()
+{
+#ifdef SHADOW_DEPTH_SHADER
+	return 1;
+#else
+	return 0;
+#endif
+}
diff --git a/Engine/Shaders/ShadowDepthPixelShader.usf b/Engine/Shaders/ShadowDepthPixelShader.usf
--- a/Engine/Shaders/ShadowDepthPixelShader.usf
+++ b/Engine/Shaders/ShadowDepthPixelShader.usf
@@ -4,6 +4,9 @@
 	ShadowDepthPixelShader.usf: Pixel shader for writing shadow depth.
 =============================================================================*/
 
+// needs to before Common.usf
+#define SHADOW_DEPTH_SHADER 1
+
 #include "Common.usf"
 #include "Material.usf"
 #include "VertexFactory.usf"
diff --git a/Engine/Shaders/ShadowDepthVertexShader.usf b/Engine/Shaders/ShadowDepthVertexShader.usf
--- a/Engine/Shaders/ShadowDepthVertexShader.usf
+++ b/Engine/Shaders/ShadowDepthVertexShader.usf
@@ -4,6 +4,9 @@
 	ShadowDepthVertexShader.usf: Vertex shader for writing shadow depth.
 =============================================================================*/
 
+// needs to before Common.usf
+#define SHADOW_DEPTH_SHADER 1
+
 #include "Common.usf"
 #include "Material.usf"
 #include "VertexFactory.usf"
