minor rename, now it's IsShadowDepthShader()

diff --git a/Engine/Shaders/Common.usf b/Engine/Shaders/Common.usf
--- a/Engine/Shaders/Common.usf
+++ b/Engine/Shaders/Common.usf
@@ -1106,7 +1106,7 @@ SafeSaturate_Def(float4)
 // Can be accessed with a Custom material node. If this turns out to be very useful we can expose as MaterialFunction
 // and potentially expose other queries as well (e.g. SkeletalMesh, HitProxy, ).
 // @return 0:no, 1:yes
-float IsShadowPixelShader()
+float IsShadowDepthShader()
 {
 #ifdef SHADOW_DEPTH_SHADER
 	return 1;
