fixed DecalMask being 0/2 instead 0/1

diff --git a/Engine/Shaders/DeferredShadingCommon.usf b/Engine/Shaders/DeferredShadingCommon.usf
--- a/Engine/Shaders/DeferredShadingCommon.usf
+++ b/Engine/Shaders/DeferredShadingCommon.usf
@@ -380,7 +380,7 @@ FGBufferData DecodeGBufferData(
 	}
 
 	uint PackedAlpha = (uint)( InGBufferA.a * 3.999 );
-	GBuffer.DecalMask = PackedAlpha & 2;
+	GBuffer.DecalMask = (PackedAlpha & 2) != 0;
 	GBuffer.HasDistanceFieldRepresentation = PackedAlpha & 1;
 	GBuffer.ShadingModelID = DecodeShadingModelId(InGBufferB.a);
 	
