Check in of prototype CircleDOF with Near DOF.

diff --git a/Engine/Shaders/PostProcessDOF.usf b/Engine/Shaders/PostProcessDOF.usf
--- a/Engine/Shaders/PostProcessDOF.usf
+++ b/Engine/Shaders/PostProcessDOF.usf
@@ -203,27 +203,12 @@ float DepthToCoc(float Depth)
 {
 	float Focus = View.DepthOfFieldFocalDistance;
 	float Radius = View.DepthOfFieldScale * 2.0;
-
-	// Returning only FAR COC for now.
-	return max(0.0,(Depth - Focus) / Depth) * Radius;
+	return ((Depth - Focus) / Depth) * Radius;
 }
 
 // pixel shader entry point
-void CircleSetupPS(
-	float4 UVAndScreenPos : TEXCOORD0
-	, out float4 OutColor0 : SV_Target0
-#if ENABLE_NEAR_BLUR 
-	, out float4 OutColor1 : SV_Target1
-#endif
-	)
+void CircleSetupPS(float4 UVAndScreenPos : TEXCOORD0, out float4 OutColor0 : SV_Target0)
 {
-#if ENABLE_NEAR_BLUR == 0
-	// We aren't writing out to the second render target, so we'll just make a dummy value here which 
-	// doesn't end up going anywhere. Then, the source code can stay neat and tidy, while the compiler can still 
-	// strip it out
-	float4 OutColor1;
-#endif
-
 	float2 UV = UVAndScreenPos.xy;
 
 	float4 DepthQuad = GatherSceneDepth(UV, PostprocessInput1Size.zw);
@@ -240,22 +225,28 @@ void CircleSetupPS(
 	// Doing a max depth reduction (erode the foreground). Less correct, but less artifacts.
 	// Perhaps need to re-open this in the future.
 
-	// Stuff max radius in alpha.
-	OutColor0.a = max(max(CocQuad.x,CocQuad.y),max(CocQuad.z,CocQuad.w));
+	#if 0
+		// Stuff max radius in alpha.
+		OutColor0.a = max(max(CocQuad.x,CocQuad.y),max(CocQuad.z,CocQuad.w));
+	#else
+		// This in theory is better but causes bleeding artifacts with temporal AA..
+		// This is important otherwise near thin objects disappear (leaves clamping artifacts in recombined pass).
+		OutColor0.a = CocQuad.x;
+		if(abs(OutColor0.a) > CocQuad.y) OutColor0.a = CocQuad.y;
+		if(abs(OutColor0.a) > CocQuad.z) OutColor0.a = CocQuad.z;
+		if(abs(OutColor0.a) > CocQuad.w) OutColor0.a = CocQuad.w;
+	#endif
 
 	// Remove samples which are outside the size.
 	// TODO: Tune the ScaleFactor.
-	float ScaleFactor = 1.0;
+	float ScaleFactor = 64.0;
 	float4 W = float4(
-		1.0-saturate((OutColor0.a - CocQuad.x) * ScaleFactor),
-		1.0-saturate((OutColor0.a - CocQuad.y) * ScaleFactor),
-		1.0-saturate((OutColor0.a - CocQuad.z) * ScaleFactor),
-		1.0-saturate((OutColor0.a - CocQuad.w) * ScaleFactor));
+		1.0-saturate(abs(OutColor0.a - CocQuad.x) * ScaleFactor),
+		1.0-saturate(abs(OutColor0.a - CocQuad.y) * ScaleFactor),
+		1.0-saturate(abs(OutColor0.a - CocQuad.z) * ScaleFactor),
+		1.0-saturate(abs(OutColor0.a - CocQuad.w) * ScaleFactor));
 
 	OutColor0.rgb = (1.0/(W.x+W.y+W.z+W.w)) * (CX.rgb*W.x + CY.rgb*W.y + CZ.rgb*W.z + CW.rgb*W.w);
-
-	// TODO: Near.
-	OutColor1 = 0;
 }
 
 
@@ -280,66 +271,163 @@ float2 RotVec(float Radius, float Radians)
 
 float2 RandomOffset;
 
+
+
+
+float Min4(float4 A)
+{
+	return min(min(A.x,A.y),min(A.z,A.w));
+}
+
+float Min16(float4 A, float4 B, float4 C, float4 D)
+{
+	return min(min(Min4(A),Min4(B)),min(Min4(C),Min4(D)));
+}
+
+// This does not actually do any dilation, but rather just a 2x2:1 min reduction.
+void CircleDilatePS(float4 UVAndScreenPos : TEXCOORD0, out float OutColor : SV_Target0)
+{
+	#if 0
+		// Leaving code in for a 4x4:1 reduction (but not using).
+	
+		// Sampling pattern (each gather4)
+		//   a b e f
+	    //   c d g h
+		//   i j M n  (M={0,0} point)
+		//   k l o p
+
+		float2 UV = UVAndScreenPos.xy + 1.0*PostprocessInput0Size.zw;
+	
+		float4 Sa = PostprocessInput0.GatherAlpha(PostprocessInput0Sampler, UV, int2(-4,-4));
+		float4 Sb = PostprocessInput0.GatherAlpha(PostprocessInput0Sampler, UV, int2(-2,-4));
+		float4 Sc = PostprocessInput0.GatherAlpha(PostprocessInput0Sampler, UV, int2(-4,-2));
+		float4 Sd = PostprocessInput0.GatherAlpha(PostprocessInput0Sampler, UV, int2(-2,-2));
+
+		float4 Se = PostprocessInput0.GatherAlpha(PostprocessInput0Sampler, UV, int2(0,-4));
+		float4 Sf = PostprocessInput0.GatherAlpha(PostprocessInput0Sampler, UV, int2(2,-4));
+		float4 Sg = PostprocessInput0.GatherAlpha(PostprocessInput0Sampler, UV, int2(0,-2));
+		float4 Sh = PostprocessInput0.GatherAlpha(PostprocessInput0Sampler, UV, int2(2,-2));
+
+		float4 Si = PostprocessInput0.GatherAlpha(PostprocessInput0Sampler, UV, int2(-4,0));
+		float4 Sj = PostprocessInput0.GatherAlpha(PostprocessInput0Sampler, UV, int2(-2,0));
+		float4 Sk = PostprocessInput0.GatherAlpha(PostprocessInput0Sampler, UV, int2(-4,2));
+		float4 Sl = PostprocessInput0.GatherAlpha(PostprocessInput0Sampler, UV, int2(-2,2));
+
+		float4 Sm = PostprocessInput0.GatherAlpha(PostprocessInput0Sampler, UV, int2(0,0));
+		float4 Sn = PostprocessInput0.GatherAlpha(PostprocessInput0Sampler, UV, int2(2,0));
+		float4 So = PostprocessInput0.GatherAlpha(PostprocessInput0Sampler, UV, int2(0,2));
+		float4 Sp = PostprocessInput0.GatherAlpha(PostprocessInput0Sampler, UV, int2(2,2));
+	
+		// Make sure near is only near blur.
+		OutColor = min(0.0,min(min(Min16(Sa,Sb,Sc,Sd),Min16(Se,Sf,Sg,Sh)),min(Min16(Si,Sj,Sk,Sl),Min16(Sm,Sn,So,Sp))));
+	
+	#else
+	
+		// Sampling pattern (each gather4)
+	    //   d g
+		//   j M  (M={0,0} point)
+	
+		float2 UV = UVAndScreenPos.xy + 1.0*PostprocessInput0Size.zw;
+	
+		float4 Sd = PostprocessInput0.GatherAlpha(PostprocessInput0Sampler, UV, int2(-2,-2));
+		float4 Sg = PostprocessInput0.GatherAlpha(PostprocessInput0Sampler, UV, int2(0,-2));
+		float4 Sj = PostprocessInput0.GatherAlpha(PostprocessInput0Sampler, UV, int2(-2,0));
+		float4 Sm = PostprocessInput0.GatherAlpha(PostprocessInput0Sampler, UV, int2(0,0));
+	
+		// Make sure near is only near blur.
+		OutColor = min(0.0, Min16(Sd, Sg, Sj, Sm));
+	
+	#endif
+}
+
+
 // pixel shader entry point
-void CirclePS(
-	float4 UVAndScreenPos : TEXCOORD0
-	, out float4 OutColor0 : SV_Target0
-#if ENABLE_NEAR_BLUR 
-	, out float4 OutColor1 : SV_Target1
-#endif
-	)
+void CirclePS(float4 UVAndScreenPos : TEXCOORD0, out float4 OutColor0 : SV_Target0)
 {
 	float2 UV = UVAndScreenPos.xy;
 
-#if ENABLE_NEAR_BLUR == 0
-	// We aren't writing out to the second render target, so we'll just make a dummy value here which 
-	// doesn't end up going anywhere. Then, the source code can stay neat and tidy, while the compiler can still 
-	// strip it out
-	float4 OutColor1;
-#endif
+	//
+	// Pass 0
+	// Dilate near minimum CoC (near CoC is negative values).
+	//
+
+	// Fixed maximum search size (in terms of Circle of Confusion radius).
+	// Higher than 8 is too noizy for 4 samples.
+	float Coc = 8.0;
+
+	// Get base semi-random direction and dither along radius.
+	// Reused throughout the rest of the algorithm.
+	float TwoPi = 2.0 * 3.14159;
+	float RadianBase = NoizSnorm(UVAndScreenPos.xy, 0.010 * RandomOffset.x) * TwoPi;
+	float RadiusBase = NoizNorm(UVAndScreenPos.xy, 0.013 * RandomOffset.x);
 
-	// TODO: Near.
-	OutColor1 = 0.0;
+	// Radius 
+	float RadiusBase2 = RadiusBase * (1.0/4.0);
+	float R1 = (RadiusBase2 + 3.0/4.0) * Coc;
+	float R2 = (RadiusBase2 + 2.0/4.0) * Coc;
+	float R3 = (RadiusBase2 + 1.0/4.0) * Coc;
+	float R4 = (RadiusBase2 + 0.0/4.0) * Coc;
+
+	float2 UV1 = RotVec(R1, RadianBase + TwoPi * 0.0/4.0);
+	float2 UV2 = RotVec(R2, RadianBase + TwoPi * 2.0/4.0);
+	float2 UV3 = RotVec(R3, RadianBase + TwoPi * 1.0/4.0);
+	float2 UV4 = RotVec(R4, RadianBase + TwoPi * 3.0/4.0);
+
+	UV1 = UVAndScreenPos.xy + UV1 * PostprocessInput0Size.zw;
+	UV2 = UVAndScreenPos.xy + UV2 * PostprocessInput0Size.zw;
+	UV3 = UVAndScreenPos.xy + UV3 * PostprocessInput0Size.zw;
+	UV4 = UVAndScreenPos.xy + UV4 * PostprocessInput0Size.zw;
+
+	float D1 = PostprocessInput1.SampleLevel(PostprocessInput1Sampler, UV1, 0).x;
+	float D2 = PostprocessInput1.SampleLevel(PostprocessInput1Sampler, UV2, 0).x;
+	float D3 = PostprocessInput1.SampleLevel(PostprocessInput1Sampler, UV3, 0).x;
+	float D4 = PostprocessInput1.SampleLevel(PostprocessInput1Sampler, UV4, 0).x;
+
+	float NearCoc = 65536.0;
+	float Feather = 0.0;
+	if(abs(D1)+Feather > R1) NearCoc = min(NearCoc, D1);
+	if(abs(D2)+Feather > R2) NearCoc = min(NearCoc, D2);
+	if(abs(D3)+Feather > R3) NearCoc = min(NearCoc, D3);
+	if(abs(D4)+Feather > R4) NearCoc = min(NearCoc, D4);
 
 	//
-	// TODO: This part of the algorithm is a placeholder for something better.
-	//       Current placeholder is a stochastic algorithm.
+	// Pass 1
 	//
 
-	// Going to grab sets of 4 samples.
+	// Going to grab sets of 4 samples per pass.
 	// Each set of 4 samples can be a smaller circle of confusion
 	// (aka can be in-front of the larger background).
 
-	//
-	// Pass 0
-	//
-
-	float TwoPi = 2.0 * 3.14159;
+	// Setup for 12 samples (3 passes of 4 samples).
+	RadiusBase *= (1.0/11.5);
 
-	// Grab circle of confusion for the pixel.
+	// Grab circle of confusion for the pixel and pixel color.
 	OutColor0 = Texture2DSampleLevel(PostprocessInput0, PostprocessInput0Sampler, UV, 0);
-	float Coc = OutColor0.a + (1.0/65536.0);
+	float FarCoc = OutColor0.a;
 
-	float RadianBase = NoizSnorm(UVAndScreenPos.xy, 0.010 * RandomOffset.x) * TwoPi;
-	float RadiusBase = NoizNorm(UVAndScreenPos.xy, 0.013 * RandomOffset.x) * (1.0/11.5);
+	// Fix in case no near exists.
+	NearCoc = min(NearCoc, FarCoc);
+
+	// Used for sample pattern.
+	Coc = max(abs(FarCoc),abs(NearCoc));
 
 	// Bring out to the smaller radius of sample sets.
 	// This has the highest chance of seeing a smaller overlapping CoC.
-	float R1 = (RadiusBase+9.0/11.5) * Coc;
-	float R2 = (RadiusBase+3.0/11.5) * Coc;
-	float R3 = (RadiusBase+6.0/11.5) * Coc;
-	float R4 = (RadiusBase+0.0/11.5) * Coc;
+	R1 = (RadiusBase+9.0/11.5) * Coc;
+	R2 = (RadiusBase+3.0/11.5) * Coc;
+	R3 = (RadiusBase+6.0/11.5) * Coc;
+	R4 = (RadiusBase+0.0/11.5) * Coc;
 
-	// Ensure at least getting different than center pixel.
+	// Ensure at least getting different sample than center pixel.
 	float R1a = max(1.0,R1);
 	float R2a = max(1.0,R2);
 	float R3a = max(1.0,R3);
 	float R4a = max(1.0,R4);
 
-	float2 UV1 = RotVec(R1a, RadianBase + TwoPi * 0.0/12.0);
-	float2 UV2 = RotVec(R2a, RadianBase + TwoPi * 3.0/12.0);
-	float2 UV3 = RotVec(R3a, RadianBase + TwoPi * 6.0/12.0);
-	float2 UV4 = RotVec(R4a, RadianBase + TwoPi * 9.0/12.0);
+	UV1 = RotVec(R1a, RadianBase + TwoPi * 0.0/12.0);
+	UV2 = RotVec(R2a, RadianBase + TwoPi * 3.0/12.0);
+	UV3 = RotVec(R3a, RadianBase + TwoPi * 6.0/12.0);
+	UV4 = RotVec(R4a, RadianBase + TwoPi * 9.0/12.0);
 
 	UV1 = UVAndScreenPos.xy + UV1 * PostprocessInput0Size.zw;
 	UV2 = UVAndScreenPos.xy + UV2 * PostprocessInput0Size.zw;
@@ -359,33 +447,53 @@ void CirclePS(
 	float W3 = R3;
 	float W4 = R4;
 
-	// Adjusting weight to remove occluded.
-	W1 *= saturate(0.0 + C1.a - R1);
-	W2 *= saturate(0.0 + C2.a - R2);
-	W3 *= saturate(0.0 + C3.a - R3);
-	W4 *= saturate(0.0 + C4.a - R4);
+	// Intersection weight: 0=sample does not intersect pixel, to 1=sample intersects.
+	// TODO: Tune feather factors.
+	float IFeather0 = 1.0/4.0;
+	float I1 = saturate((abs(C1.a) - R1) * IFeather0);
+	float I2 = saturate((abs(C2.a) - R2) * IFeather0);
+	float I3 = saturate((abs(C3.a) - R3) * IFeather0);
+	float I4 = saturate((abs(C4.a) - R4) * IFeather0);
+
+	// Check if have a more near intersecting Coc for next pass.
+	float FarCoc2 = FarCoc;
+	if(I1*W1 > 0.0) FarCoc2 = min(FarCoc2, C1.a);
+	if(I2*W2 > 0.0) FarCoc2 = min(FarCoc2, C2.a);
+	if(I3*W3 > 0.0) FarCoc2 = min(FarCoc2, C3.a);
+	if(I4*W4 > 0.0) FarCoc2 = min(FarCoc2, C4.a);
+
+	// Fully ignore intersection weight when in nearfield blur
+	// and sample average CoC is 50% between near and far CoC neighborhood.
+	float AvgCoc = (FarCoc + C1.a + C2.a + C3.a + C4.a) * (1.0/5.0);
+	// Get dilated far.
+	FarCoc = max(FarCoc, max(max(C1.a, C2.a),max(C3.a, C4.a)));
+	// Controls the transition between states.
+	float IFeather1 = 1.0;
+	float IFeather2 = 2.0;
+	float Ignore = saturate(-NearCoc * IFeather1) * saturate(((AvgCoc - FarCoc) / (NearCoc - FarCoc)) * IFeather2);
+
+	W1 *= lerp(I1, 1.0, Ignore);
+	W2 *= lerp(I2, 1.0, Ignore);
+	W3 *= lerp(I3, 1.0, Ignore);
+	W4 *= lerp(I4, 1.0, Ignore);
 
 	// Make sure at least something is not zero.
 	W0 += 1.0/65536.0;
 
-	float Coc2 = Coc;
-	if(W1 > 0.0) Coc2 = min(Coc2, C1.a);
-	if(W2 > 0.0) Coc2 = min(Coc2, C2.a);
-	if(W3 > 0.0) Coc2 = min(Coc2, C3.a);
-	if(W4 > 0.0) Coc2 = min(Coc2, C4.a);
-
+	// Start weighted accumulation.
 	OutColor0.rgb = OutColor0.rgb * W0 + C1.rgb * W1 + C2.rgb * W2 + C3.rgb * W3 + C4.rgb * W4;
-
 	float Weight = W0+W1+W2+W3+W4;
 
+	// Set current result as possible background.
 	float3 Background = OutColor0.rgb * (1.0/Weight);
 
 
 	//
-	// Pass 1
+	// Pass 2
 	//
 
 	// Drop weight of existing pass if Coc changes too much.
+	float Coc2 = max(abs(FarCoc2),abs(NearCoc));
 	float Drop = (1.0/65536.0) + 1.0 - saturate(abs(Coc - Coc2));
 	OutColor0.rgb *= Drop;
 	Weight *= Drop;
@@ -415,43 +523,52 @@ void CirclePS(
 	C3 = PostprocessInput0.SampleLevel(PostprocessInput0Sampler, UV3, 0);
 	C4 = PostprocessInput0.SampleLevel(PostprocessInput0Sampler, UV4, 0);
 
-	C1.rgb = lerp(C1.rgb, Background.rgb, saturate(C1.a - Coc2));
-	C2.rgb = lerp(C2.rgb, Background.rgb, saturate(C2.a - Coc2));
-	C3.rgb = lerp(C3.rgb, Background.rgb, saturate(C3.a - Coc2));
-	C4.rgb = lerp(C4.rgb, Background.rgb, saturate(C4.a - Coc2));
+	// Lerp to background if outside possibly smaller CoC.
+	C1.rgb = lerp(C1.rgb, Background.rgb, saturate(abs(C1.a) - Coc2));
+	C2.rgb = lerp(C2.rgb, Background.rgb, saturate(abs(C2.a) - Coc2));
+	C3.rgb = lerp(C3.rgb, Background.rgb, saturate(abs(C3.a) - Coc2));
+	C4.rgb = lerp(C4.rgb, Background.rgb, saturate(abs(C4.a) - Coc2));
 
 	W1 = R1;
 	W2 = R2;
 	W3 = R3;
 	W4 = R4;
 
-	// Adjusting weight to remove occluded.
-	W1 *= saturate(0.0 + C1.a - R1);
-	W2 *= saturate(0.0 + C2.a - R2);
-	W3 *= saturate(0.0 + C3.a - R3);
-	W4 *= saturate(0.0 + C4.a - R4);
+	// Intersection weight: 0=sample does not intersect pixel, to 1=sample intersects.
+	I1 = saturate((abs(C1.a) - R1) * IFeather0);
+	I2 = saturate((abs(C2.a) - R2) * IFeather0);
+	I3 = saturate((abs(C3.a) - R3) * IFeather0);
+	I4 = saturate((abs(C4.a) - R4) * IFeather0);
 
-	float Coc3 = Coc2;
-	if(W1 > 0.0) Coc3 = min(Coc3, C1.a);
-	if(W2 > 0.0) Coc3 = min(Coc3, C2.a);
-	if(W3 > 0.0) Coc3 = min(Coc3, C3.a);
-	if(W4 > 0.0) Coc3 = min(Coc3, C4.a);
+	// Check if have a more near intersecting Coc for next pass.
+	float FarCoc3 = FarCoc2;
+	if(I1*W1 > 0.0) FarCoc3 = min(FarCoc3, C1.a);
+	if(I2*W2 > 0.0) FarCoc3 = min(FarCoc3, C2.a);
+	if(I3*W3 > 0.0) FarCoc3 = min(FarCoc3, C3.a);
+	if(I4*W4 > 0.0) FarCoc3 = min(FarCoc3, C4.a);
+
+	W1 *= lerp(I1, 1.0, Ignore);
+	W2 *= lerp(I2, 1.0, Ignore);
+	W3 *= lerp(I3, 1.0, Ignore);
+	W4 *= lerp(I4, 1.0, Ignore);
 
 	OutColor0.rgb += C1.rgb * W1 + C2.rgb * W2 + C3.rgb * W3 + C4.rgb * W4;
 	Weight += W1+W2+W3+W4;
 
 
-
-
 	//
-	// Pass 2
+	// Pass 3
 	//
 
 	// Drop weight of existing pass if Coc changes too much.
+	float Coc3 = max(abs(FarCoc3),abs(NearCoc));
 	Drop = (1.0/65536.0) + 1.0 - saturate(abs(Coc2 - Coc3));
 	OutColor0.rgb *= Drop;
 	Weight *= Drop;
 
+	// Send near most CoC back to recombine pass.
+	OutColor0.a = min(FarCoc3, NearCoc);
+
 	R1 = (RadiusBase+11.0/11.5) * Coc3;
 	R2 = (RadiusBase+ 5.0/11.5) * Coc3;
 	R3 = (RadiusBase+ 8.0/11.5) * Coc3;
@@ -477,21 +594,25 @@ void CirclePS(
 	C3 = PostprocessInput0.SampleLevel(PostprocessInput0Sampler, UV3, 0);
 	C4 = PostprocessInput0.SampleLevel(PostprocessInput0Sampler, UV4, 0);
 
-	C1.rgb = lerp(C1.rgb, Background.rgb, saturate(C1.a - Coc3));
-	C2.rgb = lerp(C2.rgb, Background.rgb, saturate(C2.a - Coc3));
-	C3.rgb = lerp(C3.rgb, Background.rgb, saturate(C3.a - Coc3));
-	C4.rgb = lerp(C4.rgb, Background.rgb, saturate(C4.a - Coc3));
+	C1.rgb = lerp(C1.rgb, Background.rgb, saturate(abs(C1.a) - Coc3));
+	C2.rgb = lerp(C2.rgb, Background.rgb, saturate(abs(C2.a) - Coc3));
+	C3.rgb = lerp(C3.rgb, Background.rgb, saturate(abs(C3.a) - Coc3));
+	C4.rgb = lerp(C4.rgb, Background.rgb, saturate(abs(C4.a) - Coc3));
 
 	W1 = R1;
 	W2 = R2;
 	W3 = R3;
 	W4 = R4;
 
-	// Adjusting weight to remove occluded.
-	W1 *= saturate(0.0 + C1.a - R1);
-	W2 *= saturate(0.0 + C2.a - R2);
-	W3 *= saturate(0.0 + C3.a - R3);
-	W4 *= saturate(0.0 + C4.a - R4);
+	I1 = saturate((abs(C1.a) - R1) * IFeather0);
+	I2 = saturate((abs(C2.a) - R2) * IFeather0);
+	I3 = saturate((abs(C3.a) - R3) * IFeather0);
+	I4 = saturate((abs(C4.a) - R4) * IFeather0);
+
+	W1 *= lerp(I1, 1.0, Ignore);
+	W2 *= lerp(I2, 1.0, Ignore);
+	W3 *= lerp(I3, 1.0, Ignore);
+	W4 *= lerp(I4, 1.0, Ignore);
 
 	OutColor0.rgb += C1.rgb * W1 + C2.rgb * W2 + C3.rgb * W3 + C4.rgb * W4;
 	Weight += W1+W2+W3+W4;
@@ -502,14 +623,17 @@ void CirclePS(
 
 
 // pixel shader to combine the full res scene and the blurred images behind and in front of the the focal plane
-void MainCircleRecombinePS(
-	in float4 UVAndScreenPos : TEXCOORD0,
-	out float4 OutColor : SV_Target0
-	)
+void MainCircleRecombinePS(in float4 UVAndScreenPos : TEXCOORD0, out float4 OutColor : SV_Target0)
 {
 	// Circle of confusion size for the pixel.
 	float PixDepth = CalcSceneDepth(UVAndScreenPos.xy);
-	float PixCoc = DepthToCoc(PixDepth) * 2.0; // 2x because full instead of half resolution.
+	float PixCoc = DepthToCoc(PixDepth);
+
+	// Grab nearest Coc.
+	PixCoc = min(PixCoc, PostprocessInput1.SampleLevel(PostprocessInput1Sampler, UVAndScreenPos.xy, 0).a);
+
+	// Transform into sample pattern.
+	PixCoc = abs(PixCoc) * 2.0; // 2x because full instead of half resolution.
 
 	// Fetch 2 samples mirrored around the pixel 
 	// which is stochastically distributed to fill out the circle of confusion.
@@ -518,19 +642,35 @@ void MainCircleRecombinePS(
 	float2 UV = UVAndScreenPos.xy * PostprocessInput0Size.xy;
 	float RadianBase = NoizNorm(UVAndScreenPos.xy, 0.010 * RandomOffset.x) * 3.14159 * 2.0;
 	float RadiusJitter = NoizNorm(UVAndScreenPos.xy, 0.013 * RandomOffset.x);
-	float2 VP = RotVec(float2(PixCoc*sqrt(sqrt(RadiusJitter)), 0.0), RadianBase) * PostprocessInput0Size.zw;
+	float ICoc = PixCoc*sqrt(RadiusJitter);
+	float2 VP = RotVec(float2(ICoc, 0.0), RadianBase) * PostprocessInput0Size.zw;
 
 	// These two samples will still have jitter induced artifacts (very limited utility).
 	// These two samples will also have bleeding artifacts.
 	float4 CA = PostprocessInput0.SampleLevel(PostprocessInput0Sampler, UVAndScreenPos.xy + VP, 0);
 	float4 CB = PostprocessInput0.SampleLevel(PostprocessInput0Sampler, UVAndScreenPos.xy - VP, 0);
-	OutColor = (CA + CB) * 0.5;
+
+	// Weight the two samples to avoid forground into background bleed.
+	float IFeather0 = 1.0/4.0;
+	ICoc *= 0.5; // Coc is half res units.
+	float I1 = saturate((abs(CA.a) - ICoc) * IFeather0);
+	float I2 = saturate((abs(CB.a) - ICoc) * IFeather0);
+	I1 += 1.0/65536.0;
+	I2 += 1.0/65536.0;
+
+	OutColor = (CA * I1 + CB * I2) * (1.0/(I1+I2));
 
 	// Grab the half resolution neighborhood to remove the artifacts from the full resolution output.
 	// Nearest location.
-	float2 HUVBase = UVAndScreenPos.xy * PostprocessInput1Size.xy - 0.5;
-	float2 HUVFrac = frac(HUVBase);
-	float2 HUV = (trunc(HUVBase) + 0.5) * PostprocessInput1Size.zw;
+	#if 1
+		// This has higher in-focus contrast, but possibly lower noise reduction later.
+		float2 HUVBase = UVAndScreenPos.xy * PostprocessInput1Size.xy - 0.5;
+		float2 HUVFrac = frac(HUVBase);
+		float2 HUV = (trunc(HUVBase) + 0.5) * PostprocessInput1Size.zw;
+	#else
+		// This makes the mostly in-focus transition bad (too blurry).
+		float2 HUV = UVAndScreenPos.xy - 0.5 * PostprocessInput1Size.zw;
+	#endif
 
 	// Load four nearest samples.
 	float4 H0 = PostprocessInput1.SampleLevel(PostprocessInput1Sampler, HUV, 0);
@@ -543,9 +683,17 @@ void MainCircleRecombinePS(
 	float4 HMax = max(max(H0,H1),max(H2,H3));
 	float4 HMin = min(min(H0,H1),min(H2,H3));
 
+	#if 1
+		// Increase constrast of limit a little to workaround to strong denoise at near-in-focus.
+		float4 HD = HMin / 8.0;
+		float Small = 1.0 - saturate(PixCoc*PixCoc*(1.0/64.0));
+		HMax += HD * Small;
+		HMin -= HD * Small;
+	#endif
+
 	// Blend in the limited version quickly to remove HDR jitter artifacts and noise.
 	float4 OutLimited = min(max(OutColor,HMin),HMax);
-	OutColor = lerp(OutColor, OutLimited, saturate(PixCoc*PixCoc*8.0));
+	OutColor = lerp(OutColor, OutLimited, saturate(PixCoc*PixCoc*4.0));
 }
 
 
diff --git a/Engine/Source/Runtime/Renderer/Private/PostProcess/PostProcessDOF.cpp b/Engine/Source/Runtime/Renderer/Private/PostProcess/PostProcessDOF.cpp
--- a/Engine/Source/Runtime/Renderer/Private/PostProcess/PostProcessDOF.cpp
+++ b/Engine/Source/Runtime/Renderer/Private/PostProcess/PostProcessDOF.cpp
@@ -359,7 +359,7 @@ class FPostProcessCircleDOFSetupPS : public FGlobalShader
 
 	static bool ShouldCache(EShaderPlatform Platform)
 	{
-		return IsFeatureLevelSupported(Platform, ERHIFeatureLevel::SM4);
+		return IsFeatureLevelSupported(Platform, ERHIFeatureLevel::SM5);
 	}
 
 	static void ModifyCompilationEnvironment(EShaderPlatform Platform, FShaderCompilerEnvironment& OutEnvironment)
@@ -535,6 +535,195 @@ FPooledRenderTargetDesc FRCPassPostProcessCircleDOFSetup::ComputeOutputDesc(EPas
 
 
 
+/** Encapsulates the Circle DOF Dilate pixel shader. */
+template <uint32 NearBlurEnable>
+class FPostProcessCircleDOFDilatePS : public FGlobalShader
+{
+	DECLARE_SHADER_TYPE(FPostProcessCircleDOFDilatePS, Global);
+
+	static bool ShouldCache(EShaderPlatform Platform)
+	{
+		return IsFeatureLevelSupported(Platform, ERHIFeatureLevel::SM5);
+	}
+
+	static void ModifyCompilationEnvironment(EShaderPlatform Platform, FShaderCompilerEnvironment& OutEnvironment)
+	{
+		FGlobalShader::ModifyCompilationEnvironment(Platform,OutEnvironment);
+		OutEnvironment.SetDefine(TEXT("ENABLE_NEAR_BLUR"), NearBlurEnable);
+	}
+
+	/** Default constructor. */
+	FPostProcessCircleDOFDilatePS() {}
+
+public:
+	FPostProcessPassParameters PostprocessParameter;
+	FDeferredPixelShaderParameters DeferredParameters;
+	FShaderParameter DepthOfFieldParams;
+
+	/** Initialization constructor. */
+	FPostProcessCircleDOFDilatePS(const ShaderMetaType::CompiledShaderInitializerType& Initializer)
+		: FGlobalShader(Initializer)
+	{
+		PostprocessParameter.Bind(Initializer.ParameterMap);
+		DeferredParameters.Bind(Initializer.ParameterMap);
+		DepthOfFieldParams.Bind(Initializer.ParameterMap,TEXT("DepthOfFieldParams"));
+	}
+
+	// FShader interface.
+	virtual bool Serialize(FArchive& Ar)
+	{
+		bool bShaderHasOutdatedParameters = FGlobalShader::Serialize(Ar);
+		Ar << PostprocessParameter << DeferredParameters << DepthOfFieldParams;
+		return bShaderHasOutdatedParameters;
+	}
+
+	void SetParameters(const FRenderingCompositePassContext& Context)
+	{
+		const FPixelShaderRHIParamRef ShaderRHI = GetPixelShader();
+
+		FGlobalShader::SetParameters(Context.RHICmdList, ShaderRHI, Context.View);
+
+		PostprocessParameter.SetPS(ShaderRHI, Context, TStaticSamplerState<SF_Point,AM_Border,AM_Border,AM_Clamp>::GetRHI());
+
+		DeferredParameters.Set(Context.RHICmdList, ShaderRHI, Context.View);
+
+		{
+			FVector4 DepthOfFieldParamValues[2];
+
+			FRCPassPostProcessBokehDOF::ComputeDepthOfFieldParams(Context, DepthOfFieldParamValues);
+
+			SetShaderValueArray(Context.RHICmdList, ShaderRHI, DepthOfFieldParams, DepthOfFieldParamValues, 2);
+		}
+	}
+};
+
+IMPLEMENT_SHADER_TYPE(template<>,FPostProcessCircleDOFDilatePS<0>,TEXT("PostProcessDOF"),TEXT("CircleDilatePS"),SF_Pixel);
+IMPLEMENT_SHADER_TYPE(template<>,FPostProcessCircleDOFDilatePS<1>,TEXT("PostProcessDOF"),TEXT("CircleDilatePS"),SF_Pixel);
+
+void FRCPassPostProcessCircleDOFDilate::Process(FRenderingCompositePassContext& Context)
+{
+	SCOPED_DRAW_EVENT(Context.RHICmdList, DOFSetup);
+
+	const FPooledRenderTargetDesc* InputDesc = GetInputDesc(ePId_Input0);
+
+	if(!InputDesc)
+	{
+		// input is not hooked up correctly
+		return;
+	}
+
+	uint32 NumRenderTargets = 1;
+
+	const FSceneView& View = Context.View;
+	const FSceneViewFamily& ViewFamily = *(View.Family);
+
+	const auto FeatureLevel = Context.GetFeatureLevel();
+	auto ShaderMap = Context.GetShaderMap();
+
+	FIntPoint SrcSize = InputDesc->Extent;
+	FIntPoint DestSize = PassOutputs[0].RenderTargetDesc.Extent;
+
+	// e.g. 4 means the input texture is 4x smaller than the buffer size
+	uint32 ScaleFactor = GSceneRenderTargets.GetBufferSizeXY().X / SrcSize.X;
+
+	FIntRect SrcRect = View.ViewRect / ScaleFactor;
+	FIntRect DestRect = SrcRect / 2;
+
+	const FSceneRenderTargetItem& DestRenderTarget0 = PassOutputs[0].RequestSurface(Context);
+	const FSceneRenderTargetItem& DestRenderTarget1 = FSceneRenderTargetItem();
+
+	// Set the view family's render target/viewport.
+	FTextureRHIParamRef RenderTargets[2] =
+	{
+		DestRenderTarget0.TargetableTexture,
+		DestRenderTarget1.TargetableTexture
+	};
+	SetRenderTargets(Context.RHICmdList, NumRenderTargets, RenderTargets, FTextureRHIParamRef(), 0, NULL);
+
+	FLinearColor ClearColors[2] = 
+	{
+		FLinearColor(0, 0, 0, 0),
+		FLinearColor(0, 0, 0, 0)
+	};
+	// is optimized away if possible (RT size=view size, )
+	Context.RHICmdList.ClearMRT(true, NumRenderTargets, ClearColors, false, 1.0f, false, 0, DestRect);
+
+	Context.SetViewportAndCallRHI(0, 0, 0.0f, DestSize.X, DestSize.Y, 1.0f );
+
+	// set the state
+	Context.RHICmdList.SetBlendState(TStaticBlendState<>::GetRHI());
+	Context.RHICmdList.SetRasterizerState(TStaticRasterizerState<>::GetRHI());
+	Context.RHICmdList.SetDepthStencilState(TStaticDepthStencilState<false, CF_Always>::GetRHI());
+
+	TShaderMapRef<FPostProcessVS> VertexShader(ShaderMap);
+
+	if (false)
+	{
+		static FGlobalBoundShaderState BoundShaderState;
+
+
+		TShaderMapRef< FPostProcessCircleDOFDilatePS<1> > PixelShader(ShaderMap);
+		SetGlobalBoundShaderState(Context.RHICmdList, FeatureLevel, BoundShaderState, GFilterVertexDeclaration.VertexDeclarationRHI, *VertexShader, *PixelShader);
+
+		PixelShader->SetParameters(Context);
+	}
+	else
+	{
+		static FGlobalBoundShaderState BoundShaderState;
+
+		TShaderMapRef< FPostProcessCircleDOFDilatePS<0> > PixelShader(ShaderMap);
+		SetGlobalBoundShaderState(Context.RHICmdList, FeatureLevel, BoundShaderState, GFilterVertexDeclaration.VertexDeclarationRHI, *VertexShader, *PixelShader);
+
+		PixelShader->SetParameters(Context);
+	}
+
+	VertexShader->SetParameters(Context);
+
+	// Draw a quad mapping scene color to the view's render target
+	DrawRectangle(
+		Context.RHICmdList,
+		DestRect.Min.X, DestRect.Min.Y,
+		DestRect.Width() + 1, DestRect.Height() + 1,
+		SrcRect.Min.X, SrcRect.Min.Y,
+		SrcRect.Width() + 1, SrcRect.Height() + 1,
+		DestSize,
+		SrcSize,
+		*VertexShader,
+		EDRF_UseTriangleOptimization);
+
+	Context.RHICmdList.CopyToResolveTarget(DestRenderTarget0.TargetableTexture, DestRenderTarget0.ShaderResourceTexture, false, FResolveParams());
+	Context.RHICmdList.CopyToResolveTarget(DestRenderTarget1.TargetableTexture, DestRenderTarget1.ShaderResourceTexture, false, FResolveParams());
+}
+
+FPooledRenderTargetDesc FRCPassPostProcessCircleDOFDilate::ComputeOutputDesc(EPassOutputId InPassOutputId) const
+{
+	FPooledRenderTargetDesc Ret = PassInputs[0].GetOutput()->RenderTargetDesc;
+
+	//	Ret.Extent = FIntPoint::DivideAndRoundUp(ret.Extent, 2);
+	Ret.Extent /= 2;
+	Ret.Extent.X = FMath::Max(1, Ret.Extent.X);
+	Ret.Extent.Y = FMath::Max(1, Ret.Extent.Y);
+
+	Ret.Reset();
+	Ret.TargetableFlags &= ~(uint32)TexCreate_UAV;
+	Ret.TargetableFlags |= TexCreate_RenderTargetable;
+
+	Ret.DebugName = (InPassOutputId == ePId_Output0) ? TEXT("CircleDOFDilate0") : TEXT("CircleDOFDilate1");
+
+	// more precision for additive blending and we need the alpha channel
+	Ret.Format = PF_FloatRGBA;
+
+	return Ret;
+}
+
+
+
+
+
+
+
+
+
 /** Encapsulates the Circle DOF pixel shader. */
 
 static float TemporalHalton2( int32 Index, int32 Base )
@@ -565,7 +754,7 @@ class FPostProcessCircleDOFPS : public FGlobalShader
 
 	static bool ShouldCache(EShaderPlatform Platform)
 	{
-		return IsFeatureLevelSupported(Platform, ERHIFeatureLevel::SM4);
+		return IsFeatureLevelSupported(Platform, ERHIFeatureLevel::SM5);
 	}
 
 	static void ModifyCompilationEnvironment(EShaderPlatform Platform, FShaderCompilerEnvironment& OutEnvironment)
@@ -751,7 +940,7 @@ class FPostProcessCircleDOFRecombinePS : public FGlobalShader
 
 	static bool ShouldCache(EShaderPlatform Platform)
 	{
-		return IsFeatureLevelSupported(Platform, ERHIFeatureLevel::SM4);
+		return IsFeatureLevelSupported(Platform, ERHIFeatureLevel::SM5);
 	}
 
 	static void ModifyCompilationEnvironment(EShaderPlatform Platform, FShaderCompilerEnvironment& OutEnvironment)
diff --git a/Engine/Source/Runtime/Renderer/Private/PostProcess/PostProcessDOF.h b/Engine/Source/Runtime/Renderer/Private/PostProcess/PostProcessDOF.h
--- a/Engine/Source/Runtime/Renderer/Private/PostProcess/PostProcessDOF.h
+++ b/Engine/Source/Runtime/Renderer/Private/PostProcess/PostProcessDOF.h
@@ -71,7 +71,24 @@ private:
 	bool bNearBlurEnabled;
 };
 
-// down sample and setup DOF input
+
+// ePId_Input0: SceneColor
+class FRCPassPostProcessCircleDOFDilate : public TRenderingCompositePassBase<1, 1>
+{
+public:
+
+	FRCPassPostProcessCircleDOFDilate(void) {}
+
+	// interface FRenderingCompositePass ---------
+
+	virtual void Process(FRenderingCompositePassContext& Context);
+	virtual void Release() override { delete this; }
+	virtual FPooledRenderTargetDesc ComputeOutputDesc(EPassOutputId InPassOutputId) const;
+};
+
+
+
+
 // ePId_Input0: setup results for far
 // ePId_Input1: setup results for near
 // derives from TRenderingCompositePassBase<InputCount, OutputCount> 
diff --git a/Engine/Source/Runtime/Renderer/Private/PostProcess/PostProcessing.cpp b/Engine/Source/Runtime/Renderer/Private/PostProcess/PostProcessing.cpp
--- a/Engine/Source/Runtime/Renderer/Private/PostProcess/PostProcessing.cpp
+++ b/Engine/Source/Runtime/Renderer/Private/PostProcess/PostProcessing.cpp
@@ -429,31 +429,14 @@ static void AddPostProcessDepthOfFieldGaussian(FPostprocessContext& Context, FDe
 
 static void AddPostProcessDepthOfFieldCircle(FPostprocessContext& Context, FDepthOfFieldStats& Out)
 {
-	float FarSize = Context.View.FinalPostProcessSettings.DepthOfFieldFarBlurSize;
-	float NearSize = Context.View.FinalPostProcessSettings.DepthOfFieldNearBlurSize;
-
-	float MaxSize = CVarDepthOfFieldMaxSize.GetValueOnRenderThread();
-
-	FarSize = FMath::Min(FarSize, MaxSize);
-	NearSize = FMath::Min(NearSize, MaxSize);
-
-	Out.bFar = FarSize >= 0.01f;
-	Out.bNear = NearSize >= GetCachedScalabilityCVars().GaussianDOFNearThreshold;
-
-	if(!Out.bFar && !Out.bNear)
-	{
-		return;
-	}
-
 	if(Context.View.Family->EngineShowFlags.VisualizeDOF)
 	{
 		// no need for this pass
 		return;
 	}
 
-	FRenderingCompositePass* DOFSetup = Context.Graph.RegisterPass(new(FMemStack::Get()) FRCPassPostProcessCircleDOFSetup(Out.bNear));
+	FRenderingCompositePass* DOFSetup = Context.Graph.RegisterPass(new(FMemStack::Get()) FRCPassPostProcessCircleDOFSetup(false));
 	DOFSetup->SetInput(ePId_Input0, FRenderingCompositeOutputRef(Context.FinalOutput));
-	// We need the depth to create the near and far mask
 	DOFSetup->SetInput(ePId_Input1, FRenderingCompositeOutputRef(Context.SceneDepth));
 
 	FSceneViewState* ViewState = (FSceneViewState*)Context.View.State;
@@ -482,52 +465,21 @@ static void AddPostProcessDepthOfFieldCircle(FPostprocessContext& Context, FDept
 		ViewState->bBokehDOFHistory = false;
 	}
 
-	FRenderingCompositePass* DOFInputPass2 = DOFSetup;
-	EPassOutputId DOFInputPassId = ePId_Output1;
-	if( Context.View.FinalPostProcessSettings.AntiAliasingMethod == AAM_TemporalAA && ViewState )
-	{
-		FRenderingCompositePass* HistoryInput;
-		EPassOutputId DOFInputPassId2 = ePId_Output1;
-		if( ViewState->DOFHistoryRT2 && !ViewState->bBokehDOFHistory2 && !Context.View.bCameraCut )
-		{
-			HistoryInput = Context.Graph.RegisterPass( new(FMemStack::Get()) FRCPassPostProcessInput( ViewState->DOFHistoryRT2 ) );
-			DOFInputPassId2 = ePId_Output0;
-		}
-		else
-		{
-			// No history so use current as history
-			HistoryInput = DOFSetup;
-			DOFInputPassId2 = ePId_Output1;
-		}
-
-		FRenderingCompositePass* NodeTemporalAA = Context.Graph.RegisterPass( new(FMemStack::Get()) FRCPassPostProcessDOFTemporalAANear );
-		NodeTemporalAA->SetInput( ePId_Input0, FRenderingCompositeOutputRef( DOFSetup, ePId_Output1 ) );
-		NodeTemporalAA->SetInput( ePId_Input1, FRenderingCompositeOutputRef( HistoryInput, DOFInputPassId2 ) );
-		NodeTemporalAA->SetInput( ePId_Input2, FRenderingCompositeOutputRef( HistoryInput, DOFInputPassId2 ) );
-		//NodeTemporalAA->SetInput( ePId_Input3, VelocityInput );
-
-		DOFInputPass2 = NodeTemporalAA;
-		ViewState->bBokehDOFHistory2 = false;
-
-		DOFInputPassId = ePId_Output0;
-	}
-
 	FRenderingCompositeOutputRef Far;
 	FRenderingCompositeOutputRef Near;
 
-	FRenderingCompositePass* DOFApply = Context.Graph.RegisterPass(new(FMemStack::Get()) FRCPassPostProcessCircleDOF(Out.bNear));
+	FRenderingCompositePass* DOFNear = Context.Graph.RegisterPass(new(FMemStack::Get()) FRCPassPostProcessCircleDOFDilate());
+	DOFNear->SetInput(ePId_Input0, FRenderingCompositeOutputRef(DOFInputPass, ePId_Output0));
+	Near = FRenderingCompositeOutputRef(DOFNear, ePId_Output0);
+
+	FRenderingCompositePass* DOFApply = Context.Graph.RegisterPass(new(FMemStack::Get()) FRCPassPostProcessCircleDOF(false));
 	DOFApply->SetInput(ePId_Input0, FRenderingCompositeOutputRef(DOFInputPass, ePId_Output0));
+	DOFApply->SetInput(ePId_Input1, Near);
 	Far = FRenderingCompositeOutputRef(DOFApply, ePId_Output0);
-	if(Out.bNear)
-	{
-		DOFApply->SetInput(ePId_Input1, FRenderingCompositeOutputRef(DOFInputPass2, DOFInputPassId));
-		Near = FRenderingCompositeOutputRef(DOFApply, ePId_Output1);
-	}
 
-	FRenderingCompositePass* NodeRecombined = Context.Graph.RegisterPass(new(FMemStack::Get()) FRCPassPostProcessCircleDOFRecombine(Out.bNear));
+	FRenderingCompositePass* NodeRecombined = Context.Graph.RegisterPass(new(FMemStack::Get()) FRCPassPostProcessCircleDOFRecombine(false));
 	NodeRecombined->SetInput(ePId_Input0, Context.FinalOutput);
 	NodeRecombined->SetInput(ePId_Input1, Far);
-	NodeRecombined->SetInput(ePId_Input2, Near);
 
 	Context.FinalOutput = FRenderingCompositeOutputRef(NodeRecombined);
 }
