(Untested) fix for shader compilation error with GLSL.

diff --git a/Engine/Shaders/PostProcessDOF.usf b/Engine/Shaders/PostProcessDOF.usf
--- a/Engine/Shaders/PostProcessDOF.usf
+++ b/Engine/Shaders/PostProcessDOF.usf
@@ -328,11 +328,37 @@ void CircleDilatePS(float4 UVAndScreenPos : TEXCOORD0, out float OutColor : SV_T
 		//   j M  (M={0,0} point)
 	
 		float2 UV = UVAndScreenPos.xy + 1.0*PostprocessInput0Size.zw;
-	
+
+#if COMPILER_GLSL || COMPILER_GLSL_ES2
+		float4 Sd, Sg, Sj, Sm;
+		// x ~ uv+(0,1)*offset, y ~ uv+(1,1)*offset, z ~ uv+(1,0)*offset, w ~ uv
+		float ZeroOffsetA = Texture2DSample(PostprocessInput0, PostprocessInput0Sampler, UV).a;
+		
+		Sd.x = Texture2DSample(PostprocessInput0, PostprocessInput0Sampler, UV + float2(0,-2)).a;
+		Sd.y = Texture2DSample(PostprocessInput0, PostprocessInput0Sampler, UV + float2(-2,-2)).a;
+		Sd.z = Texture2DSample(PostprocessInput0, PostprocessInput0Sampler, UV + float2(-2,0)).a;
+		Sd.w = ZeroOffsetA;
+
+		Sg.x = Texture2DSample(PostprocessInput0, PostprocessInput0Sampler, UV + float2(0,-2)).a;
+		Sg.y = Sg.x;
+		Sg.z = ZeroOffsetA;
+		Sg.w = ZeroOffsetA;
+
+		Sj.x = ZeroOffsetA;
+		Sj.y = Texture2DSample(PostprocessInput0, PostprocessInput0Sampler, UV + float2(-2,0)).a;
+		Sj.z = Sj.y;
+		Sj.w = ZeroOffsetA;
+
+		Sm.x = ZeroOffsetA;
+		Sm.y = ZeroOffsetA;
+		Sm.z = ZeroOffsetA;
+		Sm.w = ZeroOffsetA;
+#else
 		float4 Sd = PostprocessInput0.GatherAlpha(PostprocessInput0Sampler, UV, int2(-2,-2));
 		float4 Sg = PostprocessInput0.GatherAlpha(PostprocessInput0Sampler, UV, int2(0,-2));
 		float4 Sj = PostprocessInput0.GatherAlpha(PostprocessInput0Sampler, UV, int2(-2,0));
 		float4 Sm = PostprocessInput0.GatherAlpha(PostprocessInput0Sampler, UV, int2(0,0));
+#endif // COMPILER_GLSL
 	
 		// Make sure near is only near blur.
 		OutColor = min(0.0, Min16(Sd, Sg, Sj, Sm));
