Use some Mac-only normalize calls on USE_SPLINEDEFORM uniform paramters to workaround [UE-5598] which only affects the Intel HD5000/Iris/Iris Pro GPU family.

diff --git a/Engine/Shaders/LocalVertexFactory.usf b/Engine/Shaders/LocalVertexFactory.usf
--- a/Engine/Shaders/LocalVertexFactory.usf
+++ b/Engine/Shaders/LocalVertexFactory.usf
@@ -369,7 +369,11 @@ FMaterialVertexParameters GetMaterialVertexParameters(FVertexFactoryInput Input,
 		YVec *= UseScale.y;
 
 		// Build overall transform
+#if MAC // @todo Remove once there is a fix for Intel Iris/HD5000 GPUs from rendering corrupt geometry without the normalize calls
+		float3x3 SliceTransform3 = mul(transpose(float3x3(normalize(SplineMeshDir), normalize(SplineMeshX), normalize(SplineMeshY))), float3x3(float3(0,0,0), XVec, YVec));
+#else
 		float3x3 SliceTransform3 = mul(transpose(float3x3(SplineMeshDir, SplineMeshX, SplineMeshY)), float3x3(float3(0,0,0), XVec, YVec));
+#endif
 		float4x3 SliceTransform = float4x3(SliceTransform3[0], SliceTransform3[1], SliceTransform3[2], SplinePos);
 
 		return SliceTransform;
@@ -405,7 +409,11 @@ FMaterialVertexParameters GetMaterialVertexParameters(FVertexFactoryInput Input,
 		YVec *= sign(UseScale.y);
 
 		// Build rotation transform
+#if MAC // @todo Remove once there is a fix for Intel Iris/HD5000 GPUs from rendering corrupt geometry without the normalize calls
+		half3x3 SliceTransform = mul(transpose(half3x3(normalize(SplineMeshDir), normalize(SplineMeshX), normalize(SplineMeshY))), half3x3(SplineDir, XVec, YVec));
+#else
 		half3x3 SliceTransform = mul(transpose(half3x3(SplineMeshDir, SplineMeshX, SplineMeshY)), half3x3(SplineDir, XVec, YVec));
+#endif
 
 		return SliceTransform;
 	}
