Fix eye material compile on PS4

diff --git a/Engine/Shaders/Definitions.usf b/Engine/Shaders/Definitions.usf
--- a/Engine/Shaders/Definitions.usf
+++ b/Engine/Shaders/Definitions.usf
@@ -163,6 +163,8 @@
 #define IFALL											[ifAll]
 /** Executes the conditional part of an if statement when the condition is true for any thread on which the current shader is running. */
 #define IFANY											[ifAny]
+/** Allows a compute shader loop termination condition to be based off of a UAV read. The loop must not contain synchronization intrinsics. */
+#define ALLOW_UAV_CONDITION								[allow_uav_condition]
 #else
 #ifndef UNROLL
 #define UNROLL
diff --git a/Engine/Shaders/ForEachLight.usf b/Engine/Shaders/ForEachLight.usf
--- a/Engine/Shaders/ForEachLight.usf
+++ b/Engine/Shaders/ForEachLight.usf
@@ -73,7 +73,7 @@ float3 AccumulateForwardLights(FMaterialPixelParameters Parameters, float3 Diffu
 	// todo: support more than 32 lights
 
 	// loop through all lights specified by this bitmask
-	[allow_uav_condition] while(culling != 0)
+	ALLOW_UAV_CONDITION while (culling != 0)
 	{
 		uint LightIndex = firstbitlow(culling);
 
