Integrate - Only the local vertex factory uses speedtree wind code
* Prevents errors reading parameters that don't get set

diff --git a/Engine/Shaders/MaterialTemplate.usf b/Engine/Shaders/MaterialTemplate.usf
--- a/Engine/Shaders/MaterialTemplate.usf
+++ b/Engine/Shaders/MaterialTemplate.usf
@@ -848,11 +848,15 @@ float3 GetSpeedTreeVertexOffsetInner(FMaterialVertexParameters Parameters, int G
 /** Vertex offset for SpeedTree wind and LOD */
 float3 GetSpeedTreeVertexOffset(FMaterialVertexParameters Parameters, int GeometryType, int WindType, int LODType, float BillboardThreshold, bool bUsePreviousFrame) 
 {
+#if VF_SUPPORTS_SPEEDTREE_WIND
 	if (bUsePreviousFrame)
 	{
 		return GetSpeedTreeVertexOffsetInner(Parameters, GeometryType, WindType, LODType, BillboardThreshold, GetPreviousSpeedTreeData());
 	}
 	return GetSpeedTreeVertexOffsetInner(Parameters, GeometryType, WindType, LODType, BillboardThreshold, GetCurrentSpeedTreeData());
+#else
+	return 0;
+#endif
 }
 
 #endif
diff --git a/Engine/Source/Runtime/Engine/Public/LocalVertexFactory.h b/Engine/Source/Runtime/Engine/Public/LocalVertexFactory.h
--- a/Engine/Source/Runtime/Engine/Public/LocalVertexFactory.h
+++ b/Engine/Source/Runtime/Engine/Public/LocalVertexFactory.h
@@ -40,6 +40,11 @@ public:
 	 */
 	static bool ShouldCache(EShaderPlatform Platform, const class FMaterial* Material, const class FShaderType* ShaderType);
 
+	static void ModifyCompilationEnvironment(EShaderPlatform Platform, const FMaterial* Material, FShaderCompilerEnvironment& OutEnvironment)
+	{
+		OutEnvironment.SetDefine(TEXT("VF_SUPPORTS_SPEEDTREE_WIND"),TEXT("1"));
+	}
+
 	/**
 	 * An implementation of the interface used by TSynchronizedResource to update the resource with new data from the game thread.
 	 */
