Integrate - Fixed PixelDepthOffset in combination with accurate velocities

diff --git a/Engine/Shaders/BasePassPixelShader.usf b/Engine/Shaders/BasePassPixelShader.usf
--- a/Engine/Shaders/BasePassPixelShader.usf
+++ b/Engine/Shaders/BasePassPixelShader.usf
@@ -573,6 +573,19 @@ uint CustomAlpha2Coverage(inout float4 InOutColor)
 }
 #endif
 
+void ApplyPixelDepthOffsetForBasePass(inout FMaterialPixelParameters MaterialParameters, inout FBasePassInterpolantsVSToPS BasePassInterpolants, out float OutDepth)
+{
+	float PixelDepthOffset = ApplyPixelDepthOffsetToMaterialParameters(MaterialParameters, OutDepth);
+
+#if WRITES_VELOCITY_TO_GBUFFER
+	BasePassInterpolants.VelocityPrevScreenPosition.w += PixelDepthOffset;
+
+	#if WRITES_VELOCITY_TO_GBUFFER_USE_POS_INTERPOLATOR
+		BasePassInterpolants.VelocityScreenPosition.w += PixelDepthOffset;
+	#endif
+#endif
+}
+
 void Main(
 	FVertexFactoryInterpolantsVSToPS Interpolants,
 	FBasePassInterpolantsVSToPS BasePassInterpolants,
@@ -624,7 +637,7 @@ void Main(
 #endif
 		
 #if OUTPUT_PIXEL_DEPTH_OFFSET
-	ApplyPixelDepthOffsetToMaterialParameters(MaterialParameters, OutDepth);
+	ApplyPixelDepthOffsetForBasePass(MaterialParameters, BasePassInterpolants, OutDepth);
 #endif
 
 	//Clip if the blend mode requires it.
@@ -633,7 +646,6 @@ void Main(
 		GetMaterialCoverageAndClipping(MaterialParameters);
 	}
 
-
 	// Store the results in local variables and reuse instead of calling the functions multiple times.
 	half3 BaseColor = GetMaterialBaseColor( MaterialParameters );
 	half  Metallic = GetMaterialMetallic( MaterialParameters );
diff --git a/Engine/Shaders/MaterialTemplate.usf b/Engine/Shaders/MaterialTemplate.usf
--- a/Engine/Shaders/MaterialTemplate.usf
+++ b/Engine/Shaders/MaterialTemplate.usf
@@ -1499,7 +1499,7 @@ half3x3 AssembleTangentToWorld( half3 TangentToWorld0, half4 TangentToWorld2 )
 	#define OPTIONAL_OutDepth
 #endif
 
-void ApplyPixelDepthOffsetToMaterialParameters(inout FMaterialPixelParameters MaterialParameters, out float OutDepth)
+float ApplyPixelDepthOffsetToMaterialParameters(inout FMaterialPixelParameters MaterialParameters, out float OutDepth)
 {
 	// Can only push into the scene for conservative depth writes
 	float PixelDepthOffset = max(GetMaterialPixelDepthOffset(MaterialParameters), 0);
@@ -1510,4 +1510,6 @@ void ApplyPixelDepthOffsetToMaterialParameters(inout FMaterialPixelParameters Ma
 	MaterialParameters.WorldPosition += MaterialParameters.CameraVector * PixelDepthOffset;
 
 	OutDepth = MaterialParameters.ScreenPosition.z / MaterialParameters.ScreenPosition.w;
+
+	return PixelDepthOffset;
 }
