CircleDOF: Add real camera controls, correctly fix for not having gatherAlpha(), etc.

diff --git a/Engine/Shaders/PostProcessDOF.usf b/Engine/Shaders/PostProcessDOF.usf
--- a/Engine/Shaders/PostProcessDOF.usf
+++ b/Engine/Shaders/PostProcessDOF.usf
@@ -198,11 +198,13 @@ void MainRecombinePS(
 // PROTOTYPE CIRCLE DOF : WORK IN PROGRESS
 //
 
+float CircleDofParams;
+
 // CIRCLE DOF: Compute circle of confusion size in pixels.
 float DepthToCoc(float Depth) 
 {
 	float Focus = View.DepthOfFieldFocalDistance;
-	float Radius = View.DepthOfFieldScale * 2.0;
+	float Radius = CircleDofParams;
 	return ((Depth - Focus) / Depth) * Radius;
 }
 
@@ -284,86 +286,44 @@ float Min16(float4 A, float4 B, float4 C, float4 D)
 	return min(min(Min4(A),Min4(B)),min(Min4(C),Min4(D)));
 }
 
-// This does not actually do any dilation, but rather just a 2x2:1 min reduction.
+// This does a 2x2:1 reduction with a 4x4:1 dilation.
 void CircleDilatePS(float4 UVAndScreenPos : TEXCOORD0, out float OutColor : SV_Target0)
 {
-	#if 0
-		// Leaving code in for a 4x4:1 reduction (but not using).
-	
-		// Sampling pattern (each gather4)
-		//   a b e f
-	    //   c d g h
-		//   i j M n  (M={0,0} point)
-		//   k l o p
-
-		float2 UV = UVAndScreenPos.xy + 1.0*PostprocessInput0Size.zw;
-	
-		float4 Sa = PostprocessInput0.GatherAlpha(PostprocessInput0Sampler, UV, int2(-4,-4));
-		float4 Sb = PostprocessInput0.GatherAlpha(PostprocessInput0Sampler, UV, int2(-2,-4));
-		float4 Sc = PostprocessInput0.GatherAlpha(PostprocessInput0Sampler, UV, int2(-4,-2));
-		float4 Sd = PostprocessInput0.GatherAlpha(PostprocessInput0Sampler, UV, int2(-2,-2));
-
-		float4 Se = PostprocessInput0.GatherAlpha(PostprocessInput0Sampler, UV, int2(0,-4));
-		float4 Sf = PostprocessInput0.GatherAlpha(PostprocessInput0Sampler, UV, int2(2,-4));
-		float4 Sg = PostprocessInput0.GatherAlpha(PostprocessInput0Sampler, UV, int2(0,-2));
-		float4 Sh = PostprocessInput0.GatherAlpha(PostprocessInput0Sampler, UV, int2(2,-2));
-
-		float4 Si = PostprocessInput0.GatherAlpha(PostprocessInput0Sampler, UV, int2(-4,0));
-		float4 Sj = PostprocessInput0.GatherAlpha(PostprocessInput0Sampler, UV, int2(-2,0));
-		float4 Sk = PostprocessInput0.GatherAlpha(PostprocessInput0Sampler, UV, int2(-4,2));
-		float4 Sl = PostprocessInput0.GatherAlpha(PostprocessInput0Sampler, UV, int2(-2,2));
-
-		float4 Sm = PostprocessInput0.GatherAlpha(PostprocessInput0Sampler, UV, int2(0,0));
-		float4 Sn = PostprocessInput0.GatherAlpha(PostprocessInput0Sampler, UV, int2(2,0));
-		float4 So = PostprocessInput0.GatherAlpha(PostprocessInput0Sampler, UV, int2(0,2));
-		float4 Sp = PostprocessInput0.GatherAlpha(PostprocessInput0Sampler, UV, int2(2,2));
-	
-		// Make sure near is only near blur.
-		OutColor = min(0.0,min(min(Min16(Sa,Sb,Sc,Sd),Min16(Se,Sf,Sg,Sh)),min(Min16(Si,Sj,Sk,Sl),Min16(Sm,Sn,So,Sp))));
-	
+	// Sampling pattern (each gather4)
+    //   d g
+	//   j M  (M={0,0} point)
+
+	#if COMPILER_GLSL || COMPILER_GLSL_ES2
+		float2 UV = UVAndScreenPos.xy + 0.5*PostprocessInput0Size.zw;
+		// This leverages nearest sampling (bilinear won't work).
+		// Probably not the best way to do this.
+		float4 Sd, Sg, Sj, Sm;
+		Sd.x = PostprocessInput0.SampleLevel(PostprocessInput0Sampler, UV.xy, 0, int2(-2,-2)).a;
+		Sd.y = PostprocessInput0.SampleLevel(PostprocessInput0Sampler, UV.xy, 0, int2(-1,-2)).a;
+		Sd.z = PostprocessInput0.SampleLevel(PostprocessInput0Sampler, UV.xy, 0, int2(-2,-1)).a;
+		Sd.w = PostprocessInput0.SampleLevel(PostprocessInput0Sampler, UV.xy, 0, int2(-1,-1)).a;
+		Sg.x = PostprocessInput0.SampleLevel(PostprocessInput0Sampler, UV.xy, 0, int2(0,-2)).a;
+		Sg.y = PostprocessInput0.SampleLevel(PostprocessInput0Sampler, UV.xy, 0, int2(1,-2)).a;
+		Sg.z = PostprocessInput0.SampleLevel(PostprocessInput0Sampler, UV.xy, 0, int2(0,-1)).a;
+		Sg.w = PostprocessInput0.SampleLevel(PostprocessInput0Sampler, UV.xy, 0, int2(1,-1)).a;
+		Sj.x = PostprocessInput0.SampleLevel(PostprocessInput0Sampler, UV.xy, 0, int2(-2,0)).a;
+		Sj.y = PostprocessInput0.SampleLevel(PostprocessInput0Sampler, UV.xy, 0, int2(-1,0)).a;
+		Sj.z = PostprocessInput0.SampleLevel(PostprocessInput0Sampler, UV.xy, 0, int2(-2,1)).a;
+		Sj.w = PostprocessInput0.SampleLevel(PostprocessInput0Sampler, UV.xy, 0, int2(-1,1)).a;
+		Sm.x = PostprocessInput0.SampleLevel(PostprocessInput0Sampler, UV.xy, 0, int2(0,0)).a;
+		Sm.y = PostprocessInput0.SampleLevel(PostprocessInput0Sampler, UV.xy, 0, int2(1,0)).a;
+		Sm.z = PostprocessInput0.SampleLevel(PostprocessInput0Sampler, UV.xy, 0, int2(0,1)).a;
+		Sm.w = PostprocessInput0.SampleLevel(PostprocessInput0Sampler, UV.xy, 0, int2(1,1)).a;
 	#else
-	
-		// Sampling pattern (each gather4)
-	    //   d g
-		//   j M  (M={0,0} point)
-	
 		float2 UV = UVAndScreenPos.xy + 1.0*PostprocessInput0Size.zw;
-
-#if COMPILER_GLSL || COMPILER_GLSL_ES2
-		float4 Sd, Sg, Sj, Sm;
-		// x ~ uv+(0,1)*offset, y ~ uv+(1,1)*offset, z ~ uv+(1,0)*offset, w ~ uv
-		float ZeroOffsetA = Texture2DSample(PostprocessInput0, PostprocessInput0Sampler, UV).a;
-		
-		Sd.x = Texture2DSample(PostprocessInput0, PostprocessInput0Sampler, UV + float2(0,-2)).a;
-		Sd.y = Texture2DSample(PostprocessInput0, PostprocessInput0Sampler, UV + float2(-2,-2)).a;
-		Sd.z = Texture2DSample(PostprocessInput0, PostprocessInput0Sampler, UV + float2(-2,0)).a;
-		Sd.w = ZeroOffsetA;
-
-		Sg.x = Texture2DSample(PostprocessInput0, PostprocessInput0Sampler, UV + float2(0,-2)).a;
-		Sg.y = Sg.x;
-		Sg.z = ZeroOffsetA;
-		Sg.w = ZeroOffsetA;
-
-		Sj.x = ZeroOffsetA;
-		Sj.y = Texture2DSample(PostprocessInput0, PostprocessInput0Sampler, UV + float2(-2,0)).a;
-		Sj.z = Sj.y;
-		Sj.w = ZeroOffsetA;
-
-		Sm.x = ZeroOffsetA;
-		Sm.y = ZeroOffsetA;
-		Sm.z = ZeroOffsetA;
-		Sm.w = ZeroOffsetA;
-#else
 		float4 Sd = PostprocessInput0.GatherAlpha(PostprocessInput0Sampler, UV, int2(-2,-2));
 		float4 Sg = PostprocessInput0.GatherAlpha(PostprocessInput0Sampler, UV, int2(0,-2));
 		float4 Sj = PostprocessInput0.GatherAlpha(PostprocessInput0Sampler, UV, int2(-2,0));
 		float4 Sm = PostprocessInput0.GatherAlpha(PostprocessInput0Sampler, UV, int2(0,0));
-#endif // COMPILER_GLSL
+	#endif // COMPILER_GLSL
 	
-		// Make sure near is only near blur.
-		OutColor = min(0.0, Min16(Sd, Sg, Sj, Sm));
-	
-	#endif
+	// Make sure near is only near blur.
+	OutColor = min(0.0, Min16(Sd, Sg, Sj, Sm));
 }
 
 
@@ -389,10 +349,10 @@ void CirclePS(float4 UVAndScreenPos : TEXCOORD0, out float4 OutColor0 : SV_Targe
 
 	// Radius 
 	float RadiusBase2 = RadiusBase * (1.0/4.0);
-	float R1 = (RadiusBase2 + 3.0/4.0) * Coc;
-	float R2 = (RadiusBase2 + 2.0/4.0) * Coc;
-	float R3 = (RadiusBase2 + 1.0/4.0) * Coc;
-	float R4 = (RadiusBase2 + 0.0/4.0) * Coc;
+	float R1 = sqrt(RadiusBase2 + 3.0/4.0) * Coc;
+	float R2 = sqrt(RadiusBase2 + 2.0/4.0) * Coc;
+	float R3 = sqrt(RadiusBase2 + 1.0/4.0) * Coc;
+	float R4 = sqrt(RadiusBase2 + 0.0/4.0) * Coc;
 
 	float2 UV1 = RotVec(R1, RadianBase + TwoPi * 0.0/4.0);
 	float2 UV2 = RotVec(R2, RadianBase + TwoPi * 2.0/4.0);
@@ -410,12 +370,13 @@ void CirclePS(float4 UVAndScreenPos : TEXCOORD0, out float4 OutColor0 : SV_Targe
 	float D4 = PostprocessInput1.SampleLevel(PostprocessInput1Sampler, UV4, 0).x;
 
 	float NearCoc = 65536.0;
-	float Feather = 0.0;
+	float Feather = 2.0;
 	if(abs(D1)+Feather > R1) NearCoc = min(NearCoc, D1);
 	if(abs(D2)+Feather > R2) NearCoc = min(NearCoc, D2);
 	if(abs(D3)+Feather > R3) NearCoc = min(NearCoc, D3);
 	if(abs(D4)+Feather > R4) NearCoc = min(NearCoc, D4);
 
+
 	//
 	// Pass 1
 	//
@@ -661,30 +622,77 @@ void MainCircleRecombinePS(in float4 UVAndScreenPos : TEXCOORD0, out float4 OutC
 	// Transform into sample pattern.
 	PixCoc = abs(PixCoc) * 2.0; // 2x because full instead of half resolution.
 
-	// Fetch 2 samples mirrored around the pixel 
-	// which is stochastically distributed to fill out the circle of confusion.
-	// TODO: Fix the "random values".
-	// The sqrt(sqrt()) is to shape toward higher weight on the CoC edge.
-	float2 UV = UVAndScreenPos.xy * PostprocessInput0Size.xy;
-	float RadianBase = NoizNorm(UVAndScreenPos.xy, 0.010 * RandomOffset.x) * 3.14159 * 2.0;
-	float RadiusJitter = NoizNorm(UVAndScreenPos.xy, 0.013 * RandomOffset.x);
-	float ICoc = PixCoc*sqrt(RadiusJitter);
-	float2 VP = RotVec(float2(ICoc, 0.0), RadianBase) * PostprocessInput0Size.zw;
-
-	// These two samples will still have jitter induced artifacts (very limited utility).
-	// These two samples will also have bleeding artifacts.
-	float4 CA = PostprocessInput0.SampleLevel(PostprocessInput0Sampler, UVAndScreenPos.xy + VP, 0);
-	float4 CB = PostprocessInput0.SampleLevel(PostprocessInput0Sampler, UVAndScreenPos.xy - VP, 0);
-
-	// Weight the two samples to avoid forground into background bleed.
-	float IFeather0 = 1.0/4.0;
-	ICoc *= 0.5; // Coc is half res units.
-	float I1 = saturate((abs(CA.a) - ICoc) * IFeather0);
-	float I2 = saturate((abs(CB.a) - ICoc) * IFeather0);
-	I1 += 1.0/65536.0;
-	I2 += 1.0/65536.0;
-
-	OutColor = (CA * I1 + CB * I2) * (1.0/(I1+I2));
+	#if 1
+		// Fetch 2 samples mirrored around the pixel 
+		// which is stochastically distributed to fill out the circle of confusion.
+		// TODO: Fix the "random values".
+		float2 UV = UVAndScreenPos.xy * PostprocessInput0Size.xy;
+		float RadianBase = NoizNorm(UVAndScreenPos.xy, 0.010 * RandomOffset.x) * 3.14159 * 2.0;
+		float RadiusJitter = NoizNorm(UVAndScreenPos.xy, 0.013 * RandomOffset.x);
+		float ICoc = PixCoc*sqrt(RadiusJitter);
+		float2 VP = RotVec(float2(ICoc, 0.0), RadianBase) * PostprocessInput0Size.zw;
+
+		// These two samples will still have jitter induced artifacts (very limited utility).
+		// These two samples will also have bleeding artifacts.
+		float4 CA = PostprocessInput0.SampleLevel(PostprocessInput0Sampler, UVAndScreenPos.xy + VP, 0);
+		float4 CB = PostprocessInput0.SampleLevel(PostprocessInput0Sampler, UVAndScreenPos.xy - VP, 0);
+
+		// Weight the two samples to avoid forground into background bleed.
+		float IFeather0 = 1.0/4.0;
+		ICoc *= 0.5; // Coc is half res units.
+		float I1 = saturate((abs(CA.a) - ICoc) * IFeather0);
+		float I2 = saturate((abs(CB.a) - ICoc) * IFeather0);
+		I1 += 1.0/65536.0;
+		I2 += 1.0/65536.0;
+	
+		OutColor = (CA * I1 + CB * I2) * (1.0/(I1+I2));
+	#else
+	    // Possibly higher quality option in the future.
+
+		// Fetch 4 samples in filled disc pattern
+		// which is stochastically distributed to fill out the circle of confusion.
+		float2 UV = UVAndScreenPos.xy * PostprocessInput0Size.xy;
+		float RadianBase = NoizNorm(UVAndScreenPos.xy, 0.010 * RandomOffset.x) * 3.14159 * 2.0;
+		float RadiusBase = NoizNorm(UVAndScreenPos.xy, 0.013 * RandomOffset.x);
+
+		float RadiusBase2 = RadiusBase * (1.0/4.0);
+		float R1 = sqrt(RadiusBase2 + 3.0/4.0) * PixCoc;
+		float R2 = sqrt(RadiusBase2 + 2.0/4.0) * PixCoc;
+		float R3 = sqrt(RadiusBase2 + 1.0/4.0) * PixCoc;
+		float R4 = sqrt(RadiusBase2 + 0.0/4.0) * PixCoc;
+
+		float TwoPi = 3.14159 * 2.0;
+		float2 UV1 = RotVec(R1, RadianBase + TwoPi * 0.0/4.0);
+		float2 UV2 = RotVec(R2, RadianBase + TwoPi * 2.0/4.0);
+		float2 UV3 = RotVec(R3, RadianBase + TwoPi * 1.0/4.0);
+		float2 UV4 = RotVec(R4, RadianBase + TwoPi * 3.0/4.0);
+
+		UV1 = UVAndScreenPos.xy + UV1 * PostprocessInput0Size.zw;
+		UV2 = UVAndScreenPos.xy + UV2 * PostprocessInput0Size.zw;
+		UV3 = UVAndScreenPos.xy + UV3 * PostprocessInput0Size.zw;
+		UV4 = UVAndScreenPos.xy + UV4 * PostprocessInput0Size.zw;
+
+		float4 CA = PostprocessInput0.SampleLevel(PostprocessInput0Sampler, UV1, 0);
+		float4 CB = PostprocessInput0.SampleLevel(PostprocessInput0Sampler, UV2, 0);
+		float4 CC = PostprocessInput0.SampleLevel(PostprocessInput0Sampler, UV3, 0);
+		float4 CD = PostprocessInput0.SampleLevel(PostprocessInput0Sampler, UV4, 0);
+
+		// Weight the two samples to avoid forground into background bleed.
+		float IFeather0 = 1.0/4.0;
+		float ICoc = PixCoc * 0.5; // Coc is half res units.
+		float I1 = saturate((abs(CA.a) - ICoc) * IFeather0);
+		float I2 = saturate((abs(CB.a) - ICoc) * IFeather0);
+		float I3 = saturate((abs(CC.a) - ICoc) * IFeather0);
+		float I4 = saturate((abs(CD.a) - ICoc) * IFeather0);
+
+		// Make sure something is non-zero.
+		I1 += 1.0/65536.0;
+		I2 += 1.0/65536.0;
+		I3 += 1.0/65536.0;
+		I4 += 1.0/65536.0;
+
+		OutColor = (CA * I1 + CB * I2 + CC * I3 + CD * I4) * (1.0/(I1+I2+I3+I4));
+	#endif
 
 	// Grab the half resolution neighborhood to remove the artifacts from the full resolution output.
 	// Nearest location.
diff --git a/Engine/Source/Runtime/Engine/Classes/Engine/Scene.h b/Engine/Source/Runtime/Engine/Classes/Engine/Scene.h
--- a/Engine/Source/Runtime/Engine/Classes/Engine/Scene.h
+++ b/Engine/Source/Runtime/Engine/Classes/Engine/Scene.h
@@ -317,6 +317,9 @@ struct FPostProcessSettings
 	UPROPERTY(BlueprintReadWrite, Category=Overrides, meta=(PinHiddenByDefault))
 	uint32 bOverride_DepthOfFieldFocalDistance:1;
 
+	UPROPERTY(BlueprintReadWrite, Category=Overrides, meta=(PinHiddenByDefault))
+	uint32 bOverride_DepthOfFieldFstop:1;
+
 	UPROPERTY(BlueprintReadWrite, Category=Overrides, meta=(PinHiddenByDefault))
 	uint32 bOverride_DepthOfFieldFocalRegion:1;
 
@@ -766,6 +769,10 @@ struct FPostProcessSettings
 	UPROPERTY(EditAnywhere, BlueprintReadWrite, Category=DepthOfField, meta=(editcondition = "bOverride_DepthOfFieldMethod", DisplayName = "Method"))
 	TEnumAsByte<enum EDepthOfFieldMethod> DepthOfFieldMethod;
 
+	/** CircleDOF only: F-stop number (1/fstop) */
+	UPROPERTY(interp, BlueprintReadWrite, Category=DepthOfField, meta=(ClampMin = "1.0", ClampMax = "32.0", editcondition = "bOverride_DepthOfFieldFstop", DisplayName = "1/Fstop"))
+	float DepthOfFieldFstop;
+
 	/** Distance in which the Depth of Field effect should be sharp, in unreal units (cm) */
 	UPROPERTY(interp, BlueprintReadWrite, Category=DepthOfField, meta=(UIMin = "0.0", UIMax = "10000.0", editcondition = "bOverride_DepthOfFieldFocalDistance", DisplayName = "Focal Distance"))
 	float DepthOfFieldFocalDistance;
@@ -970,6 +977,7 @@ struct FPostProcessSettings
 		IndirectLightingIntensity = 1.0f;
 		ColorGradingIntensity = 1.0f;
 		DepthOfFieldFocalDistance = 1000.0f;
+		DepthOfFieldFstop = 4.0f;
 		DepthOfFieldFocalRegion = 0.0f;
 		DepthOfFieldNearTransitionRegion = 300.0f;
 		DepthOfFieldFarTransitionRegion = 500.0f;
diff --git a/Engine/Source/Runtime/Engine/Private/SceneView.cpp b/Engine/Source/Runtime/Engine/Private/SceneView.cpp
--- a/Engine/Source/Runtime/Engine/Private/SceneView.cpp
+++ b/Engine/Source/Runtime/Engine/Private/SceneView.cpp
@@ -812,6 +812,7 @@ void FSceneView::OverridePostProcessSettings(const FPostProcessSettings& Src, fl
 	LERP_PP(IndirectLightingColor);
 	LERP_PP(IndirectLightingIntensity);
 	LERP_PP(DepthOfFieldFocalDistance);
+	LERP_PP(DepthOfFieldFstop);
 	LERP_PP(DepthOfFieldFocalRegion);
 	LERP_PP(DepthOfFieldNearTransitionRegion);
 	LERP_PP(DepthOfFieldFarTransitionRegion);
diff --git a/Engine/Source/Runtime/Renderer/Private/PostProcess/PostProcessDOF.cpp b/Engine/Source/Runtime/Renderer/Private/PostProcess/PostProcessDOF.cpp
--- a/Engine/Source/Runtime/Renderer/Private/PostProcess/PostProcessDOF.cpp
+++ b/Engine/Source/Runtime/Renderer/Private/PostProcess/PostProcessDOF.cpp
@@ -351,6 +351,53 @@ FPooledRenderTargetDesc FRCPassPostProcessDOFRecombine::ComputeOutputDesc(EPassO
 
 
 
+
+
+// Convert f-stop and focal distance into projected size in half resolution pixels.
+static float CircleDofCoc(const FRenderingCompositePassContext& Context)
+{
+    // Convert FOV to focal length,
+	// 
+	// fov = 2 * atan(d/(2*f))
+	// where,
+	//   d = sensor dimension (APS-C 24.576 mm)
+	//   f = focal length
+	// 
+	// f = 0.5 * d * (1/tan(fov/2))
+	float HalfFOV = FMath::Atan(1.0f / Context.View.ViewMatrices.ProjMatrix.M[0][0]);
+	float FocalLength = 0.5f * 24.576f * (1.0f/FMath::Tan(HalfFOV));
+	 
+	// Convert focal distance in world position to mm.
+	// Conversion is 1 world position = 1 cm.
+	float Distance = Context.View.FinalPostProcessSettings.DepthOfFieldFocalDistance;
+	Distance *= 10.0f;
+
+	// Convert f-stop, focal length, and focal distance to
+	// projected circle of confusion size at infinity in mm.
+    //
+	// coc = f*f / (n * (d - f))
+	// where,
+	//   f = focal length
+	//   d = focal distance
+    //   n = fstop (where n is the "n" in "f/n")
+	float Radius = FocalLength * FocalLength / (Context.View.FinalPostProcessSettings.DepthOfFieldFstop * (Distance - FocalLength));
+
+	// Scale so that APS-C 24.576 mm = full frame.
+	// Convert mm to pixels.
+	float Width = (float)Context.GetViewport().Size().X;
+	Radius = Radius * Width * (1.0f/24.576f);
+
+	// Convert to half resolution (algorithm radius is at half resolution).
+	Radius *= 0.5;
+
+	// Limit to algorithm max size.
+	if(Radius > 6.0f) 
+	{
+		Radius = 6.0f; 
+	}
+	return Radius;
+}
+
 /** Encapsulates the Circle DOF setup pixel shader. */
 template <uint32 NearBlurEnable>
 class FPostProcessCircleDOFSetupPS : public FGlobalShader
@@ -375,6 +422,7 @@ public:
 	FPostProcessPassParameters PostprocessParameter;
 	FDeferredPixelShaderParameters DeferredParameters;
 	FShaderParameter DepthOfFieldParams;
+	FShaderParameter CircleDofParams;
 
 	/** Initialization constructor. */
 	FPostProcessCircleDOFSetupPS(const ShaderMetaType::CompiledShaderInitializerType& Initializer)
@@ -383,13 +431,14 @@ public:
 		PostprocessParameter.Bind(Initializer.ParameterMap);
 		DeferredParameters.Bind(Initializer.ParameterMap);
 		DepthOfFieldParams.Bind(Initializer.ParameterMap,TEXT("DepthOfFieldParams"));
+		CircleDofParams.Bind(Initializer.ParameterMap,TEXT("CircleDofParams"));
 	}
 
 	// FShader interface.
 	virtual bool Serialize(FArchive& Ar)
 	{
 		bool bShaderHasOutdatedParameters = FGlobalShader::Serialize(Ar);
-		Ar << PostprocessParameter << DeferredParameters << DepthOfFieldParams;
+		Ar << PostprocessParameter << DeferredParameters << DepthOfFieldParams << CircleDofParams;
 		return bShaderHasOutdatedParameters;
 	}
 
@@ -410,6 +459,8 @@ public:
 
 			SetShaderValueArray(Context.RHICmdList, ShaderRHI, DepthOfFieldParams, DepthOfFieldParamValues, 2);
 		}
+
+		SetShaderValue(Context.RHICmdList, ShaderRHI, CircleDofParams, CircleDofCoc(Context));
 	}
 };
 
@@ -418,7 +469,7 @@ IMPLEMENT_SHADER_TYPE(template<>,FPostProcessCircleDOFSetupPS<1>,TEXT("PostProce
 
 void FRCPassPostProcessCircleDOFSetup::Process(FRenderingCompositePassContext& Context)
 {
-	SCOPED_DRAW_EVENT(Context.RHICmdList, DOFSetup);
+	SCOPED_DRAW_EVENT(Context.RHICmdList, CircleDOFSetup);
 
 	const FPooledRenderTargetDesc* InputDesc = GetInputDesc(ePId_Input0);
 
@@ -559,6 +610,7 @@ public:
 	FPostProcessPassParameters PostprocessParameter;
 	FDeferredPixelShaderParameters DeferredParameters;
 	FShaderParameter DepthOfFieldParams;
+	FShaderParameter CircleDofParams;
 
 	/** Initialization constructor. */
 	FPostProcessCircleDOFDilatePS(const ShaderMetaType::CompiledShaderInitializerType& Initializer)
@@ -567,13 +619,14 @@ public:
 		PostprocessParameter.Bind(Initializer.ParameterMap);
 		DeferredParameters.Bind(Initializer.ParameterMap);
 		DepthOfFieldParams.Bind(Initializer.ParameterMap,TEXT("DepthOfFieldParams"));
+		CircleDofParams.Bind(Initializer.ParameterMap,TEXT("CircleDofParams"));
 	}
 
 	// FShader interface.
 	virtual bool Serialize(FArchive& Ar)
 	{
 		bool bShaderHasOutdatedParameters = FGlobalShader::Serialize(Ar);
-		Ar << PostprocessParameter << DeferredParameters << DepthOfFieldParams;
+		Ar << PostprocessParameter << DeferredParameters << DepthOfFieldParams << CircleDofParams;
 		return bShaderHasOutdatedParameters;
 	}
 
@@ -594,6 +647,8 @@ public:
 
 			SetShaderValueArray(Context.RHICmdList, ShaderRHI, DepthOfFieldParams, DepthOfFieldParamValues, 2);
 		}
+
+		SetShaderValue(Context.RHICmdList, ShaderRHI, CircleDofParams, CircleDofCoc(Context));
 	}
 };
 
@@ -602,7 +657,7 @@ IMPLEMENT_SHADER_TYPE(template<>,FPostProcessCircleDOFDilatePS<1>,TEXT("PostProc
 
 void FRCPassPostProcessCircleDOFDilate::Process(FRenderingCompositePassContext& Context)
 {
-	SCOPED_DRAW_EVENT(Context.RHICmdList, DOFSetup);
+	SCOPED_DRAW_EVENT(Context.RHICmdList, CircleDOFNear);
 
 	const FPooledRenderTargetDesc* InputDesc = GetInputDesc(ePId_Input0);
 
@@ -746,7 +801,6 @@ static void TemporalRandom2(FVector2D* RESTRICT const Constant, uint32 FrameNumb
 	Constant->Y = TemporalHalton2(FrameNumber & 1023, 3);
 }
 
-
 template <uint32 NearBlurEnable>
 class FPostProcessCircleDOFPS : public FGlobalShader
 {
@@ -771,6 +825,7 @@ public:
 	FDeferredPixelShaderParameters DeferredParameters;
 	FShaderParameter DepthOfFieldParams;
 	FShaderParameter RandomOffset;
+	FShaderParameter CircleDofParams;
 
 	/** Initialization constructor. */
 	FPostProcessCircleDOFPS(const ShaderMetaType::CompiledShaderInitializerType& Initializer)
@@ -780,13 +835,14 @@ public:
 		DeferredParameters.Bind(Initializer.ParameterMap);
 		DepthOfFieldParams.Bind(Initializer.ParameterMap,TEXT("DepthOfFieldParams"));
 		RandomOffset.Bind(Initializer.ParameterMap, TEXT("RandomOffset"));
+		CircleDofParams.Bind(Initializer.ParameterMap,TEXT("CircleDofParams"));
 	}
 
 	// FShader interface.
 	virtual bool Serialize(FArchive& Ar)
 	{
 		bool bShaderHasOutdatedParameters = FGlobalShader::Serialize(Ar);
-		Ar << PostprocessParameter << DeferredParameters << DepthOfFieldParams << RandomOffset;
+		Ar << PostprocessParameter << DeferredParameters << DepthOfFieldParams << RandomOffset << CircleDofParams;
 		return bShaderHasOutdatedParameters;
 	}
 
@@ -811,6 +867,8 @@ public:
 		FVector2D RandomOffsetValue;
 		TemporalRandom2(&RandomOffsetValue, Context.View.Family->FrameNumber);
 		SetShaderValue(Context.RHICmdList, ShaderRHI, RandomOffset, RandomOffsetValue);
+
+		SetShaderValue(Context.RHICmdList, ShaderRHI, CircleDofParams, CircleDofCoc(Context));
 	}
 };
 
@@ -819,7 +877,7 @@ IMPLEMENT_SHADER_TYPE(template<>,FPostProcessCircleDOFPS<1>,TEXT("PostProcessDOF
 
 void FRCPassPostProcessCircleDOF::Process(FRenderingCompositePassContext& Context)
 {
-	SCOPED_DRAW_EVENT(Context.RHICmdList, DOFSetup);
+	SCOPED_DRAW_EVENT(Context.RHICmdList, CircleDOFApply);
 
 	const FPooledRenderTargetDesc* InputDesc = GetInputDesc(ePId_Input0);
 
@@ -957,6 +1015,7 @@ public:
 	FDeferredPixelShaderParameters DeferredParameters;
 	FShaderParameter DepthOfFieldUVLimit;
 	FShaderParameter RandomOffset;
+	FShaderParameter CircleDofParams;
 
 	/** Initialization constructor. */
 	FPostProcessCircleDOFRecombinePS(const ShaderMetaType::CompiledShaderInitializerType& Initializer)
@@ -966,13 +1025,14 @@ public:
 		DeferredParameters.Bind(Initializer.ParameterMap);
 		DepthOfFieldUVLimit.Bind(Initializer.ParameterMap,TEXT("DepthOfFieldUVLimit"));
 		RandomOffset.Bind(Initializer.ParameterMap, TEXT("RandomOffset"));
+		CircleDofParams.Bind(Initializer.ParameterMap,TEXT("CircleDofParams"));
 	}
 
 	// FShader interface.
 	virtual bool Serialize(FArchive& Ar)
 	{
 		bool bShaderHasOutdatedParameters = FGlobalShader::Serialize(Ar);
-		Ar << PostprocessParameter << DeferredParameters << DepthOfFieldUVLimit << RandomOffset;
+		Ar << PostprocessParameter << DeferredParameters << DepthOfFieldUVLimit << RandomOffset << CircleDofParams;
 		return bShaderHasOutdatedParameters;
 	}
 
@@ -997,6 +1057,8 @@ public:
 		FVector2D RandomOffsetValue;
 		TemporalRandom2(&RandomOffsetValue, Context.View.Family->FrameNumber);
 		SetShaderValue(Context.RHICmdList, ShaderRHI, RandomOffset, RandomOffsetValue);
+
+		SetShaderValue(Context.RHICmdList, ShaderRHI, CircleDofParams, CircleDofCoc(Context));
 	}
 };
 
