Disable development/editor only features in shaders for consoles.

diff --git a/Engine/Shaders/BasePassForForwardShadingPixelShader.usf b/Engine/Shaders/BasePassForForwardShadingPixelShader.usf
--- a/Engine/Shaders/BasePassForForwardShadingPixelShader.usf
+++ b/Engine/Shaders/BasePassForForwardShadingPixelShader.usf
@@ -56,7 +56,7 @@ half3 GetImageBasedReflectionLighting(FMaterialPixelParameters MaterialParameter
 }
 #endif
 
-#if EDITOR_PRIMITIVE_MATERIAL
+#if USE_EDITOR_COMPOSITING
 bool bEnableEditorPrimitiveDepthTest;
 // depth in the red channel in DeviceZ
 Texture2D		FilteredSceneDepthTexture;
@@ -64,7 +64,7 @@ SamplerState	FilteredSceneDepthTextureSampler;
 #endif
 
 
-#if EDITOR_PRIMITIVE_MATERIAL && (ES2_EMULATION)
+#if USE_EDITOR_COMPOSITING && (ES2_EMULATION)
 // @return 0:translucent..1:opaque
 float ClipForEditorPrimitives(FMaterialPixelParameters MaterialParameters)
 {
@@ -320,7 +320,7 @@ void Main(
 	#endif
 
 	// Editor primitive depth testing
-	#if EDITOR_PRIMITIVE_MATERIAL && (ES2_EMULATION)
+	#if USE_EDITOR_COMPOSITING && (ES2_EMULATION)
 		#if MATERIALBLENDING_MASKED
 			// some material might have a opacity value
 			OutColor.a = GetMaterialMaskInputRaw(MaterialParameters);
diff --git a/Engine/Shaders/BasePassPixelShader.usf b/Engine/Shaders/BasePassPixelShader.usf
--- a/Engine/Shaders/BasePassPixelShader.usf
+++ b/Engine/Shaders/BasePassPixelShader.usf
@@ -55,7 +55,7 @@ float4 PointSkyBentNormal;
 	#define COMPILER_GLSL 0
 #endif
 
-#define EDITOR_ALPHA2COVERAGE (EDITOR_PRIMITIVE_MATERIAL && FEATURE_LEVEL >= FEATURE_LEVEL_SM5 && !COMPILER_GLSL)
+#define EDITOR_ALPHA2COVERAGE (USE_EDITOR_COMPOSITING && FEATURE_LEVEL >= FEATURE_LEVEL_SM5 && !COMPILER_GLSL)
 
 
 #if TRANSLUCENCY_LIGHTING_SURFACE || TRANSLUCENCY_LIGHTING_SURFACE_PERPIXEL
@@ -510,7 +510,7 @@ void GetPrecomputedIndirectLightingAndSkyLight(
 	#endif
 }
 
-#if EDITOR_PRIMITIVE_MATERIAL
+#if USE_EDITOR_COMPOSITING
 	bool bEnableEditorPrimitiveDepthTest;
 	int MSAASampleCount;
 
@@ -524,7 +524,7 @@ float ClipForEditorPrimitives(FMaterialPixelParameters MaterialParameters)
 {
 	float Ret = 1;
 
-#if EDITOR_PRIMITIVE_MATERIAL && (FEATURE_LEVEL >= FEATURE_LEVEL_SM4 || ES2_EMULATION)
+#if USE_EDITOR_COMPOSITING && (FEATURE_LEVEL >= FEATURE_LEVEL_SM4 || ES2_EMULATION)
 	// Depth test manually if compositing editor primitives since the depth buffer is different (MSAA only)
 	BRANCH
 	if (bEnableEditorPrimitiveDepthTest)
@@ -630,7 +630,7 @@ void Main(
 #endif
 	);
 
-#if EDITOR_PRIMITIVE_MATERIAL && (FEATURE_LEVEL >= FEATURE_LEVEL_SM4 || ES2_EMULATION)
+#if USE_EDITOR_COMPOSITING && (FEATURE_LEVEL >= FEATURE_LEVEL_SM4 || ES2_EMULATION)
 	const bool bEditorWeightedZBuffering = true;
 #else
 	const bool bEditorWeightedZBuffering = false;
@@ -697,8 +697,7 @@ void Main(
 	half3 DiffuseColor = BaseColor - BaseColor * Metallic;
 	half3 SpecularColor = lerp( 0.08 * Specular.xxx, BaseColor, Metallic.xxx );
 
-// todo: COMPILE_SHADERS_FOR_DEVELOPMENT is unfinished feature, using XBOXONE_PROFILE as workaround
-#if COMPILE_SHADERS_FOR_DEVELOPMENT == 1 && !XBOXONE_PROFILE && !ES31_AEP_PROFILE
+#if USE_DEVELOPMENT_SHADERS
 	{
 		// this feature is only needed for development/editor - we can compile it out for a shipping build (see r.CompileShadersForDevelopment cvar help)
 		DiffuseColor = DiffuseColor * View.DiffuseOverrideParameter.w + View.DiffuseOverrideParameter.xyz;
@@ -763,8 +762,7 @@ void Main(
 
 	half3 Emissive = GetMaterialEmissive(MaterialParameters);
 
-// todo: COMPILE_SHADERS_FOR_DEVELOPMENT is unfinished feature, using XBOXONE_PROFILE as workaround
-#if COMPILE_SHADERS_FOR_DEVELOPMENT == 1 && !XBOXONE_PROFILE && !ES31_AEP_PROFILE
+#if USE_DEVELOPMENT_SHADERS
 	// this feature is only needed for development/editor - we can compile it out for a shipping build (see r.CompileShadersForDevelopment cvar help)
 	#if SM5_PROFILE || SM4_PROFILE
 		BRANCH
diff --git a/Engine/Shaders/Common.usf b/Engine/Shaders/Common.usf
--- a/Engine/Shaders/Common.usf
+++ b/Engine/Shaders/Common.usf
@@ -152,8 +152,12 @@
 
 #endif
 
+// non-editor platforms generally never want development/editor features.
+#define PLATFORM_SUPPORTS_DEVELOPMENT_SHADERS (!PS4_PROFILE && !XBOXONE_PROFILE && !ES31_AEP_PROFILE)
+#define USE_DEVELOPMENT_SHADERS (COMPILE_SHADERS_FOR_DEVELOPMENT && PLATFORM_SUPPORTS_DEVELOPMENT_SHADERS)
+
 #ifndef ES2_EMULATION
-#define ES2_EMULATION ((FEATURE_LEVEL == FEATURE_LEVEL_ES2) && (!COMPILER_GLSL_ES2 && COMPILE_SHADERS_FOR_DEVELOPMENT && !METAL_PROFILE))
+#define ES2_EMULATION ((FEATURE_LEVEL == FEATURE_LEVEL_ES2) && (!COMPILER_GLSL_ES2 && USE_DEVELOPMENT_SHADERS && !METAL_PROFILE))
 #endif
 
 // If we didn't request color component swizzling, just make it empty
diff --git a/Engine/Shaders/DeferredShadingCommon.usf b/Engine/Shaders/DeferredShadingCommon.usf
--- a/Engine/Shaders/DeferredShadingCommon.usf
+++ b/Engine/Shaders/DeferredShadingCommon.usf
@@ -402,14 +402,13 @@ FGBufferData DecodeGBufferData(
 		GBuffer.DiffuseColor = GBuffer.BaseColor - GBuffer.BaseColor * GBuffer.Metallic;
 		GBuffer.SpecularColor = lerp( 0.08 * GBuffer.Specular.xxx, GBuffer.BaseColor, GBuffer.Metallic );
 
-		// todo: COMPILE_SHADERS_FOR_DEVELOPMENT is unfinished feature, using XBOXONE_PROFILE as workaround
-		#if COMPILE_SHADERS_FOR_DEVELOPMENT == 1 && !XBOXONE_PROFILE && !ES31_AEP_PROFILE
+		#if USE_DEVELOPMENT_SHADERS
 		{
 			// this feature is only needed for development/editor - we can compile it out for a shipping build (see r.CompileShadersForDevelopment cvar help)
 			GBuffer.DiffuseColor = GBuffer.DiffuseColor * View.DiffuseOverrideParameter.www + View.DiffuseOverrideParameter.xyz;
 			GBuffer.SpecularColor = GBuffer.SpecularColor * View.SpecularOverrideParameter.w + View.SpecularOverrideParameter.xyz;
 		}
-		#endif //COMPILE_SHADERS_FOR_DEVELOPMENT == 1
+		#endif //USE_DEVELOPMENT_SHADERS
 	}
 
 	GBuffer.Velocity = InGBufferVelocity;
diff --git a/Engine/Shaders/MaterialTemplate.usf b/Engine/Shaders/MaterialTemplate.usf
--- a/Engine/Shaders/MaterialTemplate.usf
+++ b/Engine/Shaders/MaterialTemplate.usf
@@ -25,6 +25,15 @@
 	#define USE_DITHERED_LOD_TRANSITION 0
 #endif
 
+//Platforms that don't run the editor shouldn't need editor features in the shaders.
+#define PLATFORM_SUPPORTS_EDITOR_SHADERS (!PS4_PROFILE && !XBOXONE_PROFILE && !ES31_AEP_PROFILE)
+
+//Tie Editor features to platform support and the COMPILE_SHADERS_FOR_DEVELOPMENT which is set via CVAR.
+#define USE_EDITOR_SHADERS (PLATFORM_SUPPORTS_EDITOR_SHADERS && USE_DEVELOPMENT_SHADERS)
+
+//Materials also have to opt in to these features.
+#define USE_EDITOR_COMPOSITING (USE_EDITOR_SHADERS && EDITOR_PRIMITIVE_MATERIAL)
+
 /**
  * Parameters used by vertex and pixel shaders to access particle properties.
  */
@@ -546,7 +555,7 @@ MaterialFloat2 GetParticleMacroUV(FMaterialPixelParameters Parameters)
 }
 
 #ifndef ES2_EMULATION
-#define ES2_EMULATION ((FEATURE_LEVEL == FEATURE_LEVEL_ES2) && (!COMPILER_GLSL_ES2 && COMPILE_SHADERS_FOR_DEVELOPMENT && !METAL_PROFILE))
+#define ES2_EMULATION ((FEATURE_LEVEL == FEATURE_LEVEL_ES2) && (!COMPILER_GLSL_ES2 && USE_DEVELOPMENT_SHADERS && !METAL_PROFILE))
 #endif
 
 MaterialFloat4 ProcessMaterialColorTextureLookup(MaterialFloat4 TextureValue)
@@ -1016,9 +1025,8 @@ half3 GetMaterialNormal(FMaterialPixelParameters Parameters)
 	half3 RetNormal;
 
 	RetNormal = GetMaterialNormalRaw(Parameters);
-	
-	// todo: COMPILE_SHADERS_FOR_DEVELOPMENT is unfinished feature, using XBOXONE_PROFILE as workaround
-	#if COMPILE_SHADERS_FOR_DEVELOPMENT == 1 && !XBOXONE_PROFILE && !(ES2_PROFILE || ES3_1_PROFILE || ES31_AEP_PROFILE)
+		
+	#if USE_EDITOR_SHADERS && !(ES2_PROFILE || ES3_1_PROFILE || ES31_AEP_PROFILE)
 	{
 		// this feature is only needed for development/editor - we can compile it out for a shipping build (see r.CompileShadersForDevelopment cvar help)
 		half3 OverrideNormal = View.NormalOverrideParameter.xyz;
@@ -1099,9 +1107,8 @@ half GetMaterialRoughness(FMaterialPixelParameters Parameters)
 	return clamp( GetMaterialRoughnessRaw(Parameters), 0.12, 1.0 );
 #else
 	half Roughness = saturate(GetMaterialRoughnessRaw(Parameters));
-
-	// todo: COMPILE_SHADERS_FOR_DEVELOPMENT is unfinished feature, using XBOXONE_PROFILE as workaround
-	#if COMPILE_SHADERS_FOR_DEVELOPMENT == 1 && !XBOXONE_PROFILE && !ES31_AEP_PROFILE
+	
+	#if USE_EDITOR_SHADERS && !ES31_AEP_PROFILE
 	{
 		// this feature is only needed for development/editor - we can compile it out for a shipping build (see r.CompileShadersForDevelopment cvar help)
 		Roughness = Roughness * View.RoughnessOverrideParameter.y + View.RoughnessOverrideParameter.x;
