As Brian requested: YCoCg on for the high quality temporal AA.

diff --git a/Engine/Shaders/PostProcessTemporalAA.usf b/Engine/Shaders/PostProcessTemporalAA.usf
--- a/Engine/Shaders/PostProcessTemporalAA.usf
+++ b/Engine/Shaders/PostProcessTemporalAA.usf
@@ -246,11 +246,12 @@ void MainTemporalAAPS( float4 UVAndScreenPos : TEXCOORD0, float3 InExposureScale
 	#define AA_FILTERED 1
 	#define AA_BORDER 1
 	#define AA_CROSS 2
-	#define AA_GREEN_AS_LUMA 0
+	#define AA_GREEN_AS_LUMA 1
 	#define AA_AABB 1
-	#define AA_LOWPASS 1
+	#define AA_LOWPASS 0
 	#define AA_DEBUG 0
 	#define AA_VELOCITY_WEIGHTING 1
+	#define AA_YCOCG 1
 	#include "PostProcessTemporalCommon.usf"
 #else 
 	OutColor = PostprocessInput0.SampleLevel(PostprocessInput0Sampler, UVAndScreenPos.xy, 0);
