Integrate - SSS support for ray traced distance field shadows

diff --git a/Engine/Shaders/DistanceFieldShadowing.usf b/Engine/Shaders/DistanceFieldShadowing.usf
--- a/Engine/Shaders/DistanceFieldShadowing.usf
+++ b/Engine/Shaders/DistanceFieldShadowing.usf
@@ -517,8 +517,6 @@ void DistanceFieldShadowingCS(
 	BRANCH
 	if (bTileShouldComputeShadowing)
 	{
-		FGBufferData GBufferData = GetGBufferData(ScreenUV);
-
 		{
 			// World space offset along the start of the ray to avoid incorrect self-shadowing
 			float RayStartOffset = 2 + RayStartOffsetDepthScale * SceneDepth;
@@ -559,6 +557,20 @@ void DistanceFieldShadowingCS(
 				bool bUseCulling = false;
 			#endif
 
+			FGBufferData GBufferData = GetGBufferData(ScreenUV);
+
+			float SubsurfaceDensity = 0;
+			bool bUseSubsurfaceTransmission = false;
+
+			BRANCH
+			if (IsSubsurfaceModel(GBufferData.ShadingModelID))
+			{
+				// Note: this has to match shadowmap logic
+				// Derive density from a heuristic using opacity, tweaked for useful falloff ranges and to give a linear depth falloff with opacity
+				SubsurfaceDensity = -.05f * log(1 - min(GBufferData.Opacity, .999f));
+				bUseSubsurfaceTransmission = true;
+			}
+
 			Result = ShadowRayTraceThroughCulledObjects(
 				WorldRayStart, 
 				WorldRayEnd, 
@@ -566,10 +578,12 @@ void DistanceFieldShadowingCS(
 				TanLightAngle, 
 				MinSphereRadius, 
 				MaxSphereRadius, 
+				SubsurfaceDensity,
 				CulledDataParameter, 
 				NumIntersectingObjects,
 				bUseCulling,
-				bUseScatterTileCulling);
+				bUseScatterTileCulling,
+				bUseSubsurfaceTransmission);
 		}
 	}
 
diff --git a/Engine/Shaders/DistanceFieldShadowingShared.usf b/Engine/Shaders/DistanceFieldShadowingShared.usf
--- a/Engine/Shaders/DistanceFieldShadowingShared.usf
+++ b/Engine/Shaders/DistanceFieldShadowingShared.usf
@@ -45,10 +45,12 @@ float ShadowRayTraceThroughCulledObjects(
 	float TanLightAngle, 
 	float MinSphereRadius, 
 	float MaxSphereRadius, 
+	float SubsurfaceDensity,
 	uint CulledDataParameter, 
 	uint NumIntersectingObjects,
 	uniform bool bUseCulling,
-	uniform bool bUseScatterTileCulling)
+	uniform bool bUseScatterTileCulling,
+	bool bUseSubsurfaceTransmission)
 {
 	float MinRayTime = MaxRayTime;
 	float MinConeVisibility = 1;
@@ -113,8 +115,6 @@ float ShadowRayTraceThroughCulledObjects(
 				float MinStepSize = 1.0f / (4 * MaxSteps);
 
 				float MinDistance = 1000000;
-				float3 IntersectionPosition = float3(0, 0, 0);
-
 				uint StepIndex = 0;
 
 				LOOP
@@ -134,22 +134,35 @@ float ShadowRayTraceThroughCulledObjects(
 
 					MinDistance = min(MinDistance, DistanceField);
 					float SphereRadius = clamp(TanLightAngle * SampleRayTime, VolumeMinSphereRadius, VolumeMaxSphereRadius);
+					float StepVisibility = saturate(DistanceField / SphereRadius);
+
+					if (bUseSubsurfaceTransmission)
+					{
+						// Determine the distance that the light traveled through the subsurface object
+						// This assumes that anything between this subsurface pixel and the light was also a subsurface material
+						float Thickness = SampleRayTime * UVScaleAndVolumeScale.w;
+						float SubsurfaceVisibility = saturate(exp(-Thickness * SubsurfaceDensity));
+
+						// Prevent full occlusion in the range that SSS is effective
+						// Note: this may cause the trace to travel through negative regions of the distance field
+						// It also prevents visibility from ever going to 0
+						StepVisibility = max(StepVisibility, SubsurfaceVisibility);
+					}
 
-					MinConeVisibility = min(MinConeVisibility, saturate(DistanceField / SphereRadius));
-					IntersectionPosition = SampleVolumePosition;
+					MinConeVisibility = min(MinConeVisibility, StepVisibility);
 
-					float StepDistance = max(DistanceField, MinStepSize);
+					float StepDistance = max(abs(DistanceField), MinStepSize);
 					SampleRayTime += StepDistance;
 
-					// Terminate the trace if we reached a negative area or went past the end of the ray
-					if (DistanceField < 0 
+					// Terminate the trace if we are fully occluded or went past the end of the ray
+					if (MinConeVisibility < .01f
 						|| SampleRayTime > IntersectionTimes.y * VolumeRayLength)
 					{
 						break;
 					}
 				}
 
-				if (MinDistance < 0 || StepIndex == MaxSteps)
+				if (MinConeVisibility < .01f || StepIndex == MaxSteps)
 				{
 					MinConeVisibility = 0;
 					MinRayTime = min(MinRayTime, SampleRayTime * UVScaleAndVolumeScale.w);
diff --git a/Engine/Shaders/HeightfieldLighting.usf b/Engine/Shaders/HeightfieldLighting.usf
--- a/Engine/Shaders/HeightfieldLighting.usf
+++ b/Engine/Shaders/HeightfieldLighting.usf
@@ -324,10 +324,12 @@ void RayTracedShadowHeightfieldsPS(
 			TanLightAngle, 
 			MinSphereRadius, 
 			MaxSphereRadius, 
+			0,
 			CulledDataStart, 
 			NumIntersectingObjects,
 			true,
-			true);
+			true,
+			false);
 	}
 	else
 	{
