Integrate - Heightfield occlusion is now applied to objects without a distance field representation
 * Added HasHeightfieldRepresentation bit to GBuffer to accomplish this

diff --git a/Engine/Shaders/BasePassPixelShader.usf b/Engine/Shaders/BasePassPixelShader.usf
--- a/Engine/Shaders/BasePassPixelShader.usf
+++ b/Engine/Shaders/BasePassPixelShader.usf
@@ -818,6 +818,7 @@ void Main(
 		GBuffer.Opacity = Opacity;
 		GBuffer.DecalMask = Primitive.DecalReceiverMask;
 		GBuffer.HasDistanceFieldRepresentation = Primitive.HasDistanceFieldRepresentation;
+		GBuffer.HasHeightfieldRepresentation = Primitive.HasHeightfieldRepresentation;
 
 		#if MATERIAL_SHADINGMODEL_UNLIT
 			GBuffer.ShadingModelID = SHADINGMODELID_UNLIT;
diff --git a/Engine/Shaders/DeferredDecal.usf b/Engine/Shaders/DeferredDecal.usf
--- a/Engine/Shaders/DeferredDecal.usf
+++ b/Engine/Shaders/DeferredDecal.usf
@@ -172,6 +172,7 @@ void MainPS(
 	Data.ShadingModelID = SHADINGMODELID_DEFAULT_LIT;
 	Data.DecalMask = 1;
 	Data.HasDistanceFieldRepresentation = 1;
+	Data.HasHeightfieldRepresentation = 1;
 
 	#if DECAL_RENDERSTAGE == 0
 	{
diff --git a/Engine/Shaders/DeferredShadingCommon.usf b/Engine/Shaders/DeferredShadingCommon.usf
--- a/Engine/Shaders/DeferredShadingCommon.usf
+++ b/Engine/Shaders/DeferredShadingCommon.usf
@@ -199,7 +199,7 @@ float AddAngleToRoughness( float Angle, float Roughness )
 // @param Scalar clamped in 0..1 range
 // @param Mask 0..1
 // @return 8bit in range 0..1
-float Encode71(float Scalar, int Mask)
+float Encode71(float Scalar, uint Mask)
 {
 	return
 		127.0f / 255.0f * saturate(Scalar) +
@@ -210,18 +210,13 @@ float Encode71(float Scalar, int Mask)
 // @param Scalar 0..1
 // @param Mask 0..1
 // @return 7bit in 0.1
-float Decode71(float Scalar, out int Mask)
+float Decode71(float Scalar, out uint Mask)
 {
-	Mask = (int)(Scalar > 0.5f);
+	Mask = (uint)(Scalar > 0.5f);
 
 	return (Scalar - 0.5f * Mask) * 2.0f;
 }
 
-uint DecodeShadingModelId(float InPackedChannel)
-{
-	return (uint)( InPackedChannel * 255.999 );
-}
-
 #define SHADINGMODELID_UNLIT				0
 #define SHADINGMODELID_DEFAULT_LIT			1
 #define SHADINGMODELID_SUBSURFACE			2
@@ -231,6 +226,16 @@ uint DecodeShadingModelId(float InPackedChannel)
 #define SHADINGMODELID_TWOSIDED_FOLIAGE		6
 #define SHADINGMODELID_NUM					7
 
+float EncodeShadingModelId(uint ShadingModelId)
+{
+	return (ShadingModelId + .5f) / (float)SHADINGMODELID_NUM;
+}
+
+uint DecodeShadingModelId(float InPackedChannel)
+{
+	return (uint)( InPackedChannel * SHADINGMODELID_NUM );
+}
+
 bool IsSubsurfaceModel(int ShadingModel)
 {
 	return ShadingModel == SHADINGMODELID_SUBSURFACE 
@@ -268,11 +273,13 @@ struct FGBufferData
 	// 0..1 ambient occlusion  e.g.SSAO, wet surface mask, skylight mask, ...
 	float GBufferAO;
 	// 0..255 
-	int ShadingModelID;
+	uint ShadingModelID;
 	// 0..1 decal receiver mask
 	uint DecalMask;
 	// 0..1
 	uint HasDistanceFieldRepresentation; 
+	// 0..1
+	uint HasHeightfieldRepresentation; 
 	// in world units
 	float CustomDepth;
 	// in unreal units (linear), can be used to reconstruct world position,
@@ -343,7 +350,7 @@ void EncodeGBuffer(
 
 		// No space for AO. Multiply IndirectIrradiance by AO instead of storing.
 		OutGBufferC.a = EncodeIndirectIrradiance(GBuffer.IndirectIrradiance * GBuffer.GBufferAO) + QuantizationBias * (1.0 / 255.0);
-		OutGBufferB.a = GBuffer.ShadingModelID * (1.0 / 255.0);
+		OutGBufferB.a = Encode71(EncodeShadingModelId(GBuffer.ShadingModelID), GBuffer.HasHeightfieldRepresentation);
 
 		OutGBufferD.r = GBuffer.Opacity;
 		OutGBufferD.gba = GBuffer.CustomData;
@@ -382,7 +389,8 @@ FGBufferData DecodeGBufferData(
 	uint PackedAlpha = (uint)( InGBufferA.a * 3.999 );
 	GBuffer.DecalMask = (PackedAlpha & 2) != 0;
 	GBuffer.HasDistanceFieldRepresentation = PackedAlpha & 1;
-	GBuffer.ShadingModelID = DecodeShadingModelId(InGBufferB.a);
+	GBuffer.HasHeightfieldRepresentation = 0;
+	GBuffer.ShadingModelID = DecodeShadingModelId(Decode71(InGBufferB.a, GBuffer.HasHeightfieldRepresentation));
 	
 	GBuffer.BaseColor = DecodeBaseColor(InGBufferC.rgb);
 	GBuffer.Metallic = InGBufferB.r;
@@ -627,8 +635,8 @@ FGBufferData GetGBufferData(float2 UV, bool bGetNormalizedNormal = true)
 // Minimal path for just the lighting model, used to branch around unlit pixels (skybox)
 uint GetShadingModelId(float2 UV)
 {
-	float4 GBufferB = Texture2DSampleLevel(GBuffers.GBufferBTexture, GBuffers.GBufferBTextureSampler, UV, 0);
-	return DecodeShadingModelId(GBufferB.a);
+	bool Unused;
+	return DecodeShadingModelId(Decode71(Texture2DSampleLevel(GBuffers.GBufferBTexture, GBuffers.GBufferBTextureSampler, UV, 0).a, Unused));
 }
 
 // @param UV - UV space in the DBuffer textures (BufferSize resolution)
diff --git a/Engine/Shaders/DistanceFieldAOShared.usf b/Engine/Shaders/DistanceFieldAOShared.usf
--- a/Engine/Shaders/DistanceFieldAOShared.usf
+++ b/Engine/Shaders/DistanceFieldAOShared.usf
@@ -78,6 +78,25 @@ uint4 GetTileHead(uint2 TileCoordinate)
 Texture2D DistanceFieldNormalTexture;
 SamplerState DistanceFieldNormalSampler;
 
+float4 EncodeDownsampledGBuffer(FGBufferData GBufferData, float SceneDepth)
+{
+	float DepthSign = GBufferData.HasDistanceFieldRepresentation > 0 ? 1 : -1;
+	float ZSign = GBufferData.HasHeightfieldRepresentation > 0 ? 1 : -1;
+	float ZChannel = GBufferData.WorldNormal.z > 0 ? .5f : 1;
+	float4 OutValue = float4(GBufferData.WorldNormal.xy, ZChannel * ZSign, DepthSign * SceneDepth);
+	return OutValue;
+}
+
+void GetDownsampledGBuffer(float2 ScreenUV, out float3 OutNormal, out float OutDepth, out bool bHasDistanceFieldRepresentation, out bool bHasHeightfieldRepresentation)
+{
+	float4 TextureValue = Texture2DSampleLevel(DistanceFieldNormalTexture, DistanceFieldNormalSampler, ScreenUV, 0);
+	bHasDistanceFieldRepresentation = TextureValue.w > 0;
+	bHasHeightfieldRepresentation = TextureValue.z > 0; 
+	OutDepth = abs(TextureValue.w);
+	OutNormal = float3(TextureValue.xy, sqrt(1 - dot(TextureValue.xy, TextureValue.xy)));
+	OutNormal.z *= abs(TextureValue.z) > .75f ? -1 : 1;
+}
+
 int CurrentLevelDownsampleFactor;
 float2 AOBufferSize;
 
diff --git a/Engine/Shaders/DistanceFieldSurfaceCacheLighting.usf b/Engine/Shaders/DistanceFieldSurfaceCacheLighting.usf
--- a/Engine/Shaders/DistanceFieldSurfaceCacheLighting.usf
+++ b/Engine/Shaders/DistanceFieldSurfaceCacheLighting.usf
@@ -109,10 +109,7 @@ void ComputeDistanceFieldNormalPS(
 	float SceneDepth = CalcSceneDepth(ScreenUV);
 	FGBufferData GBufferData = GetGBufferData(ScreenUV);
 
-	float3 WorldNormal = GBufferData.WorldNormal;
-	float DepthSign = GBufferData.HasDistanceFieldRepresentation > 0 ? 1 : -1;
-	// Note: Encoding whether the pixel has a distance field representation in the sign bit of the alpha
-	OutColor = float4(WorldNormal, DepthSign * SceneDepth);
+	OutColor = EncodeDownsampledGBuffer(GBufferData, SceneDepth);
 }
 
 RWTexture2D<float4> RWDistanceFieldNormal;
@@ -127,9 +124,7 @@ void ComputeDistanceFieldNormalCS(
 	float SceneDepth = CalcSceneDepth(ScreenUV);
 	FGBufferData GBufferData = GetGBufferData(ScreenUV);
 
-	float3 WorldNormal = GBufferData.WorldNormal;
-	float DepthSign = GBufferData.HasDistanceFieldRepresentation > 0 ? 1 : -1;
-	float4 OutValue = float4(WorldNormal, DepthSign * SceneDepth);
+	float4 OutValue = EncodeDownsampledGBuffer(GBufferData, SceneDepth);
 	RWDistanceFieldNormal[DispatchThreadId.xy] = OutValue;
 }
 
@@ -139,10 +134,15 @@ void WriteDownsampledDepthPS(
 	out float OutDepth : SV_Depth)
 {
 	float2 DistanceFieldUVs = UVAndScreenPos.xy;
-	float Depth = abs(Texture2DSampleLevel(DistanceFieldNormalTexture, DistanceFieldNormalSampler, DistanceFieldUVs, 0).w);
+
+	float3 WorldNormal;
+	float SceneDepth;
+	bool bHasDistanceFieldRepresentation;
+	bool bHasHeightfieldRepresentation;
+	GetDownsampledGBuffer(DistanceFieldUVs, WorldNormal, SceneDepth, bHasDistanceFieldRepresentation, bHasHeightfieldRepresentation);
 
 	OutColor = 0;
-	OutDepth = ConvertToDeviceZ(Depth);
+	OutDepth = ConvertToDeviceZ(SceneDepth);
 }
 
 // For some reason gives innaccurate results at lower resolutions
@@ -216,7 +216,7 @@ void IrradianceCacheSplatVS(
 	VirtualRadius += HackExpand * ViewToCenterLength * 4;
 #endif
 
-	float4 RecordNormalAndTime = IrradianceCacheNormal.Load(JobIndex);
+	float3 RecordNormal = IrradianceCacheNormal.Load(JobIndex).xyz;
 
 	float2 CornerUVs = InUV;
 	float3 CornerPosition = VirtualPosition + (View.ViewRight * CornerUVs.x + View.ViewUp * CornerUVs.y) * VirtualRadius;
@@ -233,7 +233,7 @@ void IrradianceCacheSplatVS(
 	Output.PositionRadius = PositionAndRadius; 
 
 	float Fade = 1;
-	Output.NormalAndFade = float4(RecordNormalAndTime.xyz, Fade);
+	Output.NormalAndFade = float4(RecordNormal, Fade);
 
 #if FINAL_INTERPOLATION_PASS
 	Output.BentNormal = IrradianceCacheBentNormal.Load(JobIndex).xyz;
@@ -264,19 +264,23 @@ void IrradianceCacheSplatPS(
 #endif
 
 	float2 BaseLevelScreenUV = float2((((uint2)SVPos.xy) * DownsampleFactorToBaseLevel + .5f) * BaseLevelTexelSize);
-	float4 DistanceFieldNormalValue = Texture2DSampleLevel(DistanceFieldNormalTexture, DistanceFieldNormalSampler, BaseLevelScreenUV, 0);
+
+	float3 WorldNormal;
+	float SceneDepth;
+	bool bHasDistanceFieldRepresentation;
+	bool bHasHeightfieldRepresentation;
+	GetDownsampledGBuffer(BaseLevelScreenUV, WorldNormal, SceneDepth, bHasDistanceFieldRepresentation, bHasHeightfieldRepresentation);
 
 #if USE_SCREENVECTOR_WORLD_POSITION
 
-	float SceneW = abs(DistanceFieldNormalValue.w);
-	float3 OpaqueWorldPosition = View.ViewOrigin.xyz + Input.ScreenVector * SceneW;
+	float3 OpaqueWorldPosition = View.ViewOrigin.xyz + Input.ScreenVector * SceneDepth;
 
 #else
 
 	float2 ScreenUV = float2((((uint2)SVPos.xy) * CurrentLevelDownsampleFactor + View.ViewRectMin.xy + float2(.5f, .5f)) * View.ViewSizeAndSceneTexelSize.zw);
-	float SceneW = CalcSceneDepth(ScreenUV); 
+	SceneDepth = CalcSceneDepth(ScreenUV); 
 	float2 ScreenPosition = (ScreenUV.xy - View.ScreenPositionScaleBias.wz) / View.ScreenPositionScaleBias.xy;
-	float4 HomogeneousWorldPosition = mul(float4(ScreenPosition.xy * SceneW, SceneW, 1), View.ScreenToWorld);
+	float4 HomogeneousWorldPosition = mul(float4(ScreenPosition.xy * SceneDepth, SceneDepth, 1), View.ScreenToWorld);
 	float3 OpaqueWorldPosition = HomogeneousWorldPosition.xyz / HomogeneousWorldPosition.w;
 
 #endif
@@ -286,9 +290,8 @@ void IrradianceCacheSplatPS(
 	float Weight = 0;
 
 	BRANCH
-	if (DistanceError < 1 && SceneW < AOMaxViewDistance)
+	if (DistanceError < 1 && SceneDepth < AOMaxViewDistance)
 	{ 
-		float3 WorldNormal = DistanceFieldNormalValue.xyz;
 		float3 RecordNormal = Input.NormalAndFade.xyz;
 		float NormalError = InterpolationAngleNormalization * sqrt(saturate(1 - dot(WorldNormal, RecordNormal)));
 
@@ -311,7 +314,7 @@ void IrradianceCacheSplatPS(
 
 #if FINAL_INTERPOLATION_PASS
 		// Pixels without a distance field representation interpolate AO instead of a bent normal
-		float3 InterpolatedValue = DistanceFieldNormalValue.w < 0 ? length(Input.BentNormal).xxx : Input.BentNormal;
+		float3 InterpolatedValue = bHasDistanceFieldRepresentation ? Input.BentNormal : length(Input.BentNormal).xxx;
 		OutBentNormal.rgb = InterpolatedValue * Weight;
 		#if SUPPORT_IRRADIANCE
 			OutIrradiance.rgb = Input.Irradiance * Weight;
@@ -323,7 +326,7 @@ void IrradianceCacheSplatPS(
 	#define VISUALIZE_SPLAT_OVERDRAW 0
 	#if VISUALIZE_SPLAT_OVERDRAW && FINAL_INTERPOLATION_PASS
 		OutBentNormal.rgb = 0;
-		//OutBentNormal.a = Input.Position.w <= SceneW ? .01f : 0.0f;
+		//OutBentNormal.a = Input.Position.w <= SceneDepth ? .01f : 0.0f;
 		OutBentNormal.a = .005f;// + (Weight > .001f ? .01f : 0);
 	#endif
 
@@ -358,8 +361,11 @@ void AOCombinePS(
 	OutIrradiance = 0;
 #endif
 
-	float4 DistanceFieldNormalValue = Texture2DSampleLevel(DistanceFieldNormalTexture, DistanceFieldNormalSampler, UVAndScreenPos.xy, 0);
-	float SceneDepth = abs(DistanceFieldNormalValue.w);
+	float3 WorldNormal;
+	float SceneDepth;
+	bool bHasDistanceFieldRepresentation;
+	bool bHasHeightfieldRepresentation;
+	GetDownsampledGBuffer(UVAndScreenPos.xy, WorldNormal, SceneDepth, bHasDistanceFieldRepresentation, bHasHeightfieldRepresentation);
 
 #define VISUALIZE_ACCUMULATED_WEIGHTS 0
 #if VISUALIZE_ACCUMULATED_WEIGHTS
@@ -376,11 +382,11 @@ void AOCombinePS(
 		OutBentNormal.rgb = BentNormalAccumulation.xyz / BentNormalAccumulation.w;
 		OutBentNormal.a = SceneDepth;
 
-		if (DistanceFieldNormalValue.w < 0)
+		if (!bHasDistanceFieldRepresentation)
 		{
 			// Pixels without a distance field representation interpolate AO instead of a bent normal
 			// Construct a bent normal from the interpolated AO
-			OutBentNormal.rgb = OutBentNormal.r * DistanceFieldNormalValue.xyz;
+			OutBentNormal.rgb = OutBentNormal.r * WorldNormal;
 		}
 
 		#if SUPPORT_IRRADIANCE
@@ -588,7 +594,7 @@ void UpdateHistoryPS(
 	// 0 is stored where camera velocity should be used, which must not be filtered with valid object velocities
 	float2 VelocityN = Texture2DSampleLevel(VelocityTexture, VelocityTextureSampler, FullResTexel, 0).xy;
 
-	if (VelocityN.x > 0 && View.GeneralPurposeTweak > 0)
+	if (VelocityN.x > 0)
 	{
 		MotionVector = float2(-.5f, .5f) * DecodeVelocityFromTexture(VelocityN);
 	}
diff --git a/Engine/Shaders/DistanceFieldSurfaceCacheLightingCompute.usf b/Engine/Shaders/DistanceFieldSurfaceCacheLightingCompute.usf
--- a/Engine/Shaders/DistanceFieldSurfaceCacheLightingCompute.usf
+++ b/Engine/Shaders/DistanceFieldSurfaceCacheLightingCompute.usf
@@ -624,7 +624,12 @@ void PopulateCacheCS(
 			float3 OpaqueWorldPosition = HomogeneousWorldPosition.xyz / HomogeneousWorldPosition.w;
 
 			float2 BaseLevelScreenUV = (DispatchThreadId.xy * DownsampleFactorToBaseLevel + float2(.5f, .5f)) * BaseLevelTexelSize;
-			float3 WorldNormal = Texture2DSampleLevel(DistanceFieldNormalTexture, DistanceFieldNormalSampler, BaseLevelScreenUV, 0).xyz;
+
+			float3 WorldNormal;
+			float Unused;
+			bool bHasDistanceFieldRepresentation;
+			bool bHasHeightfieldRepresentation;
+			GetDownsampledGBuffer(BaseLevelScreenUV, WorldNormal, Unused, bHasDistanceFieldRepresentation, bHasHeightfieldRepresentation);
 
 			//@todo - offset shading position along normal to avoid incorrect self-occlusion?
 			float3 WorldShadingPosition = OpaqueWorldPosition;
@@ -636,11 +641,11 @@ void PopulateCacheCS(
 				int NextSampleIndex;
 				InterlockedAdd(NumQueuedIrradianceCacheRecords, 1, NextSampleIndex);
 				// W stores max allowed radius, used to limit overdraw from nearby samples placed in the high resolution passes
-				float MaxRadiusScale = GBufferData.HasDistanceFieldRepresentation > 0 ? .005f : .0005f;
+				float MaxRadiusScale = bHasDistanceFieldRepresentation > 0 ? .005f : .0005f;
 				// Sign of W stores whether the sample was placed on geometry with a distance field representation or not
-				float RadiusSign = GBufferData.HasDistanceFieldRepresentation > 0 ? 1 : -1;
+				float RadiusSign = bHasDistanceFieldRepresentation > 0 ? 1 : -1;
 				CachedIrradianceCachePositionRadius[NextSampleIndex] = float4(WorldShadingPosition, RadiusSign * SceneDepth * CurrentLevelDownsampleFactor * MaxRadiusScale);
-				CachedIrradianceCacheNormal[NextSampleIndex] = float4(WorldNormal, 0);
+				CachedIrradianceCacheNormal[NextSampleIndex] = float4(WorldNormal, bHasHeightfieldRepresentation ? 1 : 0);
 				CachedIrradianceCacheTileCoordinate[NextSampleIndex] = TileCoordinate;
 			}
 		}
diff --git a/Engine/Shaders/HeightfieldLighting.usf b/Engine/Shaders/HeightfieldLighting.usf
--- a/Engine/Shaders/HeightfieldLighting.usf
+++ b/Engine/Shaders/HeightfieldLighting.usf
@@ -393,13 +393,14 @@ void CalculateHeightfieldOcclusionCS(
 	uint RecordIndex = StartIndex + DispatchThreadId.x;
 
 	float4 PositionAndRadius = IrradianceCachePositionRadius[RecordIndex];
+	float4 WorldNormalAndHasHeightfieldRepresentation = IrradianceCacheNormal[RecordIndex];
 
 	// Only compute heightfield occlusion contribution for samples created on a distance field, 
 	// Otherwise we are potentially self-shadowing, the heightfield is not accurate enough for that
 	BRANCH
-	if (RecordIndex < NumRecords && PositionAndRadius.w > 0)
+	if (RecordIndex < NumRecords && WorldNormalAndHasHeightfieldRepresentation.w == 0)
 	{
-		float3 WorldNormal = IrradianceCacheNormal[RecordIndex].xyz;
+		float3 WorldNormal = WorldNormalAndHasHeightfieldRepresentation.xyz;
 		float3 WorldPosition = PositionAndRadius.xyz;
 
 		float3 TangentX;
diff --git a/Engine/Shaders/MaterialTemplate.usf b/Engine/Shaders/MaterialTemplate.usf
--- a/Engine/Shaders/MaterialTemplate.usf
+++ b/Engine/Shaders/MaterialTemplate.usf
@@ -895,6 +895,7 @@ float3 GetShadingModelColor(FGBufferData GBuffer)
 	else if (GBuffer.ShadingModelID == SHADINGMODELID_DEFAULT_LIT) return float3(0.1f, 1.0f, 0.1f); // Green
 	else if (GBuffer.ShadingModelID == SHADINGMODELID_SUBSURFACE) return float3(1.0f, 0.1f, 0.1f); // Red
 	else if (GBuffer.ShadingModelID == SHADINGMODELID_PREINTEGRATED_SKIN) return float3(0.6f, 0.4f, 0.1f); // Brown
+	else if (GBuffer.ShadingModelID == SHADINGMODELID_CLEAR_COAT) return float3(0.1f, 0.4f, 0.4f); 
 	else if (GBuffer.ShadingModelID == SHADINGMODELID_SUBSURFACE_PROFILE) return float3(0.2f, 0.6f, 0.5f); // Cyan
 	else if (GBuffer.ShadingModelID == SHADINGMODELID_TWOSIDED_FOLIAGE) return float3(0.2f, 0.2f, 0.8f); // Blue
 	else return float3(1.0f, 1.0f, 1.0f); // White
@@ -905,6 +906,7 @@ float3 GetShadingModelColor(FGBufferData GBuffer)
 		case SHADINGMODELID_DEFAULT_LIT: return float3(0.1f, 1.0f, 0.1f); // Green
 		case SHADINGMODELID_SUBSURFACE: return float3(1.0f, 0.1f, 0.1f); // Red
 		case SHADINGMODELID_PREINTEGRATED_SKIN: return float3(0.6f, 0.4f, 0.1f); // Brown
+		case SHADINGMODELID_CLEAR_COAT: return float3(0.1f, 0.4f, 0.4f); // Brown
 		case SHADINGMODELID_SUBSURFACE_PROFILE: return float3(0.2f, 0.6f, 0.5f); // Cyan
 		case SHADINGMODELID_TWOSIDED_FOLIAGE: return float3(0.2f, 0.2f, 0.8f); // Cyan
 		default: return float3(1.0f, 1.0f, 1.0f); // White
diff --git a/Engine/Source/Runtime/Engine/Private/DynamicMeshBuilder.cpp b/Engine/Source/Runtime/Engine/Private/DynamicMeshBuilder.cpp
--- a/Engine/Source/Runtime/Engine/Private/DynamicMeshBuilder.cpp
+++ b/Engine/Source/Runtime/Engine/Private/DynamicMeshBuilder.cpp
@@ -473,6 +473,7 @@ void FDynamicMeshBuilder::GetMesh(const FMatrix& LocalToWorld,const FMaterialRen
 			bReceivesDecals,
 			false,
 			false,
+			false,
 			1.0f		// LPV bias
 			);
 
@@ -542,6 +543,7 @@ void FDynamicMeshBuilder::Draw(FPrimitiveDrawInterface* PDI,const FMatrix& Local
 			bReceivesDecals,
 			false,
 			false,
+			false,
 			1.0f		// LPV bias
 			);
 
diff --git a/Engine/Source/Runtime/Engine/Private/NiagaraEffectRenderer.cpp b/Engine/Source/Runtime/Engine/Private/NiagaraEffectRenderer.cpp
--- a/Engine/Source/Runtime/Engine/Private/NiagaraEffectRenderer.cpp
+++ b/Engine/Source/Runtime/Engine/Private/NiagaraEffectRenderer.cpp
@@ -115,6 +115,7 @@ void NiagaraEffectRendererSprites::GetDynamicMeshElements(const TArray<const FSc
 				SceneProxy->GetLocalBounds(),
 				SceneProxy->ReceivesDecals(),
 				false,
+				false,
 				SceneProxy->UseEditorDepthTest()
 				);
 			WorldSpacePrimitiveUniformBuffer.SetContents(PrimitiveUniformShaderParameters);
@@ -370,6 +371,7 @@ void NiagaraEffectRendererRibbon::GetDynamicMeshElements(const TArray<const FSce
 				SceneProxy->GetLocalBounds(),
 				SceneProxy->ReceivesDecals(),
 				false,
+				false,
 				SceneProxy->UseEditorDepthTest()
 				);
 			WorldSpacePrimitiveUniformBuffer.SetContents(PrimitiveUniformShaderParameters);
diff --git a/Engine/Source/Runtime/Engine/Private/Particles/ParticleSystemRender.cpp b/Engine/Source/Runtime/Engine/Private/Particles/ParticleSystemRender.cpp
--- a/Engine/Source/Runtime/Engine/Private/Particles/ParticleSystemRender.cpp
+++ b/Engine/Source/Runtime/Engine/Private/Particles/ParticleSystemRender.cpp
@@ -6552,6 +6552,7 @@ void FParticleSystemSceneProxy::UpdateWorldSpacePrimitiveUniformBuffer() const
 			GetLocalBounds(),
 			ReceivesDecals(),
 			false,
+			false,
 			UseEditorDepthTest(),
 			1.0f			// LPV bias
 			);
diff --git a/Engine/Source/Runtime/Engine/Private/PrimitiveSceneProxy.cpp b/Engine/Source/Runtime/Engine/Private/PrimitiveSceneProxy.cpp
--- a/Engine/Source/Runtime/Engine/Private/PrimitiveSceneProxy.cpp
+++ b/Engine/Source/Runtime/Engine/Private/PrimitiveSceneProxy.cpp
@@ -185,14 +185,15 @@ void FPrimitiveSceneProxy::UpdateActorPosition(FVector InActorPosition)
 			// Update the uniform shader parameters.
 			const FPrimitiveUniformShaderParameters PrimitiveUniformShaderParameters =
 			GetPrimitiveUniformShaderParameters(
-			PrimitiveSceneProxy->LocalToWorld, 
-			InActorPosition, 
-			PrimitiveSceneProxy->Bounds,
-			PrimitiveSceneProxy->LocalBounds, 
-			PrimitiveSceneProxy->bReceivesDecals, 
-			PrimitiveSceneProxy->HasDistanceFieldRepresentation(), 
-			PrimitiveSceneProxy->UseEditorDepthTest(),
-			PrimitiveSceneProxy->GetLpvBiasMultiplier() );
+				PrimitiveSceneProxy->LocalToWorld, 
+				InActorPosition, 
+				PrimitiveSceneProxy->Bounds,
+				PrimitiveSceneProxy->LocalBounds, 
+				PrimitiveSceneProxy->bReceivesDecals, 
+				PrimitiveSceneProxy->HasDistanceFieldRepresentation(), 
+				PrimitiveSceneProxy->SupportsHeightfieldRepresentation(),
+				PrimitiveSceneProxy->UseEditorDepthTest(),
+				PrimitiveSceneProxy->GetLpvBiasMultiplier() );
 
 			PrimitiveSceneProxy->UniformBuffer.SetContents(PrimitiveUniformShaderParameters);
 			PrimitiveSceneProxy->OnActorPositionChanged();
@@ -214,7 +215,7 @@ void FPrimitiveSceneProxy::SetTransform(const FMatrix& InLocalToWorld, const FBo
 	ActorPosition = InActorPosition;
 	
 	// Update the uniform shader parameters.
-	const FPrimitiveUniformShaderParameters PrimitiveUniformShaderParameters = GetPrimitiveUniformShaderParameters(LocalToWorld, ActorPosition, Bounds, LocalBounds, bReceivesDecals, HasDistanceFieldRepresentation(), UseEditorDepthTest(), LpvBiasMultiplier );
+	const FPrimitiveUniformShaderParameters PrimitiveUniformShaderParameters = GetPrimitiveUniformShaderParameters(LocalToWorld, ActorPosition, Bounds, LocalBounds, bReceivesDecals, HasDistanceFieldRepresentation(), SupportsHeightfieldRepresentation(), UseEditorDepthTest(), LpvBiasMultiplier );
 	UniformBuffer.SetContents(PrimitiveUniformShaderParameters);
 	
 	// Notify the proxy's implementation of the change.
diff --git a/Engine/Source/Runtime/Engine/Public/PrimitiveUniformShaderParameters.h b/Engine/Source/Runtime/Engine/Public/PrimitiveUniformShaderParameters.h
--- a/Engine/Source/Runtime/Engine/Public/PrimitiveUniformShaderParameters.h
+++ b/Engine/Source/Runtime/Engine/Public/PrimitiveUniformShaderParameters.h
@@ -17,6 +17,7 @@ BEGIN_UNIFORM_BUFFER_STRUCT(FPrimitiveUniformShaderParameters,ENGINE_API)
 	DECLARE_UNIFORM_BUFFER_STRUCT_MEMBER(FVector,ActorWorldPosition)
 	DECLARE_UNIFORM_BUFFER_STRUCT_MEMBER_EX(float,DecalReceiverMask,EShaderPrecisionModifier::Half)
 	DECLARE_UNIFORM_BUFFER_STRUCT_MEMBER_EX(float,HasDistanceFieldRepresentation,EShaderPrecisionModifier::Half)
+	DECLARE_UNIFORM_BUFFER_STRUCT_MEMBER_EX(float,HasHeightfieldRepresentation,EShaderPrecisionModifier::Half)
 	DECLARE_UNIFORM_BUFFER_STRUCT_MEMBER_EX(float,UseEditorDepthTest,EShaderPrecisionModifier::Half)
 	DECLARE_UNIFORM_BUFFER_STRUCT_MEMBER_EX(FVector4,ObjectOrientation,EShaderPrecisionModifier::Half)
 	DECLARE_UNIFORM_BUFFER_STRUCT_MEMBER_EX(FVector4,NonUniformScale,EShaderPrecisionModifier::Half)
@@ -34,6 +35,7 @@ inline FPrimitiveUniformShaderParameters GetPrimitiveUniformShaderParameters(
 	const FBoxSphereBounds& LocalBounds,
 	bool bReceivesDecals,
 	bool bHasDistanceFieldRepresentation,
+	bool bHasHeightfieldRepresentation,
 	bool bUseEditorDepthTest,
 	float LpvBiasMultiplier = 1.0f
 )
@@ -69,6 +71,7 @@ inline FPrimitiveUniformShaderParameters GetPrimitiveUniformShaderParameters(
 	Result.LocalToWorldDeterminantSign = FMath::FloatSelect(LocalToWorld.RotDeterminant(),1,-1);
 	Result.DecalReceiverMask = bReceivesDecals ? 1 : 0;
 	Result.HasDistanceFieldRepresentation = bHasDistanceFieldRepresentation ? 1 : 0;
+	Result.HasHeightfieldRepresentation = bHasHeightfieldRepresentation ? 1 : 0;
 	Result.UseEditorDepthTest = bUseEditorDepthTest ? 1 : 0;
 	return Result;
 }
@@ -84,7 +87,7 @@ inline TUniformBufferRef<FPrimitiveUniformShaderParameters> CreatePrimitiveUnifo
 {
 	check(IsInRenderingThread());
 	return TUniformBufferRef<FPrimitiveUniformShaderParameters>::CreateUniformBufferImmediate(
-		GetPrimitiveUniformShaderParameters(LocalToWorld, WorldBounds.Origin, WorldBounds, LocalBounds, bReceivesDecals, false, bUseEditorDepthTest, LpvBiasMultiplier ),
+		GetPrimitiveUniformShaderParameters(LocalToWorld, WorldBounds.Origin, WorldBounds, LocalBounds, bReceivesDecals, false, false, bUseEditorDepthTest, LpvBiasMultiplier ),
 		UniformBuffer_MultiFrame
 		);
 }
@@ -111,6 +114,7 @@ public:
 			FBoxSphereBounds(EForceInit::ForceInit),
 			true,
 			false,
+			false,
 			true,
 			1.0f		// LPV bias
 			));
diff --git a/Engine/Source/Runtime/Landscape/Private/LandscapeRender.cpp b/Engine/Source/Runtime/Landscape/Private/LandscapeRender.cpp
--- a/Engine/Source/Runtime/Landscape/Private/LandscapeRender.cpp
+++ b/Engine/Source/Runtime/Landscape/Private/LandscapeRender.cpp
@@ -2607,6 +2607,12 @@ void FLandscapeComponentSceneProxy::GetHeightfieldRepresentation(UTexture2D*& Ou
 		HeightmapScaleBias.Z + SubsectionSizeVerts * NumSubsections * HeightmapScaleBias.X, 
 		HeightmapScaleBias.W + SubsectionSizeVerts * NumSubsections * HeightmapScaleBias.Y);
 
+	if (NumSubsections > 1)
+	{
+		OutDescription.MinMaxUV.Z -= HeightmapScaleBias.X;
+		OutDescription.MinMaxUV.W -= HeightmapScaleBias.Y;
+	}
+
 	OutDescription.HeightfieldRect = FIntRect(SectionBase.X, SectionBase.Y, SectionBase.X + NumSubsections * SubsectionSizeQuads, SectionBase.Y + NumSubsections * SubsectionSizeQuads);
 
 	OutDescription.NumSubsections = NumSubsections;
