Integrate - Material AO and SSAO now apply to heightfield GI

diff --git a/Engine/Shaders/SkyLighting.usf b/Engine/Shaders/SkyLighting.usf
--- a/Engine/Shaders/SkyLighting.usf
+++ b/Engine/Shaders/SkyLighting.usf
@@ -90,7 +90,8 @@ void SkyLightDiffusePS(
 		#endif
 #endif
 
-		float ScalarFactors = lerp(SkyVisibility, 1, OcclusionTintAndMinOcclusion.w) * GBuffer.GBufferAO * ScreenSpaceData.AmbientOcclusion;
+		float AmbientOcclusion = GBuffer.GBufferAO * ScreenSpaceData.AmbientOcclusion;
+		float ScalarFactors = lerp(SkyVisibility, 1, OcclusionTintAndMinOcclusion.w) * AmbientOcclusion;
 
 		if (ShadingModelId == SHADINGMODELID_TWOSIDED_FOLIAGE)
 		{
@@ -105,7 +106,7 @@ void SkyLightDiffusePS(
 		// Apply AO to the sky diffuse and account for darkening due to the geometry term
 		Lighting += (ScalarFactors * DotProductFactor) * DiffuseLookup * GBuffer.DiffuseColor + (1 - SkyVisibility) * OcclusionTintAndMinOcclusion.xyz;
 
-		Lighting += DiffuseIrradiance * GBuffer.DiffuseColor;
+		Lighting += DiffuseIrradiance * GBuffer.DiffuseColor * AmbientOcclusion;
 
 		// For debugging
 		#if SUPPORT_IRRADIANCE
