Integrate - Fixed bug where heightfield GI would gather lighting from the DFAO bent normal instead of the geometry normal
Backed off on heightfield GI local shadowing intensity, helps prevent dense trees from blocking all GI
GI is now scaled down on objects without a distance field or heightfield representation to compensate for lack of local occlusion

diff --git a/Engine/Shaders/DistanceFieldGlobalIllumination.usf b/Engine/Shaders/DistanceFieldGlobalIllumination.usf
--- a/Engine/Shaders/DistanceFieldGlobalIllumination.usf
+++ b/Engine/Shaders/DistanceFieldGlobalIllumination.usf
@@ -841,7 +841,7 @@ float3 ComputeVirtualPointLighting(uint VPLIndex, uint InstancedVPLIndex, uint R
 						StepBentNormalBuffer[RelativeRecordIndex * NUM_VISIBILITY_STEPS + LowerStepIndex + 1].xyz, 
 						saturate(LerpAlpha));
 					
-					float Shadow = GetVPLOcclusion(InterpolatedBentNormal, VPLDirection);
+					float Shadow = GetVPLOcclusion(InterpolatedBentNormal, VPLDirection, .5f, 1);
 
 				#else
 
diff --git a/Engine/Shaders/DistanceFieldLightingShared.usf b/Engine/Shaders/DistanceFieldLightingShared.usf
--- a/Engine/Shaders/DistanceFieldLightingShared.usf
+++ b/Engine/Shaders/DistanceFieldLightingShared.usf
@@ -297,11 +297,9 @@ float ApproximateConeConeIntersection(float ArcLength0, float ArcLength1, float
 	return Intersection;
 }
 
-float GetVPLOcclusion(float3 BentNormalAO, float3 NormalizedVectorToVPL)
+float GetVPLOcclusion(float3 BentNormalAO, float3 NormalizedVectorToVPL, float VPLConeAngle, float VPLOcclusionStrength)
 {
 	float BentNormalLength = length(BentNormalAO);
-	float VPLConeAngle = .5f;
-	float VPLOcclusionStrength = 1;
 	float UnoccludedAngle = BentNormalLength * PI / VPLOcclusionStrength;
 	float AngleBetween = acos(dot(BentNormalAO, NormalizedVectorToVPL) / max(BentNormalLength, .0001f));
 	float Visibility = ApproximateConeConeIntersection(VPLConeAngle, UnoccludedAngle, AngleBetween);
diff --git a/Engine/Shaders/DistanceFieldSurfaceCacheLighting.usf b/Engine/Shaders/DistanceFieldSurfaceCacheLighting.usf
--- a/Engine/Shaders/DistanceFieldSurfaceCacheLighting.usf
+++ b/Engine/Shaders/DistanceFieldSurfaceCacheLighting.usf
@@ -347,6 +347,7 @@ Buffer<float4> DebugBuffer;
 
 // 1 / ((1 - FadeDistanceFraction) * AOMaxViewDistance)
 float DistanceFadeScale;
+float SelfOcclusionReplacement;
 
 /** Normalizes the splatted surface cache values, packs depth in alpha. */
 void AOCombinePS(
@@ -391,6 +392,11 @@ void AOCombinePS(
 
 		#if SUPPORT_IRRADIANCE
 			OutIrradiance.rgb = Texture2DSampleLevel(IrradianceTexture, IrradianceSampler, UVAndScreenPos.xy, 0).xyz / BentNormalAccumulation.w;
+
+			if (!bHasDistanceFieldRepresentation && !bHasHeightfieldRepresentation)
+			{
+				OutIrradiance.rgb *= SelfOcclusionReplacement;
+			}
 		#endif
 	}
 	else
@@ -416,79 +422,6 @@ void AOCombinePS(
 		OutBentNormal.a = abs(OutBentNormal.a);
 	}
 
-#endif
-
-	/////////////////////////////////////////////////////////////////////////////////////////////////////////////////////////////////////
-#define VISUALIZE_SCENE_COLOR_GATHER 0
-#define VISUALIZE_DEBUG 0
-#if VISUALIZE_DEBUG && SUPPORT_IRRADIANCE
-
-	float2 ScreenPosition = (UVAndScreenPos.xy - View.ScreenPositionScaleBias.wz) / View.ScreenPositionScaleBias.xy;
-
-	float4 HomogeneousWorldPosition = mul(float4(ScreenPosition * SceneDepth, SceneDepth, 1), View.ScreenToWorld);
-	float3 OpaqueWorldPosition = HomogeneousWorldPosition.xyz / HomogeneousWorldPosition.w;
-	float3 ShadingPosition = DebugBuffer[0].xyz;
-
-	float Distance = length(OpaqueWorldPosition - ShadingPosition);
-	float ScaledDistance = pow(Distance / 20, .01f);
-	OutIrradiance.rgb = 2 * (1 - saturate(ScaledDistance)) * float3(.1f, 1, .1f);
-
-	for (int i = 0; i < 9; i++)
-	{
-		float4 GatheredSceneColorAndVisibility = DebugBuffer[1 + i * 3 + 2];
-
-		for (int j = 0; j < 1; j++)
-		{
-			float4 ClosestPositionAndRadius = DebugBuffer[1 + i * 3 + j];
-			float Distance2 = length(OpaqueWorldPosition - ClosestPositionAndRadius.xyz) / ClosestPositionAndRadius.w;
-			Distance2 = pow(Distance2, .5f);
-
-			if (Distance2 < 1)
-			{
-				//OutIrradiance.rgb += (1 - saturate(Distance2)) * 2 * float3(1, .1f, .1f) * (1 - GatheredSceneColorAndVisibility.w);
-				#if !VISUALIZE_SCENE_COLOR_GATHER
-					//OutIrradiance.rgb += GatheredSceneColorAndVisibility.xyz * (1 - GatheredSceneColorAndVisibility.w);
-				#endif
-			}
-
-			//if (GatheredSceneColorAndVisibility.w < .9f)
-			{
-				OutIrradiance.rgb += (1 - saturate(abs(Distance2 - 1) * 100)) * .4f * float3(1, .1f, .1f);
-			}
-		}
-	}
-
-	#if VISUALIZE_SCENE_COLOR_GATHER
-		
-		float SphereRadius = 50;
-
-		if (Distance < SphereRadius)
-		{
-			float3 ConePositionAtClosestToOccluder = ShadingPosition;
-			float4 ClipPosition = mul(float4(ConePositionAtClosestToOccluder, 1), View.WorldToClip);
-			float2 ScreenUV = ScreenAlignedPosition(ClipPosition).xy;
-			float UVDiskRadius = SphereRadius / ClipPosition.w * View.ViewToClip[0][0] * .5f;
-
-			uint NumSamplesPerCone = 81;
-
-			float3 GatheredColor = 0;
-			float GatheredWeight = .0001f;
-
-			for (uint SampleIndex = 0; SampleIndex < NumSamplesPerCone; SampleIndex++)
-			{
-				float2 Random = Hammersley(SampleIndex, NumSamplesPerCone, uint2(0, 0));
-				float2 UVOffset = UniformSampleDisk(Random);
-				float2 SampleUV = ScreenUV + UVOffset * UVDiskRadius;
-				float3 SampleSceneColor = CalcSceneColor(SampleUV);
-				float SampleSceneDepth = CalcSceneDepth(SampleUV);
-				float Weight = 1 - saturate(abs(SampleSceneDepth - ClipPosition.w) / SphereRadius);
-				GatheredColor += SampleSceneColor * Weight;
-				GatheredWeight += Weight;
-			} 
-		
-			OutIrradiance.rgb = GatheredColor / GatheredWeight;
-		}
-	#endif
 #endif
 }
 
diff --git a/Engine/Shaders/HeightfieldLighting.usf b/Engine/Shaders/HeightfieldLighting.usf
--- a/Engine/Shaders/HeightfieldLighting.usf
+++ b/Engine/Shaders/HeightfieldLighting.usf
@@ -356,7 +356,7 @@ void LightHeightfieldsPS(
 	float3 WorldHeightfieldNormal = GetGlobalHeightfieldNormal(Input.GlobalUV);
 	float3 WorldHeightfieldShadingPosition = GetGlobalHeightfieldWorldPosition(Input.LocalPosition, Input.GlobalUV, GetLocalToWorld(Input.HeightfieldIndex));
 	
-	float4 LightVector = mul(WorldHeightfieldShadingPosition, WorldToLight);
+	float4 LightVector = mul(float4(WorldHeightfieldShadingPosition, 1), WorldToLight);
 	LightVector.xyz /= LightVector.w;
 
 	float3 LightFunction = GetLightFunctionColor(LightVector.xyz, WorldHeightfieldShadingPosition);
@@ -536,7 +536,7 @@ float3 ComputeVirtualPointLighting(float3 WorldPosition, float3 WorldNormal, flo
 				#endif
 
 				float CosTheta = DirectionDot / Distance;
-				float Shadow = GetVPLOcclusion(BentNormal, VPLDirection);
+				float Shadow = GetVPLOcclusion(BentNormal, VPLDirection, .5f, .5f);
 				float VPLCosineLobe = saturate(VPLNormalDot / Distance);
 
 				Irradiance += (saturate(CosTheta) * VPLCosineLobe * DistanceAttenuation * Shadow) * VPLFlux;
@@ -550,8 +550,8 @@ float3 ComputeVirtualPointLighting(float3 WorldPosition, float3 WorldNormal, flo
 
 float3 GatherHeightfieldLighting(
 	float3 WorldPosition,
-	float3 BentNormal,
 	float3 WorldNormal,
+	float3 BentNormal,
 	float2 ShadingGlobalUV, 
 	float2 GlobalUVExtents,
 	float2 GlobalTexelSize, 
diff --git a/Engine/Source/Runtime/Renderer/Private/DistanceFieldGlobalIllumination.cpp b/Engine/Source/Runtime/Renderer/Private/DistanceFieldGlobalIllumination.cpp
--- a/Engine/Source/Runtime/Renderer/Private/DistanceFieldGlobalIllumination.cpp
+++ b/Engine/Source/Runtime/Renderer/Private/DistanceFieldGlobalIllumination.cpp
@@ -89,6 +89,14 @@ FAutoConsoleVariableRef CVarVPLSpreadUpdateOver(
 	ECVF_Cheat | ECVF_RenderThreadSafe
 	);
 
+float GVPLSelfOcclusionReplacement = .3f;
+FAutoConsoleVariableRef CVarVPLSelfOcclusionReplacement(
+	TEXT("r.VPLSelfOcclusionReplacement"),
+	GVPLSelfOcclusionReplacement,
+	TEXT(""),
+	ECVF_Cheat | ECVF_RenderThreadSafe
+	);
+
 TGlobalResource<FVPLResources> GVPLResources;
 TGlobalResource<FVPLResources> GCulledVPLResources;
 
diff --git a/Engine/Source/Runtime/Renderer/Private/DistanceFieldSurfaceCacheLighting.cpp b/Engine/Source/Runtime/Renderer/Private/DistanceFieldSurfaceCacheLighting.cpp
--- a/Engine/Source/Runtime/Renderer/Private/DistanceFieldSurfaceCacheLighting.cpp
+++ b/Engine/Source/Runtime/Renderer/Private/DistanceFieldSurfaceCacheLighting.cpp
@@ -2127,6 +2127,7 @@ public:
 		DistanceFieldNormalSampler.Bind(Initializer.ParameterMap, TEXT("DistanceFieldNormalSampler"));
 		DebugBuffer.Bind(Initializer.ParameterMap, TEXT("DebugBuffer"));
 		DistanceFadeScale.Bind(Initializer.ParameterMap, TEXT("DistanceFadeScale"));
+		SelfOcclusionReplacement.Bind(Initializer.ParameterMap, TEXT("SelfOcclusionReplacement"));
 	}
 
 	void SetParameters(
@@ -2164,6 +2165,9 @@ public:
 
 		const float DistanceFadeScaleValue = 1.0f / ((1.0f - GAOViewFadeDistanceScale) * GAOMaxViewDistance);
 		SetShaderValue(RHICmdList, ShaderRHI, DistanceFadeScale, DistanceFadeScaleValue);
+
+		extern float GVPLSelfOcclusionReplacement;
+		SetShaderValue(RHICmdList, ShaderRHI, SelfOcclusionReplacement, GVPLSelfOcclusionReplacement);
 	}
 
 	// FShader interface.
@@ -2181,6 +2185,7 @@ public:
 		Ar << DistanceFieldNormalSampler;
 		Ar << DebugBuffer;
 		Ar << DistanceFadeScale;
+		Ar << SelfOcclusionReplacement;
 		return bShaderHasOutdatedParameters;
 	}
 
@@ -2197,6 +2202,7 @@ private:
 	FShaderResourceParameter DistanceFieldNormalSampler;
 	FShaderResourceParameter DebugBuffer;
 	FShaderParameter DistanceFadeScale;
+	FShaderParameter SelfOcclusionReplacement;
 };
 
 IMPLEMENT_SHADER_TYPE(template<>,TDistanceFieldAOCombinePS<true>,TEXT("DistanceFieldSurfaceCacheLighting"),TEXT("AOCombinePS"),SF_Pixel);
