Integrate - Fixed small objects not consistently affecting DFAO - noticeable in test maps

diff --git a/Engine/Shaders/DistanceFieldSurfaceCacheLightingCompute.usf b/Engine/Shaders/DistanceFieldSurfaceCacheLightingCompute.usf
--- a/Engine/Shaders/DistanceFieldSurfaceCacheLightingCompute.usf
+++ b/Engine/Shaders/DistanceFieldSurfaceCacheLightingCompute.usf
@@ -809,7 +809,7 @@ void ConeTraceOcclusionCS(
 
 				BRANCH
 				// Skip tracing objects with a small projected angle 
-				if (ObjectPositionAndRadius.w * ObjectPositionAndRadius.w / ObjectDistanceSq > Square(.5f))
+				if (ObjectPositionAndRadius.w * ObjectPositionAndRadius.w / ObjectDistanceSq > Square(.25f))
 				{
 					float3 LocalPositionExtent = LoadObjectLocalPositionExtent(ObjectIndex);
 					float4x4 WorldToVolume = LoadObjectWorldToVolume(ObjectIndex);
