Integrate - Distance Field AO history filtering improvements
* Rejecting history based on world position differences, not camera Z.  This fixes a lot of noise under camera movement.
* The last texel of the history is never reused to avoid artifacts from bilinear filtering of a viewport in a larger render target
* Fixed error in the reprojected depth from storing something else in the sign of the depth
* There's now a 2d spatial filter of stable results applied to any pixels whose history was discarded (unstable)

diff --git a/Engine/Shaders/DistanceFieldSurfaceCacheLighting.usf b/Engine/Shaders/DistanceFieldSurfaceCacheLighting.usf
--- a/Engine/Shaders/DistanceFieldSurfaceCacheLighting.usf
+++ b/Engine/Shaders/DistanceFieldSurfaceCacheLighting.usf
@@ -456,9 +456,9 @@ void FillGapsPS(
 			float ValidMask = TextureValue.w > 0;
 
 			// Don't let depth weight go to 0 with huge depth differences
-			float DepthWeight = max(1 / exp2(abs(SceneDepth - SampleSceneDepth) * .01f), .001f);
+			float DepthWeight = max(exp2(-abs(SceneDepth - SampleSceneDepth) * .01f), .001f);
 
-			float2 Weight2D = 1 / exp2(abs(float2(x, y) * 10.0f / HALF_FILL_KERNEL_SIZE));
+			float2 Weight2D = exp2(-abs(float2(x, y) * 10.0f / HALF_FILL_KERNEL_SIZE));
 			float ScreenSpaceSpatialWeight = max(Weight2D.x, Weight2D.y);
 
 			float Weight = ValidMask * ScreenSpaceSpatialWeight * DepthWeight;
@@ -512,49 +512,72 @@ void UpdateHistoryPS(
 #endif
 	)
 {
-	float2 FullResTexel = UVAndScreenPos.xy - .5f * View.ViewSizeAndSceneTexelSize.zw;
-
-	float3 xyd;
-    xyd.xy = UVAndScreenPos.zw * float2(0.5, -0.5) + 0.5;
-    xyd.z = Texture2DSampleLevel(SceneDepthTexture, SceneDepthTextureSampler, FullResTexel, 0).r;
-
-    float scaleM = 1.0 / (dot(xyd, CameraMotion[0].xyz) + CameraMotion[0].w);
-    float2 MotionVector;
-    MotionVector.x = ((xyd.x * ((CameraMotion[1].x * xyd.y) + (CameraMotion[1].y * xyd.z) + CameraMotion[1].z)) + (CameraMotion[1].w * xyd.y) + (CameraMotion[2].x * xyd.x * xyd.x) + (CameraMotion[2].y * xyd.z) + CameraMotion[2].z) * scaleM;
-    MotionVector.y = ((xyd.y * ((CameraMotion[3].x * xyd.x) + (CameraMotion[3].y * xyd.z) + CameraMotion[3].z)) + (CameraMotion[3].w * xyd.x) + (CameraMotion[4].x * xyd.y * xyd.y) + (CameraMotion[4].y * xyd.z) + CameraMotion[4].z) * scaleM;
+	// Distance field AO was computed at 0,0 regardless of viewrect min
+	float2 DistanceFieldUVs = UVAndScreenPos.xy - View.ViewRectMin.xy * View.ViewSizeAndSceneTexelSize.zw;
+	float4 NewValue = Texture2DSampleLevel(BentNormalAOTexture, BentNormalAOSampler, DistanceFieldUVs, 0);
+	float SceneDepth = abs(NewValue.w);
 
-	// Note: have to sample from one of the high res texels exactly to avoid filtering of the velocity buffer
-	// 0 is stored where camera velocity should be used, which must not be filtered with valid object velocities
-	float2 VelocityN = Texture2DSampleLevel(VelocityTexture, VelocityTextureSampler, FullResTexel, 0).xy;
+	float2 OldUV;
 
-	if (VelocityN.x > 0)
 	{
-		MotionVector = float2(-.5f, .5f) * DecodeVelocityFromTexture(VelocityN);
-	}
+		float3 xyd;
+		xyd.xy = UVAndScreenPos.zw * float2(0.5, -0.5) + 0.5;
+		xyd.z = ConvertToDeviceZ(SceneDepth);
+
+		float scaleM = 1.0 / (dot(xyd, CameraMotion[0].xyz) + CameraMotion[0].w);
+		float2 MotionVector;
+		// Unwound vector-matrix transform with special care for precision taken from Temporal AA
+		MotionVector.x = ((xyd.x * ((CameraMotion[1].x * xyd.y) + (CameraMotion[1].y * xyd.z) + CameraMotion[1].z)) + (CameraMotion[1].w * xyd.y) + (CameraMotion[2].x * xyd.x * xyd.x) + (CameraMotion[2].y * xyd.z) + CameraMotion[2].z) * scaleM;
+		MotionVector.y = ((xyd.y * ((CameraMotion[3].x * xyd.x) + (CameraMotion[3].y * xyd.z) + CameraMotion[3].z)) + (CameraMotion[3].w * xyd.x) + (CameraMotion[4].x * xyd.y * xyd.y) + (CameraMotion[4].y * xyd.z) + CameraMotion[4].z) * scaleM;
+
+		// Note: have to sample from one of the high res texels exactly to avoid filtering of the velocity buffer
+		// 0 is stored where camera velocity should be used, which must not be filtered with valid object velocities
+		float2 FullResTexel = UVAndScreenPos.xy - .5f * View.ViewSizeAndSceneTexelSize.zw;
+		float2 VelocityN = Texture2DSampleLevel(VelocityTexture, VelocityTextureSampler, FullResTexel, 0).xy;
+
+		if (VelocityN.x > 0)
+		{
+			// Use the per-pixel velocity vector where valid, this handles object, bone and WPO movement
+			MotionVector = float2(-.5f, .5f) * DecodeVelocityFromTexture(VelocityN);
+		}
 
-    float2 OldUV = UVAndScreenPos.xy + MotionVector / float2(0.5, -0.5) * View.ScreenPositionScaleBias.xy;
+		OldUV = UVAndScreenPos.xy + MotionVector / float2(0.5, -0.5) * View.ScreenPositionScaleBias.xy;
+	}
+	
 	float2 OldScreenPosition = (OldUV.xy - View.ScreenPositionScaleBias.wz) / View.ScreenPositionScaleBias.xy;
 
-	// Distance field AO was computed at 0,0 regardless of viewrect min
-	float2 DistanceFieldUVs = UVAndScreenPos.xy - View.ViewRectMin.xy * View.ViewSizeAndSceneTexelSize.zw;
-	float4 NewValue = Texture2DSampleLevel(BentNormalAOTexture, BentNormalAOSampler, DistanceFieldUVs, 0);
 	float2 OldDistanceFieldUVs = OldUV.xy - View.ViewRectMin.xy * View.ViewSizeAndSceneTexelSize.zw;
 	float4 HistoryValue = Texture2DSampleLevel(BentNormalHistoryTexture, BentNormalHistorySampler, OldDistanceFieldUVs, 0);
 	
-	//@todo - this doesn't take camera movement into account
-	float DepthDeltaFraction = abs(abs(NewValue.a) - abs(HistoryValue.a)) / abs(NewValue.a);
+	float2 MinUV = View.ViewRectMin.xy * View.ViewSizeAndSceneTexelSize.zw;
+	// Pull in the max UV to exclude the region which will read outside the viewport due to bilinear filtering
+	float2 MaxUV = MinUV + (View.ViewSizeAndSceneTexelSize.xy - 1 * DOWNSAMPLE_FACTOR) * View.ViewSizeAndSceneTexelSize.zw;
+
+	float3 WorldPosition;
+	float3 PrevWorldPosition;
+
+	{
+		float4 HomogeneousWorldPosition = mul(float4(UVAndScreenPos.zw * SceneDepth, SceneDepth, 1), View.ScreenToWorld);
+		WorldPosition = HomogeneousWorldPosition.xyz / HomogeneousWorldPosition.w;
+
+		float HistoryDepth = abs(HistoryValue.w);
+		float4 PrevPositionTranslatedWorld = mul(float4(OldScreenPosition * HistoryDepth, HistoryDepth, 1), View.PrevScreenToTranslatedWorld);
+		PrevPositionTranslatedWorld.xyz /= PrevPositionTranslatedWorld.w;	
+		PrevWorldPosition = PrevPositionTranslatedWorld.xyz - View.PrevPreViewTranslation;
+	}
+
 	float EffectiveHistoryWeight = HistoryWeight;
-	
+
 	FLATTEN
-	if (any(OldScreenPosition > 1) 
-		|| any(OldScreenPosition < -1)
-		|| DepthDeltaFraction > .03f)
+	if (any(OldUV > MaxUV) 
+		|| any(OldUV < MinUV)
+		// Discard history if we are shading a new position (newly revealed)
+		|| length(PrevWorldPosition - WorldPosition) > 100)
 	{
 		EffectiveHistoryWeight = 0;
 	}
 
 	OutBentNormal.rgb = lerp(NewValue.rgb, HistoryValue.rgb, EffectiveHistoryWeight);
-	//OutBentNormal.rgb = HistoryValue.rgb + (NewValue.rgb - HistoryValue.rgb) * (1 - exp(-DeltaTime * EffectiveHistoryWeight));
 
 #if SUPPORT_IRRADIANCE
 	float3 NewIrradiance = Texture2DSampleLevel(IrradianceTexture, IrradianceSampler, DistanceFieldUVs, 0).xyz;
@@ -564,7 +587,86 @@ void UpdateHistoryPS(
 #endif
 
 	OutBentNormal.rgb = isnan(OutBentNormal.rgb) ? 0 : OutBentNormal.rgb;
-	OutBentNormal.a = NewValue.a;
+
+	// Sign bit of alpha stores whether the history was rejected or not
+	OutBentNormal.a = abs(NewValue.a) * (EffectiveHistoryWeight > 0 ? 1 : -1);
+}
+
+#define HALF_HISTORY_FILL_KERNEL_SIZE 2
+
+/** Seeds newly rejected history values (which are sources of temporal instability) with the results of a spatial search from stable history values */
+void FilterHistoryPS(
+	in float4 UVAndScreenPos : TEXCOORD0
+	,out float4 OutBentNormal : SV_Target0
+#if SUPPORT_IRRADIANCE
+	,out float4 OutIrradiance : SV_Target1
+#endif
+	)
+{
+	float4 HistoryValue = Texture2DSampleLevel(BentNormalAOTexture, BentNormalAOSampler, UVAndScreenPos.xy, 0);
+
+#if SUPPORT_IRRADIANCE
+	float3 IrradianceValue = Texture2DSampleLevel(IrradianceTexture, IrradianceSampler, UVAndScreenPos.xy, 0).xyz;
+#endif
+
+	// Only do the spatial search for pixels who discarded their history value
+	if (HistoryValue.w < 0)
+	{
+		float SceneDepth = abs(HistoryValue.w);
+
+		float4 Accumulation = 0;
+		float3 IrradianceAccumulation = 0;
+
+		for (float y = -HALF_HISTORY_FILL_KERNEL_SIZE; y <= HALF_HISTORY_FILL_KERNEL_SIZE; y++)
+		{
+			for (float x = -HALF_HISTORY_FILL_KERNEL_SIZE; x <= HALF_HISTORY_FILL_KERNEL_SIZE; x++)
+			{
+				float2 UVOffset = BentNormalAOTexelSize * float2(x, y);
+				float4 TextureValue = Texture2DSampleLevel(BentNormalAOTexture, BentNormalAOSampler, UVAndScreenPos.xy + UVOffset, 0);
+
+				float SampleSceneDepth = abs(TextureValue.w);
+				float ValidMask = TextureValue.w > 0;
+
+				// Weight by depth to avoid pulling in values of a foreground object
+				// This is a careful tradeoff between ghosting behind panning foreground objects and successful spatial searches to reduce flickering
+				float DepthWeight = exp2(-1000 * abs(SceneDepth - SampleSceneDepth) / SceneDepth);
+
+				float2 Weight2D = exp2(-abs(float2(x, y) * 10.0f / HALF_HISTORY_FILL_KERNEL_SIZE));
+				float ScreenSpaceSpatialWeight = max(Weight2D.x, Weight2D.y);
+
+				float Weight = ValidMask * ScreenSpaceSpatialWeight * DepthWeight;
+
+				Accumulation.rgb += TextureValue.rgb * Weight;
+				Accumulation.a += Weight;
+
+				#if SUPPORT_IRRADIANCE
+					float3 Irradiance = Texture2DSampleLevel(IrradianceTexture, IrradianceSampler, UVAndScreenPos.xy + UVOffset, 0).xyz;
+					IrradianceAccumulation += Irradiance * Weight;
+				#endif
+			}
+		}
+
+		// Only change the history value if the spatial search turned up something applicable
+		if (Accumulation.a > 0)
+		{
+			float InvWeight = 1.0f / Accumulation.a;
+			// Construct the history value as if the spatial search result was the previous history, 
+			// And the AO we just computed this frame was the new value
+			HistoryValue.xyz = lerp(HistoryValue.xyz, Accumulation.xyz * InvWeight, HistoryWeight);
+
+			#if SUPPORT_IRRADIANCE
+				IrradianceValue = lerp(IrradianceValue, IrradianceAccumulation * InvWeight, HistoryWeight);
+			#endif
+		}
+	}
+
+	OutBentNormal = HistoryValue;
+	// Remove sign bit so future reprojection interpolation isn't affected
+	OutBentNormal.w = abs(OutBentNormal.w);
+
+#if SUPPORT_IRRADIANCE
+	OutIrradiance = float4(IrradianceValue, 0);
+#endif
 }
 
 /** Upsamples the AO results to full resolution using a bilateral filter. */
diff --git a/Engine/Source/Runtime/Renderer/Private/DistanceFieldSurfaceCacheLighting.cpp b/Engine/Source/Runtime/Renderer/Private/DistanceFieldSurfaceCacheLighting.cpp
--- a/Engine/Source/Runtime/Renderer/Private/DistanceFieldSurfaceCacheLighting.cpp
+++ b/Engine/Source/Runtime/Renderer/Private/DistanceFieldSurfaceCacheLighting.cpp
@@ -2442,6 +2442,104 @@ private:
 IMPLEMENT_SHADER_TYPE(template<>,TUpdateHistoryPS<true>,TEXT("DistanceFieldSurfaceCacheLighting"),TEXT("UpdateHistoryPS"),SF_Pixel);
 IMPLEMENT_SHADER_TYPE(template<>,TUpdateHistoryPS<false>,TEXT("DistanceFieldSurfaceCacheLighting"),TEXT("UpdateHistoryPS"),SF_Pixel);
 
+
+template<bool bSupportIrradiance>
+class TFilterHistoryPS : public FGlobalShader
+{
+	DECLARE_SHADER_TYPE(TFilterHistoryPS, Global);
+public:
+
+	static bool ShouldCache(EShaderPlatform Platform)
+	{
+		return IsFeatureLevelSupported(Platform, ERHIFeatureLevel::SM5) && DoesPlatformSupportDistanceFieldAO(Platform);
+	}
+
+	static void ModifyCompilationEnvironment(EShaderPlatform Platform, FShaderCompilerEnvironment& OutEnvironment)
+	{
+		OutEnvironment.SetDefine(TEXT("DOWNSAMPLE_FACTOR"), GAODownsampleFactor);
+		OutEnvironment.SetDefine(TEXT("SUPPORT_IRRADIANCE"), bSupportIrradiance ? TEXT("1") : TEXT("0"));
+	}
+
+	/** Default constructor. */
+	TFilterHistoryPS() {}
+
+	/** Initialization constructor. */
+	TFilterHistoryPS(const ShaderMetaType::CompiledShaderInitializerType& Initializer)
+		: FGlobalShader(Initializer)
+	{
+		BentNormalAOTexture.Bind(Initializer.ParameterMap, TEXT("BentNormalAOTexture"));
+		BentNormalAOSampler.Bind(Initializer.ParameterMap, TEXT("BentNormalAOSampler"));
+		IrradianceTexture.Bind(Initializer.ParameterMap, TEXT("IrradianceTexture"));
+		IrradianceSampler.Bind(Initializer.ParameterMap, TEXT("IrradianceSampler"));
+		HistoryWeight.Bind(Initializer.ParameterMap, TEXT("HistoryWeight"));
+		BentNormalAOTexelSize.Bind(Initializer.ParameterMap, TEXT("BentNormalAOTexelSize"));
+	}
+
+	void SetParameters(
+		FRHICommandList& RHICmdList, 
+		const FSceneView& View, 
+		FSceneRenderTargetItem& BentNormalHistoryTextureValue, 
+		IPooledRenderTarget* IrradianceHistoryRT)
+	{
+		const FPixelShaderRHIParamRef ShaderRHI = GetPixelShader();
+
+		FGlobalShader::SetParameters(RHICmdList, ShaderRHI, View);
+
+		SetTextureParameter(
+			RHICmdList,
+			ShaderRHI,
+			BentNormalAOTexture,
+			BentNormalAOSampler,
+			TStaticSamplerState<SF_Bilinear,AM_Clamp,AM_Clamp,AM_Clamp>::GetRHI(),
+			BentNormalHistoryTextureValue.ShaderResourceTexture
+			);
+
+		if (IrradianceTexture.IsBound())
+		{
+			SetTextureParameter(
+				RHICmdList,
+				ShaderRHI,
+				IrradianceTexture,
+				IrradianceSampler,
+				TStaticSamplerState<SF_Bilinear,AM_Clamp,AM_Clamp,AM_Clamp>::GetRHI(),
+				IrradianceHistoryRT->GetRenderTargetItem().ShaderResourceTexture
+				);
+		}
+
+		SetShaderValue(RHICmdList, ShaderRHI, HistoryWeight, GAOHistoryWeight);
+		
+		const FIntPoint DownsampledBufferSize(GSceneRenderTargets.GetBufferSizeXY() / FIntPoint(GAODownsampleFactor, GAODownsampleFactor));
+		const FVector2D BaseLevelTexelSizeValue(1.0f / DownsampledBufferSize.X, 1.0f / DownsampledBufferSize.Y);
+		SetShaderValue(RHICmdList, ShaderRHI, BentNormalAOTexelSize, BaseLevelTexelSizeValue);
+	}
+
+	// FShader interface.
+	virtual bool Serialize(FArchive& Ar)
+	{
+		bool bShaderHasOutdatedParameters = FGlobalShader::Serialize(Ar);
+		Ar << BentNormalAOTexture;
+		Ar << BentNormalAOSampler;
+		Ar << IrradianceTexture;
+		Ar << IrradianceSampler;
+		Ar << HistoryWeight;
+		Ar << BentNormalAOTexelSize;
+		return bShaderHasOutdatedParameters;
+	}
+
+private:
+
+	FShaderResourceParameter BentNormalAOTexture;
+	FShaderResourceParameter BentNormalAOSampler;
+	FShaderResourceParameter IrradianceTexture;
+	FShaderResourceParameter IrradianceSampler;
+	FShaderParameter HistoryWeight;
+	FShaderParameter BentNormalAOTexelSize;
+};
+
+IMPLEMENT_SHADER_TYPE(template<>,TFilterHistoryPS<true>,TEXT("DistanceFieldSurfaceCacheLighting"),TEXT("FilterHistoryPS"),SF_Pixel);
+IMPLEMENT_SHADER_TYPE(template<>,TFilterHistoryPS<false>,TEXT("DistanceFieldSurfaceCacheLighting"),TEXT("FilterHistoryPS"),SF_Pixel);
+
+
 void AllocateOrReuseAORenderTarget(TRefCountPtr<IPooledRenderTarget>& Target, const TCHAR* Name, bool bWithAlpha)
 {
 	if (!Target)
@@ -2489,9 +2587,9 @@ void UpdateHistory(
 				AllocateOrReuseAORenderTarget(NewIrradianceHistory, IrradianceHistoryRTName, false);
 			}
 
-			{
-				SCOPED_DRAW_EVENT(RHICmdList, UpdateHistory);
+			SCOPED_DRAW_EVENT(RHICmdList, UpdateHistory);
 
+			{
 				FTextureRHIParamRef RenderTargets[2] =
 				{
 					NewBentNormalHistory->GetRenderTargetItem().TargetableTexture,
@@ -2537,12 +2635,72 @@ void UpdateHistory(
 					*VertexShader);
 			}
 
-			// Update the view state's render target reference with the new history
-			*BentNormalHistoryState = NewBentNormalHistory;
-			BentNormalHistoryOutput = NewBentNormalHistory;
+			const FPooledRenderTargetDesc& HistoryDesc = (*BentNormalHistoryState)->GetDesc();
+
+			// Reallocate history if buffer sizes have changed
+			if (HistoryDesc.Extent != GSceneRenderTargets.GetBufferSizeXY() / FIntPoint(GAODownsampleFactor, GAODownsampleFactor))
+			{
+				GRenderTargetPool.FreeUnusedResource(*BentNormalHistoryState);
+				*BentNormalHistoryState = NULL;
+				// Update the view state's render target reference with the new history
+				AllocateOrReuseAORenderTarget(*BentNormalHistoryState, BentNormalHistoryRTName, true);
+
+				if (bUseDistanceFieldGI)
+				{
+					GRenderTargetPool.FreeUnusedResource(*IrradianceHistoryState);
+					*IrradianceHistoryState = NULL;
+					AllocateOrReuseAORenderTarget(*IrradianceHistoryState, IrradianceHistoryRTName, false);
+				}
+			}
+
+			{
+				FTextureRHIParamRef RenderTargets[2] =
+				{
+					(*BentNormalHistoryState)->GetRenderTargetItem().TargetableTexture,
+					bUseDistanceFieldGI ? (*IrradianceHistoryState)->GetRenderTargetItem().TargetableTexture : NULL,
+				};
+
+				SetRenderTargets(RHICmdList, ARRAY_COUNT(RenderTargets) - (bUseDistanceFieldGI ? 0 : 1), RenderTargets, FTextureRHIParamRef(), 0, NULL);
+
+				RHICmdList.SetViewport(0, 0, 0.0f, View.ViewRect.Width() / GAODownsampleFactor, View.ViewRect.Height() / GAODownsampleFactor, 1.0f);
+				RHICmdList.SetRasterizerState(TStaticRasterizerState<FM_Solid, CM_None>::GetRHI());
+				RHICmdList.SetDepthStencilState(TStaticDepthStencilState<false, CF_Always>::GetRHI());
+				RHICmdList.SetBlendState(TStaticBlendState<>::GetRHI());
+
+				TShaderMapRef<FPostProcessVS> VertexShader(View.ShaderMap);
+
+				if (bUseDistanceFieldGI)
+				{
+					TShaderMapRef<TFilterHistoryPS<true> > PixelShader(View.ShaderMap);
+
+					static FGlobalBoundShaderState BoundShaderState;
+					SetGlobalBoundShaderState(RHICmdList, View.GetFeatureLevel(), BoundShaderState, GFilterVertexDeclaration.VertexDeclarationRHI, *VertexShader, *PixelShader);
+					PixelShader->SetParameters(RHICmdList, View, NewBentNormalHistory->GetRenderTargetItem(), NewIrradianceHistory);
+				}
+				else
+				{
+					TShaderMapRef<TFilterHistoryPS<false> > PixelShader(View.ShaderMap);
+
+					static FGlobalBoundShaderState BoundShaderState;
+					SetGlobalBoundShaderState(RHICmdList, View.GetFeatureLevel(), BoundShaderState, GFilterVertexDeclaration.VertexDeclarationRHI, *VertexShader, *PixelShader);
+					PixelShader->SetParameters(RHICmdList, View, NewBentNormalHistory->GetRenderTargetItem(), NewIrradianceHistory);
+				}
+
+				VertexShader->SetParameters(RHICmdList, View);
+
+				DrawRectangle( 
+					RHICmdList,
+					0, 0, 
+					View.ViewRect.Width() / GAODownsampleFactor, View.ViewRect.Height() / GAODownsampleFactor,
+					0, 0,
+					View.ViewRect.Width() / GAODownsampleFactor, View.ViewRect.Height() / GAODownsampleFactor,
+					FIntPoint(View.ViewRect.Width() / GAODownsampleFactor, View.ViewRect.Height() / GAODownsampleFactor),
+					GSceneRenderTargets.GetBufferSizeXY() / FIntPoint(GAODownsampleFactor, GAODownsampleFactor),
+					*VertexShader);
+			}
 
-			*IrradianceHistoryState = NewIrradianceHistory;
-			IrradianceHistoryOutput = NewIrradianceHistory;
+			BentNormalHistoryOutput = *BentNormalHistoryState;
+			IrradianceHistoryOutput = *IrradianceHistoryState;
 		}
 		else
 		{
