Integrate - Fixed regression from cl 2353634 in how subsurface color was handled with SSS shading models

diff --git a/Engine/Shaders/BasePassPixelShader.usf b/Engine/Shaders/BasePassPixelShader.usf
--- a/Engine/Shaders/BasePassPixelShader.usf
+++ b/Engine/Shaders/BasePassPixelShader.usf
@@ -724,10 +724,10 @@ void Main(
 		#if MATERIAL_SHADINGMODEL_SUBSURFACE || MATERIAL_SHADINGMODEL_PREINTEGRATED_SKIN
 			// Add subsurface energy to diffuse
 			//@todo - better subsurface handling for these shading models with skylight and precomputed GI
-			DiffuseColor += SubsurfaceColor;
+			EffectiveDiffuseColor += SubsurfaceColor;
 		#endif
 
-		Color += (DiffuseIndirectLighting * DiffuseColor + SubsurfaceIndirectLighting * SubsurfaceColor) * MaterialAO;
+		Color += (DiffuseIndirectLighting * EffectiveDiffuseColor + SubsurfaceIndirectLighting * SubsurfaceColor) * MaterialAO;
 
 		#if SIMPLE_DYNAMIC_LIGHTING
 			// always unshadowed so BiasedNDotL is not needed
