Integrate - Fixed Surfel GI shadowing

diff --git a/Engine/Shaders/DistanceFieldGlobalIllumination.usf b/Engine/Shaders/DistanceFieldGlobalIllumination.usf
--- a/Engine/Shaders/DistanceFieldGlobalIllumination.usf
+++ b/Engine/Shaders/DistanceFieldGlobalIllumination.usf
@@ -420,6 +420,14 @@ float RayTraceThroughLightTileCulledObjects(
 	return MinConeVisibility;
 }
 
+Buffer<uint> ShadowObjectIndirectArguments;
+
+uint GetShadowCulledNumObjects()
+{
+	// IndexCount, NumInstances, StartIndex, BaseVertexIndex, FirstInstance
+	return ShadowObjectIndirectArguments[1];
+}
+
 float3 ComputeVPLFlux(uint SurfelIndex, float4x4 InstanceToWorld)
 {
 	// Distance for directional lights to trace
@@ -450,7 +458,7 @@ float3 ComputeVPLFlux(uint SurfelIndex, float4x4 InstanceToWorld)
 				float3 WorldRayEnd;
 				float MaxRayTime;
 				float TanLightAngle;
-				uint NumIntersectingObjects = GetCulledNumObjects();
+				uint NumIntersectingObjects = GetShadowCulledNumObjects();
 				uint CulledDataStart = 0;
 
 				if (bIsPointLight)
@@ -475,7 +483,6 @@ float3 ComputeVPLFlux(uint SurfelIndex, float4x4 InstanceToWorld)
 
 					#if USE_SHADOW_CULLING_FOR_VPL_PLACEMENT
 
-						uint CulledDataStart;
 						GetShadowTileCulledData(WorldRayStart, CulledDataStart, NumIntersectingObjects);
 
 					#endif
@@ -496,7 +503,7 @@ float3 ComputeVPLFlux(uint SurfelIndex, float4x4 InstanceToWorld)
 		float3 DiffuseColor = LoadSurfelDiffuseColor(SurfelIndex);
 		//float3 EmissiveColor = LoadSurfelEmissiveColor(SurfelIndex);
 
-		Flux = (DiffuseColor * AccumulatedLighting/* + EmissiveColor*/) * PI * PositionAndRadius.w * PositionAndRadius.w;
+		Flux = (DiffuseColor * AccumulatedLighting/* + EmissiveColor*/);
 	}
 
 	return Flux;
@@ -814,10 +821,18 @@ float3 ComputeVirtualPointLighting(uint VPLIndex, uint InstancedVPLIndex, uint R
 				#define USE_INVERSE_SQUARED_DISK_APPROX 1
 
 				#if USE_INVERSE_SQUARED_DISK_APPROX
-					float EffectiveDiskRadius = VPLPositionAndRadius.w * 1;
-					float DistanceAttenuation = 1.0f / (PI * DistanceSq + PI * EffectiveDiskRadius * EffectiveDiskRadius);
-					//@todo - this is scale dependent
-					DistanceAttenuation = min(DistanceAttenuation, .0001f * RadialAttenuation(VPLToGatherPoint / VPLGatherRadius, 4));
+					float DiskRadiusSq = VPLPositionAndRadius.w * VPLPositionAndRadius.w;
+					float DistanceAttenuation = DiskRadiusSq / (DistanceSq + DiskRadiusSq);
+
+					float MinDistanceAttenuation = DiskRadiusSq / (VPLGatherRadius * VPLGatherRadius + DiskRadiusSq);
+					DistanceAttenuation = max(DistanceAttenuation - MinDistanceAttenuation, 0);
+
+					#define CONSERVE_ENERGY 0
+					#if CONSERVE_ENERGY 
+						float Integral = VPLPositionAndRadius.w * atan(VPLGatherRadius / VPLPositionAndRadius.w);
+						float EnergyConservationScale = Integral / (Integral - MinDistanceAttenuation * VPLGatherRadius);
+						DistanceAttenuation *= EnergyConservationScale;
+					#endif
 				#else
 					float DistanceAttenuation = RadialAttenuation(VPLToGatherPoint / VPLGatherRadius, 8);
 					DistanceAttenuation *= .0001f;
@@ -874,7 +889,7 @@ float3 ComputeVirtualPointLighting(uint VPLIndex, uint InstancedVPLIndex, uint R
 	float DistanceBehindMask = 1 - saturate(DistanceBehindVPL / (.5f * VPLPositionAndRadius.w));
 
 	float EffectiveDiskRadius = VPLPositionAndRadius.w * 1;
-	float DistanceAttenuation = 1.0f / (PI * DistanceSq + PI * EffectiveDiskRadius * EffectiveDiskRadius);
+	float DistanceAttenuation = VPLPositionAndRadius.w * VPLPositionAndRadius.w / (DistanceSq + EffectiveDiskRadius * EffectiveDiskRadius);
 	float3 VPLFlux = LoadVPLFlux(InstancedVPLIndex);
 	float NormalMask = saturate(dot(VPLNormal, WorldNormal));
 
diff --git a/Engine/Shaders/DistanceFieldShadowing.usf b/Engine/Shaders/DistanceFieldShadowing.usf
--- a/Engine/Shaders/DistanceFieldShadowing.usf
+++ b/Engine/Shaders/DistanceFieldShadowing.usf
@@ -145,7 +145,10 @@ void CullObjectsForShadowCS(
 	uint SourceIndex = ObjectIndex;
 	uint DestIndex = ObjectIndex;
 
-	CopyCulledObjectData(DestIndex, SourceIndex);
+	if (ObjectIndex < NumSceneObjects)
+	{
+		CopyCulledObjectData(DestIndex, SourceIndex);
+	}
 #endif
 }
 
diff --git a/Engine/Source/Runtime/Renderer/Private/DistanceFieldGlobalIllumination.cpp b/Engine/Source/Runtime/Renderer/Private/DistanceFieldGlobalIllumination.cpp
--- a/Engine/Source/Runtime/Renderer/Private/DistanceFieldGlobalIllumination.cpp
+++ b/Engine/Source/Runtime/Renderer/Private/DistanceFieldGlobalIllumination.cpp
@@ -639,6 +639,7 @@ public:
 		ShadowTileArrayData.Bind(Initializer.ParameterMap, TEXT("ShadowTileArrayData"));
 		ShadowTileListGroupSize.Bind(Initializer.ParameterMap, TEXT("ShadowTileListGroupSize"));
 		WorldToShadow.Bind(Initializer.ParameterMap, TEXT("WorldToShadow"));
+		ShadowObjectIndirectArguments.Bind(Initializer.ParameterMap, TEXT("ShadowObjectIndirectArguments"));
 		ShadowCulledObjectBounds.Bind(Initializer.ParameterMap, TEXT("ShadowCulledObjectBounds"));
 		ShadowCulledObjectData.Bind(Initializer.ParameterMap, TEXT("ShadowCulledObjectData"));
 		ObjectProcessStride.Bind(Initializer.ParameterMap, TEXT("ObjectProcessStride"));
@@ -694,6 +695,7 @@ public:
 
 		SetShaderValue(RHICmdList, ShaderRHI, WorldToShadow, WorldToShadowMatrixValue);
 
+		SetSRVParameter(RHICmdList, ShaderRHI, ShadowObjectIndirectArguments, GShadowCulledObjectBuffers.Buffers.ObjectIndirectArguments.SRV);
 		SetSRVParameter(RHICmdList, ShaderRHI, ShadowCulledObjectBounds, GShadowCulledObjectBuffers.Buffers.Bounds.SRV);
 		SetSRVParameter(RHICmdList, ShaderRHI, ShadowCulledObjectData, GShadowCulledObjectBuffers.Buffers.Data.SRV);
 
@@ -722,6 +724,7 @@ public:
 		Ar << ShadowTileArrayData;
 		Ar << ShadowTileListGroupSize;
 		Ar << WorldToShadow;
+		Ar << ShadowObjectIndirectArguments;
 		Ar << ShadowCulledObjectBounds;
 		Ar << ShadowCulledObjectData;
 		Ar << ObjectProcessStride;
@@ -743,6 +746,7 @@ private:
 	FShaderResourceParameter ShadowTileArrayData;
 	FShaderParameter ShadowTileListGroupSize;
 	FShaderParameter WorldToShadow;
+	FShaderResourceParameter ShadowObjectIndirectArguments;
 	FShaderResourceParameter ShadowCulledObjectBounds;
 	FShaderResourceParameter ShadowCulledObjectData;
 	FShaderParameter ObjectProcessStride;
