Fixed motion blur artifacts of sampling off screen.

diff --git a/Engine/Shaders/PostProcessMotionBlur.usf b/Engine/Shaders/PostProcessMotionBlur.usf
--- a/Engine/Shaders/PostProcessMotionBlur.usf
+++ b/Engine/Shaders/PostProcessMotionBlur.usf
@@ -631,11 +631,14 @@ void MainRecombinePS(float4 UVAndScreenPos : TEXCOORD0, out float4 OutColor : SV
 float2 DepthCmp( float CenterDepth, float SampleDepth, float DepthScale )
 {
 	return saturate( 0.5 + float2( DepthScale, -DepthScale ) * (SampleDepth - CenterDepth) );
+	//return SampleDepth > CenterDepth ? float2(1,0) : float2(0,1);
 }
 
 float2 SpreadCmp( float OffsetLength, float2 SpreadLength, float PixelToSampleScale )
 {
 	return saturate( PixelToSampleScale * SpreadLength - max( OffsetLength - 1, 0 ) );
+	//return PixelToSampleScale * SpreadLength > OffsetLength ? 1 : 0;
+	//return SpreadLength > SearchLengthPixels * OffsetLength / ( StepCount - 0.5 ) ? 1 : 0;
 }
 
 float SampleWeight( float CenterDepth, float SampleDepth, float OffsetLength, float CenterSpreadLength, float SampleSpreadLength, float PixelToSampleScale, float DepthScale )
@@ -666,7 +669,7 @@ void MainNewPS( float4 UVAndScreenPos : TEXCOORD0, out float4 OutColor : SV_Targ
 
 	float VelocityScale = MotionBlurParameters.x;
 
-	const uint StepCount = GetStepCountFromQuality() / 2;
+	const uint StepCount = GetStepCountFromQuality();
 
 #if 1
 	uint2 PixelPos = uint2(UVAndScreenPos.zw * ScreenPosToPixel.xy + ScreenPosToPixel.zw + 0.5f);
@@ -680,6 +683,7 @@ void MainNewPS( float4 UVAndScreenPos : TEXCOORD0, out float4 OutColor : SV_Targ
 	
 	float2 MaxVelocity = PostprocessInput3.SampleLevel( PostprocessInput3Sampler, UV + TileJitter * PostprocessInput3Size.zw, 0 ).xy;
 	//float2 MaxVelocity = PostprocessInput3[ PixelPos / 16 ].xy;
+
 	float2 MaxVelocityScreen = MaxVelocity.xy * VelocityScale;
 	float2 MaxVelocityPixels = MaxVelocityScreen * ScreenPosToPixel.xy;
 
@@ -700,14 +704,22 @@ void MainNewPS( float4 UVAndScreenPos : TEXCOORD0, out float4 OutColor : SV_Targ
 	CenterVelocity = DecodeVelocityFromTexture( CenterVelocity ) * VelocityScale * ScreenPosToPixel.xy;
 	float CenterVelocityLength = length( CenterVelocity );
 
-	float4 ColorAccum = 0;
+	
+	// Clip MaxVelocity to screen rect
+	float2 InvVelocityScreen = rcp( MaxVelocityScreen );
+	float2 MinIntersect = -InvVelocityScreen - ScreenSpacePos * InvVelocityScreen;
+	float2 MaxIntersect =  InvVelocityScreen - ScreenSpacePos * InvVelocityScreen;
+	float4 FarIntersect = float4( max( MinIntersect, MaxIntersect ), max( -MinIntersect, -MaxIntersect ) );
+	float2 Intersect = saturate( min( FarIntersect.xz, FarIntersect.yw ) );
+
+	// +/-
+	float4 SearchVectorPixels = MaxVelocityPixels.xyxy * float4( Intersect.xx, -Intersect.yy );
+	float2 SearchLengthPixels = { length( SearchVectorPixels.xy ), length( SearchVectorPixels.zw ) };
 
 	// converts pixel length to sample steps
-	float PixelToSampleScale = ( StepCount - 0.5 ) / length( MaxVelocityPixels );
-	//float PixelToSampleScale = ( StepCount - 0.5 + (Random - 0.5) ) / length( MaxVelocityPixels );
+	float2 PixelToSampleScale = ( StepCount - 0.5 ) / SearchLengthPixels;
 
-	// WHY?
-	PixelToSampleScale *= 2;
+	float4 ColorAccum = 0;
 
 	UNROLL
 	for( uint i = 0; i < StepCount; i++ )
@@ -721,9 +733,12 @@ void MainNewPS( float4 UVAndScreenPos : TEXCOORD0, out float4 OutColor : SV_Targ
 		float OffsetLength = (float)i + 0.5 + (Random - 0.5);
 		float OffsetFraction = OffsetLength / ( StepCount - 0.5 );
 
+		//float2 TestVector = (i & 1) > 0 ? CenterVelocity : SearchVectorPixels;
+		//float2 PixelToSampleScale = ( StepCount - 0.5 ) / length( TestVector );
+
 		// can be optimized
-		SampleUV[0] = UV + OffsetFraction * MaxVelocityPixels * PostprocessInput0Size.zw;
-		SampleUV[1] = UV - OffsetFraction * MaxVelocityPixels * PostprocessInput0Size.zw;
+		SampleUV[0] = UV + OffsetFraction * SearchVectorPixels.xy * PostprocessInput0Size.zw;
+		SampleUV[1] = UV + OffsetFraction * SearchVectorPixels.zw * PostprocessInput0Size.zw;
 
 		UNROLL
 		for( uint j = 0; j < 2; j++ )
@@ -738,7 +753,18 @@ void MainNewPS( float4 UVAndScreenPos : TEXCOORD0, out float4 OutColor : SV_Targ
 			SampleVelocity = DecodeVelocityFromTexture( SampleVelocity ) * VelocityScale * ScreenPosToPixel.xy;
 			SampleVelocityLength[j] = length( SampleVelocity );
 
-			Weight[j] = SampleWeight( CenterDepth, SampleDepth[j], OffsetLength, CenterVelocityLength, SampleVelocityLength[j], PixelToSampleScale, SOFT_Z_EXTENT );
+			Weight[j] = SampleWeight( CenterDepth, SampleDepth[j], OffsetLength, CenterVelocityLength, SampleVelocityLength[j], PixelToSampleScale[j], SOFT_Z_EXTENT );
+
+#if 0
+			float TestLengthPixels = OffsetFraction * SearchLengthPixels[j];
+
+			float2 DepthWeights = DepthCmp( CenterDepth, SampleDepth[j], SOFT_Z_EXTENT );
+			
+			//Weight[j]  = DepthWeights.x * saturate( 1 - TestLengthPixels / CenterVelocityLength );
+			//Weight[j] += DepthWeights.y * saturate( 1 - TestLengthPixels / SampleVelocityLength[j] );
+			//Weight[j] += ( 1 - smoothstep(0.95f * CenterVelocityLength, 1.05f * CenterVelocityLength, TestLengthPixels ) ) *
+			//			( 1 - smoothstep(0.95f * SampleVelocityLength[j], 1.05f * SampleVelocityLength[j], TestLengthPixels ) ) * 2;
+#endif
 		}
 		
 		bool2 Mirror = bool2( SampleDepth[0] > SampleDepth[1], SampleVelocityLength[1] > SampleVelocityLength[0] );
@@ -749,7 +775,19 @@ void MainNewPS( float4 UVAndScreenPos : TEXCOORD0, out float4 OutColor : SV_Targ
 		ColorAccum += Weight[1] * float4( SampleColor[1].rgb, 1 );
 	}
 
+	//ColorAccum += float4( CenterColor, 1 ) / CenterVelocityLength;
+	
 	ColorAccum *= 0.5 / StepCount;
 	OutColor.rgb = ColorAccum.rgb + ( 1 - ColorAccum.a ) * CenterColor;
+	
+#if 0
+	// Fudge factor to increase alpha
+	ColorAccum.rgb /= ColorAccum.a;
+	ColorAccum.a = saturate( ColorAccum.a * 1.2 );
+	OutColor.rgb = lerp( CenterColor, ColorAccum.rgb, ColorAccum.a );
+#endif
+	
 	OutColor.a = 1;
+
+	//OutColor.rgb = ColorAccum.a * 0.2;
 }
\ No newline at end of file
diff --git a/Engine/Source/Runtime/Renderer/Private/PostProcess/PostProcessMotionBlur.cpp b/Engine/Source/Runtime/Renderer/Private/PostProcess/PostProcessMotionBlur.cpp
--- a/Engine/Source/Runtime/Renderer/Private/PostProcess/PostProcessMotionBlur.cpp
+++ b/Engine/Source/Runtime/Renderer/Private/PostProcess/PostProcessMotionBlur.cpp
@@ -887,17 +887,17 @@ public:
 
 				FSamplerStateRHIParamRef Filters[] =
 				{
-					TStaticSamplerState<SF_Bilinear,AM_Border,AM_Border,AM_Clamp>::GetRHI(),
-					TStaticSamplerState<SF_Point,AM_Border,AM_Border,AM_Clamp>::GetRHI(),
-					TStaticSamplerState<SF_Bilinear,AM_Border,AM_Border,AM_Clamp>::GetRHI(),
-					TStaticSamplerState<SF_Point,AM_Border,AM_Border,AM_Clamp>::GetRHI(),
+					TStaticSamplerState<SF_Bilinear,AM_Clamp,AM_Clamp,AM_Clamp>::GetRHI(),
+					TStaticSamplerState<SF_Point,AM_Clamp,AM_Clamp,AM_Clamp>::GetRHI(),
+					TStaticSamplerState<SF_Bilinear,AM_Clamp,AM_Clamp,AM_Clamp>::GetRHI(),
+					TStaticSamplerState<SF_Point,AM_Clamp,AM_Clamp,AM_Clamp>::GetRHI(),
 				};
 
 				PostprocessParameter.SetPS( ShaderRHI, Context, 0, false, Filters );
 			}
 			else
 			{
-				PostprocessParameter.SetPS(ShaderRHI, Context, TStaticSamplerState<SF_Point,AM_Border,AM_Border,AM_Clamp>::GetRHI());
+				PostprocessParameter.SetPS(ShaderRHI, Context, TStaticSamplerState<SF_Point,AM_Clamp,AM_Clamp,AM_Clamp>::GetRHI());
 			}
 		}
 
