CSM fixes for forward renderer
 MAX_FORWARD_SHADOWCASCADES is correctly defined for OutEnvironment.
 Only render shadow maps that have been successfully allocated texture space.
 fwd shader can deal with full range of MAX_FORWARD_SHADOWCASCADES.
 limit shadow map renders to no more than MAX_FORWARD_SHADOWCASCADES.
 reduce shadow map size to ensure all tiles fit within GMaxShadowDepthBufferSize.
Fixes UE-10282.

diff --git a/Engine/Shaders/BasePassForForwardShadingPixelShader.usf b/Engine/Shaders/BasePassForForwardShadingPixelShader.usf
--- a/Engine/Shaders/BasePassForForwardShadingPixelShader.usf
+++ b/Engine/Shaders/BasePassForForwardShadingPixelShader.usf
@@ -204,7 +204,7 @@ void Main(
 #if MOVABLE_DIRECTIONAL_LIGHT_CSM
 		// Cascaded Shadow Map
 		{
-			FPCFSamplerSettings Settings;			
+			FPCFSamplerSettings Settings;
 			Settings.ShadowDepthTexture = View.DirectionalLightShadowTexture;
 			Settings.ShadowDepthTextureSampler = View.DirectionalLightShadowSampler;
 			Settings.TransitionScale = View.DirectionalLightShadowTransition;
@@ -212,24 +212,17 @@ void Main(
 			Settings.bSubsurface = false;
 			Settings.DensityMulConstant = 0;
 			Settings.ProjectionDepthBiasParameters = 0;
-						
-			float4 ShadowPosition;	
-	#if MAX_FORWARD_SHADOWCASCADES == 2			
-			{				
-				if (MaterialParameters.ScreenPosition.w >= View.DirectionalLightShadowDistances.x)
-				{
-					ShadowPosition = mul(float4(MaterialParameters.ScreenPosition.xyw, 1), View.DirectionalLightScreenToShadow[1]);
-				}
-				else
+
+			float4 ShadowPosition;
+			for (int i = 0; i < MAX_FORWARD_SHADOWCASCADES; i++)
+			{
+				if (i == (MAX_FORWARD_SHADOWCASCADES - 1) // ignore test for last shadow map.
+					|| MaterialParameters.ScreenPosition.w < View.DirectionalLightShadowDistances[i])
 				{
-					ShadowPosition = mul(float4(MaterialParameters.ScreenPosition.xyw, 1), View.DirectionalLightScreenToShadow[0]);
+					ShadowPosition = mul(float4(MaterialParameters.ScreenPosition.xyw, 1), View.DirectionalLightScreenToShadow[i]);
+					break; // position found.
 				}
 			}
-	#else			
-			{
-				ShadowPosition = mul(float4(MaterialParameters.ScreenPosition.xyw, 1), View.DirectionalLightScreenToShadow[0]);				
-			}
-	#endif
 
 			// Clamp pixel depth in light space for shadowing opaque, because areas of the shadow depth buffer that weren't rendered to will have been cleared to 1
 			// We want to force the shadow comparison to result in 'unshadowed' in that case, regardless of whether the pixel being shaded is in front or behind that plane
diff --git a/Engine/Source/Runtime/Renderer/Private/LightMapRendering.cpp b/Engine/Source/Runtime/Renderer/Private/LightMapRendering.cpp
--- a/Engine/Source/Runtime/Renderer/Private/LightMapRendering.cpp
+++ b/Engine/Source/Runtime/Renderer/Private/LightMapRendering.cpp
@@ -11,7 +11,7 @@ void FMovableDirectionalLightCSMLightingPolicy::ModifyCompilationEnvironment(ESh
 {
 	OutEnvironment.SetDefine(TEXT("MOVABLE_DIRECTIONAL_LIGHT"), TEXT("1"));
 	OutEnvironment.SetDefine(TEXT("MOVABLE_DIRECTIONAL_LIGHT_CSM"), TEXT("1"));
-	OutEnvironment.SetDefine(TEXT("MAX_FORWARD_SHADOWCASCADES"), TEXT(PREPROCESSOR_TO_STRING(MAX_FORWARD_SHADOWCASCADES)));
+	OutEnvironment.SetDefine(TEXT(PREPROCESSOR_TO_STRING(MAX_FORWARD_SHADOWCASCADES)), MAX_FORWARD_SHADOWCASCADES);
 
 	FNoLightMapPolicy::ModifyCompilationEnvironment(Platform, Material, OutEnvironment);
 }
\ No newline at end of file
diff --git a/Engine/Source/Runtime/Renderer/Private/LightMapRendering.h b/Engine/Source/Runtime/Renderer/Private/LightMapRendering.h
--- a/Engine/Source/Runtime/Renderer/Private/LightMapRendering.h
+++ b/Engine/Source/Runtime/Renderer/Private/LightMapRendering.h
@@ -984,7 +984,7 @@ public:
 	static void ModifyCompilationEnvironment(EShaderPlatform Platform, const FMaterial* Material, FShaderCompilerEnvironment& OutEnvironment)
 	{
 		OutEnvironment.SetDefine(TEXT("MOVABLE_DIRECTIONAL_LIGHT"), TEXT("1"));
-		OutEnvironment.SetDefine(TEXT("MAX_FORWARD_SHADOWCASCADES"), TEXT(PREPROCESSOR_TO_STRING(MAX_FORWARD_SHADOWCASCADES)));
+		OutEnvironment.SetDefine(TEXT(PREPROCESSOR_TO_STRING(MAX_FORWARD_SHADOWCASCADES)), MAX_FORWARD_SHADOWCASCADES);
 		FSimpleDirectionalLightAndSHIndirectPolicy::ModifyCompilationEnvironment(Platform, Material, OutEnvironment);
 	}
 
@@ -1060,7 +1060,7 @@ public:
 	static void ModifyCompilationEnvironment(EShaderPlatform Platform, const FMaterial* Material, FShaderCompilerEnvironment& OutEnvironment)
 	{
 		OutEnvironment.SetDefine(TEXT("MOVABLE_DIRECTIONAL_LIGHT"), TEXT("1"));
-		OutEnvironment.SetDefine(TEXT("MAX_FORWARD_SHADOWCASCADES"), TEXT(PREPROCESSOR_TO_STRING(MAX_FORWARD_SHADOWCASCADES)));
+		OutEnvironment.SetDefine(TEXT(PREPROCESSOR_TO_STRING(MAX_FORWARD_SHADOWCASCADES)), MAX_FORWARD_SHADOWCASCADES);
 
 		Super::ModifyCompilationEnvironment(Platform, Material, OutEnvironment);
 	}
diff --git a/Engine/Source/Runtime/Renderer/Private/PostProcess/SceneRenderTargets.cpp b/Engine/Source/Runtime/Renderer/Private/PostProcess/SceneRenderTargets.cpp
--- a/Engine/Source/Runtime/Renderer/Private/PostProcess/SceneRenderTargets.cpp
+++ b/Engine/Source/Runtime/Renderer/Private/PostProcess/SceneRenderTargets.cpp
@@ -235,6 +235,17 @@ void FSceneRenderTargets::Allocate(const FSceneViewFamily& ViewFamily)
 
 	int32 MaxShadowResolution = GetCachedScalabilityCVars().MaxShadowResolution;
 
+	if (ViewFamily.Scene->ShouldUseDeferredRenderer() == false)
+	{
+		// ensure there is always enough space for forward renderer's tiled shadow maps
+		// by reducing the shadow map resolution.
+		int32 MaxShadowDepthBufferDim = FMath::Max(GMaxShadowDepthBufferSizeX, GMaxShadowDepthBufferSizeY);
+		if (MaxShadowResolution * 2 >  MaxShadowDepthBufferDim)
+		{
+			MaxShadowResolution = MaxShadowDepthBufferDim / 2;
+		}
+	}
+
 	int32 TranslucencyLightingVolumeDim = GTranslucencyLightingVolumeDim;
 
 	uint32 Mobile32bpp = !IsMobileHDR() || IsMobileHDR32bpp();
diff --git a/Engine/Source/Runtime/Renderer/Private/ShadowRendering.cpp b/Engine/Source/Runtime/Renderer/Private/ShadowRendering.cpp
--- a/Engine/Source/Runtime/Renderer/Private/ShadowRendering.cpp
+++ b/Engine/Source/Runtime/Renderer/Private/ShadowRendering.cpp
@@ -3547,11 +3547,14 @@ bool FForwardShadingSceneRenderer::RenderShadowDepthMap(FRHICommandListImmediate
 		for (int32 ShadowIndex = 0; ShadowIndex < Shadows.Num(); ShadowIndex++)
 		{
 			FProjectedShadowInfo* ProjectedShadowInfo = Shadows[ShadowIndex];
-			check(!ProjectedShadowInfo->bTranslucentShadow);
-			ProjectedShadowInfo->bRendered = false;
-			ProjectedShadowInfo->bAllocated = true;
-			ProjectedShadowInfo->RenderDepth(RHICmdList, this, SetShadowRenderTargets);
-			bAttenuationBufferDirty = true;
+			if (ProjectedShadowInfo->bAllocated)
+			{
+				check(!ProjectedShadowInfo->bTranslucentShadow);
+				ProjectedShadowInfo->RenderDepth(RHICmdList, this, SetShadowRenderTargets);
+				ProjectedShadowInfo->bRendered = true;
+				ProjectedShadowInfo->bAllocated = false;
+				bAttenuationBufferDirty = true;
+			}
 		}
 
 		GSceneRenderTargets.FinishRenderingShadowDepth(RHICmdList);
diff --git a/Engine/Source/Runtime/Renderer/Private/ShadowSetup.cpp b/Engine/Source/Runtime/Renderer/Private/ShadowSetup.cpp
--- a/Engine/Source/Runtime/Renderer/Private/ShadowSetup.cpp
+++ b/Engine/Source/Runtime/Renderer/Private/ShadowSetup.cpp
@@ -2257,7 +2257,7 @@ void FForwardShadingSceneRenderer::InitDynamicShadows(FRHICommandListImmediate&
 					AddViewDependentWholeSceneShadowsForView(ViewDependentWholeSceneShadows, ViewDependentWholeSceneShadowsThatNeedCulling, VisibleLightInfo, *LightSceneInfo);					
 				}
 
-				const int32 NumShadows = VisibleLightInfo.AllProjectedShadows.Num();
+				const int32 NumShadows = FMath::Min(VisibleLightInfo.AllProjectedShadows.Num(), MAX_FORWARD_SHADOWCASCADES);
 				if (NumShadows > 0)
 				{
 					//create the shadow depth texture and/or surface
