fixed spot lights for TLM_SurfacePerPixelLighting, remaining: directional lights and shadow support

diff --git a/Engine/Shaders/ForEachLight.usf b/Engine/Shaders/ForEachLight.usf
--- a/Engine/Shaders/ForEachLight.usf
+++ b/Engine/Shaders/ForEachLight.usf
@@ -15,38 +15,31 @@ float3 SimplePointLightDiffuse( FScreenSpaceData ScreenSpaceData )
 	return Diffuse_Lambert(ScreenSpaceData.GBuffer.DiffuseColor);
 }
 
-float3 LightingForOneLight(FMaterialPixelParameters Parameters, FSimpleDeferredLightData LightData, FScreenSpaceData ScreenSpaceData)
+float3 LightingForOneLight(FMaterialPixelParameters Parameters, uint LightIndex, FScreenSpaceData ScreenSpaceData)
 {
 	float3 CameraVector = normalize(Parameters.WorldPosition - View.ViewOrigin.xyz);
-
-	// At the moment we only support very simple lighting (no shadows, phong, no area lights, basic attenuation).
-	// Later we can extend this function.
-
-	float3 Surface = GetSimpleDynamicLighting(Parameters.WorldPosition, CameraVector, ScreenSpaceData, LightData);
-	// todo: investigate why this doesn't work
-//	float3 Surface = GetSimpleDynamicLighting(Parameters.WorldPosition, Parameters.CameraVector, ScreenSpaceData, LightData);
-
-	// todo: cleanup (share with former code)
-
-	// assuming inverse square falloff
-
-	half Opacity = GetMaterialOpacity(Parameters);
-	float3 ToLight = LightData.LightPositionAndInvRadius.xyz - Parameters.WorldPosition;
-
-	// Sphere falloff (technically just 1/d2 but this avoids inf)
-	float DistanceSqr = dot( ToLight, ToLight );
-	float DistanceAttenuation = 1 / ( DistanceSqr + 1 );
 	
-	float LightRadiusMask = Square( saturate( 1 - Square( DistanceSqr * Square(LightData.LightPositionAndInvRadius.w) ) ) );
-
-	float SubsurfaceAttenuation	= DistanceAttenuation * LightRadiusMask;
-
-	float3 Subsurface = (1 - Opacity) * LightData.LightColorAndFalloffExponent.rgb * SubsurfaceAttenuation
-		* TwoSidedFoliageShadingSubsurface(ScreenSpaceData.GBuffer.DiffuseColor, normalize(ToLight), -CameraVector, ScreenSpaceData.GBuffer.WorldNormal);
-
-
-	return Subsurface + Surface;
+	FDeferredLightData LightData = (FDeferredLightData)0;
+	LightData.LightPositionAndInvRadius = ForwardLightData.LightPositionAndInvRadius[LightIndex];
+	LightData.LightColorAndFalloffExponent = ForwardLightData.LightColorAndFalloffExponent[LightIndex];
+	{
+		float4 Value = ForwardLightData.LightDirectionAndSpotlightMaskAndMinRoughness[LightIndex];
+		LightData.LightDirection = Value.xyz;
+		LightData.bSpotLight = Value.w > 0;
+		LightData.MinRoughness = abs(Value.w);
+	}
+	LightData.SpotAnglesAndSourceRadius = float4( ForwardLightData.SpotAnglesAndSourceRadiusAndSimpleLighting[LightIndex].xyz, 0 );
+	// currently we don't support shadows
+	LightData.ShadowMapChannelMask = 0;
+	LightData.bInverseSquared = ForwardLightData.LightColorAndFalloffExponent[LightIndex].w == 0;
+	// Only radial lights supported with tiled deferred
+	LightData.bRadialLight = true;
+	// The only type of shadowing supported for lights using tiled is static shadowing, so the light should only compute shadowing if it has static shadowing
+	LightData.bShadowed = dot(LightData.ShadowMapChannelMask, float4(1, 1, 1, 1));
+	
+	float2 ScreenUV = Parameters.SVPosition.xy * View.ViewSizeAndSceneTexelSize.zw;
 
+	return GetDynamicLighting(Parameters.WorldPosition, CameraVector, ScreenUV, ScreenSpaceData, ScreenSpaceData.GBuffer.ShadingModelID, LightData, float4(1, 1, 1, 1), uint2(0,0)).rgb;
 }
 
 // @param WorldNormal is not assumed to be normalized
@@ -68,6 +61,13 @@ float3 AccumulateForwardLights(FMaterialPixelParameters Parameters, float3 Diffu
 	ScreenSpaceData.GBuffer.Roughness = Roughness;
 	ScreenSpaceData.GBuffer.SpecularColor = SpecularColor;
 	ScreenSpaceData.GBuffer.DiffuseColor = DiffuseColor;
+	ScreenSpaceData.GBuffer.BaseColor = 1;
+	ScreenSpaceData.GBuffer.Specular = 1;
+	ScreenSpaceData.GBuffer.Opacity = GetMaterialOpacity(Parameters);
+	ScreenSpaceData.GBuffer.PrecomputedShadowFactors = 1;
+//	ScreenSpaceData.GBuffer.ShadingModelID = SHADINGMODELID_DEFAULT_LIT;
+	ScreenSpaceData.GBuffer.ShadingModelID = SHADINGMODELID_TWOSIDED_FOLIAGE;
+	ScreenSpaceData.GBuffer.CustomData = DiffuseColor;
 	ScreenSpaceData.AmbientOcclusion = 1;
 
 	// todo: support more than 32 lights
@@ -79,12 +79,7 @@ float3 AccumulateForwardLights(FMaterialPixelParameters Parameters, float3 Diffu
 
 		culling &= ~(1 << LightIndex);
 
-		FSimpleDeferredLightData LightData = (FSimpleDeferredLightData)0;
-		LightData.LightPositionAndInvRadius = ForwardLightData.LightPositionAndInvRadius[LightIndex];
-		LightData.LightColorAndFalloffExponent = ForwardLightData.LightColorAndFalloffExponent[LightIndex];
-		LightData.bInverseSquared = true;
-
-		Ret += LightingForOneLight(Parameters, LightData, ScreenSpaceData);
+		Ret += LightingForOneLight(Parameters, LightIndex, ScreenSpaceData);
 	}
 #endif
 
diff --git a/Engine/Source/Runtime/Engine/Public/SceneView.h b/Engine/Source/Runtime/Engine/Public/SceneView.h
--- a/Engine/Source/Runtime/Engine/Public/SceneView.h
+++ b/Engine/Source/Runtime/Engine/Public/SceneView.h
@@ -238,6 +238,8 @@ BEGIN_UNIFORM_BUFFER_STRUCT(FForwardLightData,ENGINE_API)
 	DECLARE_UNIFORM_BUFFER_STRUCT_MEMBER(float,InvTileSize)
 	DECLARE_UNIFORM_BUFFER_STRUCT_MEMBER_ARRAY(FVector4,LightPositionAndInvRadius,[GMaxNumForwardLights])
 	DECLARE_UNIFORM_BUFFER_STRUCT_MEMBER_ARRAY(FVector4,LightColorAndFalloffExponent,[GMaxNumForwardLights])
+	DECLARE_UNIFORM_BUFFER_STRUCT_MEMBER_ARRAY(FVector4,LightDirectionAndSpotlightMaskAndMinRoughness,[GMaxNumForwardLights])
+	DECLARE_UNIFORM_BUFFER_STRUCT_MEMBER_ARRAY(FVector4,SpotAnglesAndSourceRadiusAndSimpleLighting,[GMaxNumForwardLights])
 END_UNIFORM_BUFFER_STRUCT(FForwardLightData)
 
 //////////////////////////////////////////////////////////////////////////
diff --git a/Engine/Source/Runtime/Renderer/Private/SceneRendering.cpp b/Engine/Source/Runtime/Renderer/Private/SceneRendering.cpp
--- a/Engine/Source/Runtime/Renderer/Private/SceneRendering.cpp
+++ b/Engine/Source/Runtime/Renderer/Private/SceneRendering.cpp
@@ -701,8 +701,12 @@ void FViewInfo::CreateForwardLightDataUniformBuffer(FForwardLightData &OutForwar
 
 	FScene* Scene = (FScene*)(Family->Scene);
 
-	if(Scene)
+	// Reflection override skips direct specular because it tends to be blindingly bright with a perfectly smooth surface
+	if(Scene && !Family->EngineShowFlags.ReflectionOverride)
 	{
+		// we test after adding each light so we need at least space for one element
+		check(GMaxNumForwardLights > 0);
+
 		// Build a list of visible lights.
 		for (TSparseArray<FLightSceneInfoCompact>::TConstIterator LightIt(Scene->Lights); LightIt; ++LightIt)
 		{
@@ -711,28 +715,80 @@ void FViewInfo::CreateForwardLightDataUniformBuffer(FForwardLightData &OutForwar
 
 			if(!LightSceneInfoCompact.Color.IsAlmostBlack()
 				// Only render lights with dynamic lighting or unbuilt static lights
-				&& (!LightSceneInfo->Proxy->HasStaticLighting() || !LightSceneInfo->bPrecomputedLightingIsValid)
-				// Reflection override skips direct specular because it tends to be blindingly bright with a perfectly smooth surface
-				&& !Family->EngineShowFlags.ReflectionOverride
-				&& LightSceneInfoCompact.LightType == LightType_Point)		// for now we only handle point lights
+				&& (!LightSceneInfo->Proxy->HasStaticLighting() || !LightSceneInfo->bPrecomputedLightingIsValid))
 			{
-				// Check if the light is visible in this view.
-				if(LightSceneInfo->ShouldRenderLight(*this))
+				const ELightComponentType LightType = (const ELightComponentType)LightSceneInfoCompact.LightType;
+
+				FVector NormalizedLightDirection;
+				FVector2D SpotAngles;
+				float SourceRadius;
+				float SourceLength;
+				float MinRoughness;
+
+				// Get the light parameters
+				LightSceneInfo->Proxy->GetParameters(
+					OutForwardLightData.LightPositionAndInvRadius[LightIndex],
+					OutForwardLightData.LightColorAndFalloffExponent[LightIndex],
+					NormalizedLightDirection,
+					SpotAngles,
+					SourceRadius,
+					SourceLength,
+					MinRoughness);
+
+				if(LightType == LightType_Point || LightType == LightType_Spot)
 				{
+					// Check if the light is visible in this view.
+					if(!LightSceneInfo->ShouldRenderLight(*this))
+					{
+						continue;
+					}
+
 					FVector4 BoundingSphereVector = *(FVector4*)&LightSceneInfoCompact.BoundingSphereVector;
 
 					float InvRadius = 1.0f / BoundingSphereVector.W;
 
 					OutForwardLightData.LightPositionAndInvRadius[LightIndex] = FVector4(FVector(BoundingSphereVector), InvRadius);
-					OutForwardLightData.LightColorAndFalloffExponent[LightIndex] = LightSceneInfoCompact.Color;
-					++LightIndex;
 
-					if(LightIndex >= GMaxNumForwardLights)
+					// SpotlightMaskAndMinRoughness, >0:Spotlight, MinRoughness = abs();
+					float W = FMath::Max(0.0001f, MinRoughness) * ((LightType == LightType_Spot) ? 1 : -1);
+
+					OutForwardLightData.LightDirectionAndSpotlightMaskAndMinRoughness[LightIndex] = FVector4(NormalizedLightDirection, W);
+					OutForwardLightData.SpotAnglesAndSourceRadiusAndSimpleLighting[LightIndex] = FVector4(SpotAngles.X, SpotAngles.Y, SourceRadius, 0);
+
+					if (LightSceneInfo->Proxy->IsInverseSquared())
 					{
-						// we cannot handle more lights this way
-						break;
+						// Correction for lumen units
+						OutForwardLightData.LightColorAndFalloffExponent[LightIndex].X *= 16.0f;
+						OutForwardLightData.LightColorAndFalloffExponent[LightIndex].Y *= 16.0f;
+						OutForwardLightData.LightColorAndFalloffExponent[LightIndex].Z *= 16.0f;
+						OutForwardLightData.LightColorAndFalloffExponent[LightIndex].W = 0;
 					}
 				}
+				else
+				{
+//					FVector Direction = -LightSceneInfo->Proxy->GetDirection();
+
+					// we don't handle directional light types yet
+					continue;
+				}
+
+				{
+					// SpotlightMaskAndMinRoughness, >0:Spotlight, MinRoughness = abs();
+					float W = FMath::Max(0.0001f, MinRoughness) * ((LightSceneInfo->Proxy->GetLightType() == LightType_Spot) ? 1 : -1);
+
+					OutForwardLightData.LightDirectionAndSpotlightMaskAndMinRoughness[LightIndex] = FVector4(NormalizedLightDirection, W);
+				}
+
+				OutForwardLightData.SpotAnglesAndSourceRadiusAndSimpleLighting[LightIndex] = FVector4(SpotAngles.X, SpotAngles.Y, SourceRadius, 0);
+
+				// we want to add one light
+				++LightIndex;
+
+				if(LightIndex >= GMaxNumForwardLights)
+				{
+					// we cannot handle more lights this way
+					break;
+				}
 			}
 		}
 	}
