Integrate - Sky lighting is occluded by min(SSAO, DFAO, Material AO) instead of multiplying them together which was causing overdarkening

diff --git a/Engine/Shaders/SkyLighting.usf b/Engine/Shaders/SkyLighting.usf
--- a/Engine/Shaders/SkyLighting.usf
+++ b/Engine/Shaders/SkyLighting.usf
@@ -90,8 +90,9 @@ void SkyLightDiffusePS(
 		#endif
 #endif
 
-		float AmbientOcclusion = GBuffer.GBufferAO * ScreenSpaceData.AmbientOcclusion;
-		float ScalarFactors = lerp(SkyVisibility, 1, OcclusionTintAndMinOcclusion.w) * AmbientOcclusion;
+		SkyVisibility = min(SkyVisibility, min(GBuffer.GBufferAO, ScreenSpaceData.AmbientOcclusion));
+		
+		float ScalarFactors = lerp(SkyVisibility, 1, OcclusionTintAndMinOcclusion.w);
 
 		if (ShadingModelId == SHADINGMODELID_TWOSIDED_FOLIAGE)
 		{
@@ -106,7 +107,7 @@ void SkyLightDiffusePS(
 		// Apply AO to the sky diffuse and account for darkening due to the geometry term
 		Lighting += (ScalarFactors * DotProductFactor) * DiffuseLookup * GBuffer.DiffuseColor + (1 - SkyVisibility) * OcclusionTintAndMinOcclusion.xyz;
 
-		Lighting += DiffuseIrradiance * GBuffer.DiffuseColor * AmbientOcclusion;
+		Lighting += DiffuseIrradiance * GBuffer.DiffuseColor * (GBuffer.GBufferAO * ScreenSpaceData.AmbientOcclusion);
 
 		// For debugging
 		#if SUPPORT_IRRADIANCE
