UE4 - merge GDC branch, code [REDACTED_HANDLE] to main

diff --git a/Engine/Shaders/BasePassPixelShader.usf b/Engine/Shaders/BasePassPixelShader.usf
--- a/Engine/Shaders/BasePassPixelShader.usf
+++ b/Engine/Shaders/BasePassPixelShader.usf
@@ -89,12 +89,14 @@ half3 GetImageBasedReflectionLighting(FMaterialPixelParameters MaterialParameter
 	float4 SpecularIBL = TextureCubeSampleLevel(ReflectionCubemap, ReflectionCubemapSampler, MaterialParameters.ReflectionVector, AbsoluteSpecularMip);
 #endif
 
+#if ALLOW_STATIC_LIGHTING
 	FLATTEN
 	if (View.UseLightmaps > 0)
 	{
 		// Note: make sure this matches the lightmap mixing done on opaque (ReflectionEnvironmentTiledDeferredMain)
 		SpecularIBL.rgb *= IndirectIrradiance;
 	}
+#endif
 
 #if ENABLE_SKY_LIGHT
 	BRANCH
@@ -116,27 +118,30 @@ half3 GetImageBasedReflectionLighting(FMaterialPixelParameters MaterialParameter
 #endif
 
 #if MATERIAL_SSR
-	//uint FrameRandom = (uint)(View.TemporalAAParams.r * 1551);
-	uint FrameRandom = (View.FrameNumber & 15) * 1551;
+	if( View.CameraCut == 0 )
+	{
+		//uint FrameRandom = (uint)(View.TemporalAAParams.r * 1551);
+		uint FrameRandom = (View.FrameNumber & 15) * 1551;
 
-	uint Morton = MortonCode( (uint)MaterialParameters.SVPosition.x & 3 ) | ( MortonCode( (uint)MaterialParameters.SVPosition.y & 3 ) * 2 );
-	uint PixelIndex = reversebits( Morton ) >> 28;
-	//uint PixelIndex = ( (uint)MaterialParameters.SVPosition.x & 3 ) | ( ( (uint)MaterialParameters.SVPosition.y & 3 ) * 2 );
-	//PixelIndex = ( PixelIndex * 1551 ) & 15;
+		uint Morton = MortonCode( (uint)MaterialParameters.SVPosition.x & 3 ) | ( MortonCode( (uint)MaterialParameters.SVPosition.y & 3 ) * 2 );
+		uint PixelIndex = reversebits( Morton ) >> 28;
+		//uint PixelIndex = ( (uint)MaterialParameters.SVPosition.x & 3 ) | ( ( (uint)MaterialParameters.SVPosition.y & 3 ) * 2 );
+		//PixelIndex = ( PixelIndex * 1551 ) & 15;
 
-	uint Offset = ( PixelIndex + FrameRandom ) & 15;
-	float StepOffset = Offset / 15.0;
-	StepOffset -= 0.5;
+		uint Offset = ( PixelIndex + FrameRandom ) & 15;
+		float StepOffset = Offset / 15.0;
+		StepOffset -= 0.5;
 		
-	float4 HitUVzTime = RayCastCheap( HZBTexture, HZBSampler, MaterialParameters.ReflectionVector, MaterialParameters.ScreenPosition.w, MaterialParameters.WorldPosition_CamRelative, StepOffset );
-	//float4 HitUVzTime = RayCast( HZBTexture, HZBSampler, MaterialParameters.ReflectionVector, 0, MaterialParameters.ScreenPosition.w, MaterialParameters.WorldPosition_CamRelative, 12, StepOffset );
+		//float4 HitUVzTime = RayCastCheap( HZBTexture, HZBSampler, MaterialParameters.ReflectionVector, MaterialParameters.ScreenPosition.w, MaterialParameters.WorldPosition_CamRelative, StepOffset );
+		float4 HitUVzTime = RayCast( HZBTexture, HZBSampler, MaterialParameters.ReflectionVector, 0, MaterialParameters.ScreenPosition.w, MaterialParameters.WorldPosition_CamRelative, 12, StepOffset );
 
-	// if there was a hit
-	BRANCH if( HitUVzTime.w < 1 )
-	{
-		float4 SSR = SampleScreenColor( PrevSceneColor, PrevSceneColorSampler, HitUVzTime.xyz );
-		SSR *= saturate( 2 - 6.6 * Roughness );
-		SpecularIBL.rgb = SpecularIBL.rgb * (1 - SSR.a) + SSR.rgb;
+		// if there was a hit
+		BRANCH if( HitUVzTime.w < 1 )
+		{
+			float4 SSR = SampleScreenColor( PrevSceneColor, PrevSceneColorSampler, HitUVzTime.xyz );
+			SSR *= saturate( 2 - 6.6 * Roughness );
+			SpecularIBL.rgb = SpecularIBL.rgb * (1 - SSR.a) + SSR.rgb;
+		}
 	}
 #endif
 
diff --git a/Engine/Shaders/Common.usf b/Engine/Shaders/Common.usf
--- a/Engine/Shaders/Common.usf
+++ b/Engine/Shaders/Common.usf
@@ -820,7 +820,7 @@ MaterialFloat2 GetViewportCoordinates(MaterialFloat2 InFragmentCoordinates)
  */
 MaterialFloat4 UnpackNormalMap( MaterialFloat4 TextureSample )
 {
-#if COMPILER_GLSL_ES2 && IOS
+#if COMPILER_GLSL_ES2
 	return MaterialFloat4(TextureSample.rgb * 2 - 1, 1);
 #else
 	#if DXT5_NORMALMAPS
diff --git a/Engine/Shaders/DeferredLightingCommon.usf b/Engine/Shaders/DeferredLightingCommon.usf
--- a/Engine/Shaders/DeferredLightingCommon.usf
+++ b/Engine/Shaders/DeferredLightingCommon.usf
@@ -248,25 +248,17 @@ float3 SubsurfaceShadingSubsurface( FGBufferData GBuffer, float3 L, float3 V, ha
 	return SubsurfaceColor * lerp(BackScatter, 1, InScatter);
 }
 
-float3 TwoSidedFoliageShadingSubsurface( float3 SubsurfaceColor, float3 L, float3 V, half3 N )
+float3 SubsurfaceShadingTwoSided( float3 SubsurfaceColor, float3 L, float3 V, half3 N )
 {
-	float VoL = saturate(dot(V, -L));
-	// here we raise to the power of 4
-	float VoLLobe = VoL * VoL;
-	VoLLobe *= VoLLobe;
-	// so for energy conservation we need to do the same here
-	const float Exponent = 4;
-	// Modified phong energy conservation
-	VoLLobe = VoLLobe * (Exponent + 1) / (2.0f * PI);
-
-	float DiffuseWrap = .6f;
-	// Energy conserving wrap
-	float BackDiffuse = saturate(dot(-N, L) * DiffuseWrap * DiffuseWrap + (1 - DiffuseWrap) * DiffuseWrap) / PI;  
-	// Leaves have very small amounts of directional backscattering
-	float ViewDependenceAmount = .2f;
-	float BackShading = lerp(BackDiffuse, VoLLobe, ViewDependenceAmount);
-
-	return BackShading * SubsurfaceColor;
+	float Wrap = 0.5;
+	float NoL = saturate( ( dot(-N, L) + Wrap ) / Square( 1 + Wrap ) );
+
+	float VoL = saturate( dot(V, -L) );
+	float a = 0.6;
+	float a2 = a * a;
+	float d = ( VoL * a2 - VoL ) * VoL + 1;	// 2 mad
+	float GGX = (a2 / PI) / (d * d);		// 2 mul, 1 rcp
+	return NoL * GGX * SubsurfaceColor;
 }
 
 Texture2D		PreIntegratedBRDF;
@@ -312,7 +304,7 @@ float3 SubsurfaceShading( FGBufferData GBuffer, float3 L, float3 V, half3 N )
 		case SHADINGMODELID_PREINTEGRATED_SKIN:
 			return SubsurfaceShadingPreintegratedSkin( GBuffer, L, V, N );
 		case SHADINGMODELID_TWOSIDED_FOLIAGE:
-			return TwoSidedFoliageShadingSubsurface( SubsurfaceColor, L, V, N );
+			return SubsurfaceShadingTwoSided( SubsurfaceColor, L, V, N );
 		default:
 			return 0;
 	}
diff --git a/Engine/Shaders/DeferredShadingCommon.usf b/Engine/Shaders/DeferredShadingCommon.usf
--- a/Engine/Shaders/DeferredShadingCommon.usf
+++ b/Engine/Shaders/DeferredShadingCommon.usf
@@ -348,8 +348,12 @@ void EncodeGBuffer(
 		OutGBufferB.b = GBuffer.Roughness;
 		OutGBufferC.rgb = EncodeBaseColor(GBuffer.BaseColor);
 
+#if ALLOW_STATIC_LIGHTING
 		// No space for AO. Multiply IndirectIrradiance by AO instead of storing.
 		OutGBufferC.a = EncodeIndirectIrradiance(GBuffer.IndirectIrradiance * GBuffer.GBufferAO) + QuantizationBias * (1.0 / 255.0);
+#else
+		OutGBufferC.a = GBuffer.GBufferAO;
+#endif
 		OutGBufferB.a = Encode71(EncodeShadingModelId(GBuffer.ShadingModelID), GBuffer.HasHeightfieldRepresentation);
 
 		OutGBufferD.r = GBuffer.Opacity;
@@ -399,8 +403,13 @@ FGBufferData DecodeGBufferData(
 	
 	GBuffer.Roughness = InGBufferB.b;
 	GBuffer.CustomData = InGBufferD.gba;
+#if ALLOW_STATIC_LIGHTING
 	GBuffer.GBufferAO = 1;
 	GBuffer.IndirectIrradiance = DecodeIndirectIrradiance(InGBufferC.a);
+#else
+	GBuffer.GBufferAO = InGBufferC.a;
+	GBuffer.IndirectIrradiance = 1;
+#endif
 	GBuffer.PrecomputedShadowFactors = InGBufferE;
 	GBuffer.CustomDepth = ConvertFromDeviceZ(CustomNativeDepth);
 	GBuffer.Depth = SceneDepth;
diff --git a/Engine/Shaders/LandscapeVertexFactory.usf b/Engine/Shaders/LandscapeVertexFactory.usf
--- a/Engine/Shaders/LandscapeVertexFactory.usf
+++ b/Engine/Shaders/LandscapeVertexFactory.usf
@@ -158,7 +158,10 @@ float4 VertexFactoryGetWorldPosition(FVertexFactoryInput Input, FVertexFactoryIn
 
 float4 VertexFactoryGetPreviousWorldPosition(FVertexFactoryInput Input, FVertexFactoryIntermediates Intermediates)
 {
-	return TransformLocalToPrevTranslatedWorld(GetLocalPosition(Intermediates));
+	// Note we still use LocalToWorld. Landscape's transform never changes.
+	float3 LocalPosition = GetLocalPosition(Intermediates);
+	float3 RotatedPosition = Primitive.LocalToWorld[0].xyz * LocalPosition.xxx + Primitive.LocalToWorld[1].xyz * LocalPosition.yyy + Primitive.LocalToWorld[2].xyz * LocalPosition.zzz;
+	return float4(RotatedPosition + (Primitive.LocalToWorld[3].xyz + View.PrevPreViewTranslation.xyz),1);
 }
 
 /** Calculate the texture coordinates generated by Landscape */
diff --git a/Engine/Shaders/MaterialTemplate.usf b/Engine/Shaders/MaterialTemplate.usf
--- a/Engine/Shaders/MaterialTemplate.usf
+++ b/Engine/Shaders/MaterialTemplate.usf
@@ -790,7 +790,7 @@ float3 GetSpeedTreeVertexOffsetInner(FMaterialVertexParameters Parameters, int G
 		// leaves and facing leaves
 		if (GeometryType == SPEEDTREE_GEOMETRY_TYPE_FACINGLEAF || 
 				(GeometryType == SPEEDTREE_GEOMETRY_TYPE_LEAF && 
-				(LODType == SPEEDTREE_LOD_TYPE_SMOOTH || (WindType > SPEEDTREE_WIND_TYPE_FASTEST && WindType < SPEEDTREE_WIND_TYPE_PALM))))
+				(LODType == SPEEDTREE_LOD_TYPE_SMOOTH || (WindType > SPEEDTREE_WIND_TYPE_FASTEST && WindType != SPEEDTREE_WIND_TYPE_PALM))))
 		{
 			// remove anchor pos from vertex position
 			float3 Anchor = float3(Parameters.TexCoords[4].y, Parameters.TexCoords[5].x, Parameters.TexCoords[5].y);
@@ -825,7 +825,7 @@ float3 GetSpeedTreeVertexOffsetInner(FMaterialVertexParameters Parameters, int G
 			}
 
 			// leaf wind
-			if (WindType > SPEEDTREE_WIND_TYPE_FASTEST && WindType < SPEEDTREE_WIND_TYPE_PALM) 
+			if (WindType > SPEEDTREE_WIND_TYPE_FASTEST && WindType != SPEEDTREE_WIND_TYPE_PALM) 
 			{
 				float4 WindExtra = float4(Parameters.TexCoords[6].x, Parameters.TexCoords[6].y, Parameters.TexCoords[7].x, Parameters.TexCoords[7].y);
 				float LeafWindTrigOffset = Anchor.x + Anchor.y;
diff --git a/Engine/Shaders/PostProcessMotionBlur.usf b/Engine/Shaders/PostProcessMotionBlur.usf
--- a/Engine/Shaders/PostProcessMotionBlur.usf
+++ b/Engine/Shaders/PostProcessMotionBlur.usf
@@ -785,7 +785,7 @@ void MainNewPS( float4 UVAndScreenPos : TEXCOORD0, out float4 OutColor : SV_Targ
 
 	float VelocityScale = MotionBlurParameters.x;
 
-	const uint StepCount = GetStepCountFromQuality();
+	const uint StepCount = GetStepCountFromQuality() / 2;
 
 #if 1
 	uint2 PixelPos = uint2(UVAndScreenPos.zw * ScreenPosToPixel.xy + ScreenPosToPixel.zw + 0.5f);
diff --git a/Engine/Shaders/PostProcessTemporalCommon.usf b/Engine/Shaders/PostProcessTemporalCommon.usf
--- a/Engine/Shaders/PostProcessTemporalCommon.usf
+++ b/Engine/Shaders/PostProcessTemporalCommon.usf
@@ -125,6 +125,10 @@ which can be configured via defines for various temporal AA passes.
 	#define AA_DOF 0
 #endif
 
+// Completely non-optimized prototype for solving noise problem.
+#ifndef AA_TEST
+	#define AA_TEST 0
+#endif
 
 
 	// FIND MOTION OF PIXEL AND NEAREST IN NEIGHBORHOOD
@@ -216,6 +220,8 @@ which can be configured via defines for various temporal AA passes.
 	// as the upper left of the pixel instead of the center of the pixel.
 	BackN = (BackN * ScreenPosToPixel.xy + ScreenPosToPixel.zw + 0.5) * PostprocessInput0Size.zw;
 
+
+
 	// FILTER PIXEL (RESAMPLE TO REMOVE JITTER OFFSET) AND GET NEIGHBORHOOD
     // --------------------------------------------------------------------
 	// 012
@@ -365,6 +371,190 @@ which can be configured via defines for various temporal AA passes.
 		#endif
 	#endif
 
+
+	// Prototype low frequency test.
+	// This has to be re-done to do the reduction/etc prior at non-full res/etc.
+	#if AA_TEST
+
+		// Get low resolution version of current frame.
+
+		//   abc
+		//  defgh
+		//  ijklm
+		//  nopqr
+		//   stu
+		float4 TxA = PostprocessInput0.SampleLevel(PostprocessInput0Sampler, UVAndScreenPos.xy, 0, int2( -1, -2));
+		float4 TxB = PostprocessInput0.SampleLevel(PostprocessInput0Sampler, UVAndScreenPos.xy, 0, int2(  0, -2));
+		float4 TxC = PostprocessInput0.SampleLevel(PostprocessInput0Sampler, UVAndScreenPos.xy, 0, int2(  1, -2));
+
+		float4 TxD = PostprocessInput0.SampleLevel(PostprocessInput0Sampler, UVAndScreenPos.xy, 0, int2( -2, -1));
+		float4 TxE = PostprocessInput0.SampleLevel(PostprocessInput0Sampler, UVAndScreenPos.xy, 0, int2( -1, -1));
+		float4 TxF = PostprocessInput0.SampleLevel(PostprocessInput0Sampler, UVAndScreenPos.xy, 0, int2(  0, -1));
+		float4 TxG = PostprocessInput0.SampleLevel(PostprocessInput0Sampler, UVAndScreenPos.xy, 0, int2(  1, -1));
+		float4 TxH = PostprocessInput0.SampleLevel(PostprocessInput0Sampler, UVAndScreenPos.xy, 0, int2(  2, -1));
+
+		float4 TxI = PostprocessInput0.SampleLevel(PostprocessInput0Sampler, UVAndScreenPos.xy, 0, int2( -2, 0));
+		float4 TxJ = PostprocessInput0.SampleLevel(PostprocessInput0Sampler, UVAndScreenPos.xy, 0, int2( -1, 0));
+		float4 TxK = PostprocessInput0.SampleLevel(PostprocessInput0Sampler, UVAndScreenPos.xy, 0, int2(  0, 0));
+		float4 TxL = PostprocessInput0.SampleLevel(PostprocessInput0Sampler, UVAndScreenPos.xy, 0, int2(  1, 0));
+		float4 TxM = PostprocessInput0.SampleLevel(PostprocessInput0Sampler, UVAndScreenPos.xy, 0, int2(  2, 0));
+
+		float4 TxN = PostprocessInput0.SampleLevel(PostprocessInput0Sampler, UVAndScreenPos.xy, 0, int2( -2, 1));
+		float4 TxO = PostprocessInput0.SampleLevel(PostprocessInput0Sampler, UVAndScreenPos.xy, 0, int2( -1, 1));
+		float4 TxP = PostprocessInput0.SampleLevel(PostprocessInput0Sampler, UVAndScreenPos.xy, 0, int2(  0, 1));
+		float4 TxQ = PostprocessInput0.SampleLevel(PostprocessInput0Sampler, UVAndScreenPos.xy, 0, int2(  1, 1));
+		float4 TxR = PostprocessInput0.SampleLevel(PostprocessInput0Sampler, UVAndScreenPos.xy, 0, int2(  2, 1));
+
+		float4 TxS = PostprocessInput0.SampleLevel(PostprocessInput0Sampler, UVAndScreenPos.xy, 0, int2( -1, 2));
+		float4 TxT = PostprocessInput0.SampleLevel(PostprocessInput0Sampler, UVAndScreenPos.xy, 0, int2(  0, 2));
+		float4 TxU = PostprocessInput0.SampleLevel(PostprocessInput0Sampler, UVAndScreenPos.xy, 0, int2(  1, 2));
+
+		float4 TxAvg = (TxA+TxB+TxC+ TxD+TxE+TxF+TxG+TxH+ TxI+TxJ+TxK+TxL+TxM+ TxN+TxO+TxP+TxQ+TxR+ TxS+TxT+TxU) * (1.0/21.0);
+
+		TxAvg.rgb *= HdrWeightG(TxAvg.rgb, InExposureScale);
+		TxAvg.rgb = RGBToYCoCg(TxAvg.rgb);
+
+		// Get low resolution version of reprojected frame.
+		float4 TxA2 = PostprocessInput1.SampleLevel(PostprocessInput1Sampler, BackN.xy, 0, int2( -1, -2));
+		float4 TxB2 = PostprocessInput1.SampleLevel(PostprocessInput1Sampler, BackN.xy, 0, int2(  0, -2));
+		float4 TxC2 = PostprocessInput1.SampleLevel(PostprocessInput1Sampler, BackN.xy, 0, int2(  1, -2));
+
+		float4 TxD2 = PostprocessInput1.SampleLevel(PostprocessInput1Sampler, BackN.xy, 0, int2( -2, -1));
+		float4 TxE2 = PostprocessInput1.SampleLevel(PostprocessInput1Sampler, BackN.xy, 0, int2( -1, -1));
+		float4 TxF2 = PostprocessInput1.SampleLevel(PostprocessInput1Sampler, BackN.xy, 0, int2(  0, -1));
+		float4 TxG2 = PostprocessInput1.SampleLevel(PostprocessInput1Sampler, BackN.xy, 0, int2(  1, -1));
+		float4 TxH2 = PostprocessInput1.SampleLevel(PostprocessInput1Sampler, BackN.xy, 0, int2(  2, -1));
+
+		float4 TxI2 = PostprocessInput1.SampleLevel(PostprocessInput1Sampler, BackN.xy, 0, int2( -2, 0));
+		float4 TxJ2 = PostprocessInput1.SampleLevel(PostprocessInput1Sampler, BackN.xy, 0, int2( -1, 0));
+		float4 TxK2 = PostprocessInput1.SampleLevel(PostprocessInput1Sampler, BackN.xy, 0, int2(  0, 0));
+		float4 TxL2 = PostprocessInput1.SampleLevel(PostprocessInput1Sampler, BackN.xy, 0, int2(  1, 0));
+		float4 TxM2 = PostprocessInput1.SampleLevel(PostprocessInput1Sampler, BackN.xy, 0, int2(  2, 0));
+
+		float4 TxN2 = PostprocessInput1.SampleLevel(PostprocessInput1Sampler, BackN.xy, 0, int2( -2, 1));
+		float4 TxO2 = PostprocessInput1.SampleLevel(PostprocessInput1Sampler, BackN.xy, 0, int2( -1, 1));
+		float4 TxP2 = PostprocessInput1.SampleLevel(PostprocessInput1Sampler, BackN.xy, 0, int2(  0, 1));
+		float4 TxQ2 = PostprocessInput1.SampleLevel(PostprocessInput1Sampler, BackN.xy, 0, int2(  1, 1));
+		float4 TxR2 = PostprocessInput1.SampleLevel(PostprocessInput1Sampler, BackN.xy, 0, int2(  2, 1));
+
+		float4 TxS2 = PostprocessInput1.SampleLevel(PostprocessInput1Sampler, BackN.xy, 0, int2( -1, 2));
+		float4 TxT2 = PostprocessInput1.SampleLevel(PostprocessInput1Sampler, BackN.xy, 0, int2(  0, 2));
+		float4 TxU2 = PostprocessInput1.SampleLevel(PostprocessInput1Sampler, BackN.xy, 0, int2(  1, 2));
+
+		float4 TxAvg2 = (TxA2+TxB2+TxC2+ TxD2+TxE2+TxF2+TxG2+TxH2+ TxI2+TxJ2+TxK2+TxL2+TxM2+ TxN2+TxO2+TxP2+TxQ2+TxR2+ TxS2+TxT2+TxU2) * (1.0/21.0);
+
+		TxAvg2.rgb *= HdrWeightG(TxAvg2.rgb, InExposureScale);
+		TxAvg2.rgb = RGBToYCoCg(TxAvg2.rgb);
+
+
+
+
+		TxA.rgb = RGBToYCoCg(TxA.rgb * HdrWeightG(TxA.rgb, InExposureScale));
+		TxB.rgb = RGBToYCoCg(TxB.rgb * HdrWeightG(TxB.rgb, InExposureScale));
+		TxC.rgb = RGBToYCoCg(TxC.rgb * HdrWeightG(TxC.rgb, InExposureScale));
+
+		TxD.rgb = RGBToYCoCg(TxD.rgb * HdrWeightG(TxD.rgb, InExposureScale));
+		TxE.rgb = RGBToYCoCg(TxE.rgb * HdrWeightG(TxE.rgb, InExposureScale));
+		TxF.rgb = RGBToYCoCg(TxF.rgb * HdrWeightG(TxF.rgb, InExposureScale));
+		TxG.rgb = RGBToYCoCg(TxG.rgb * HdrWeightG(TxG.rgb, InExposureScale));
+		TxH.rgb = RGBToYCoCg(TxH.rgb * HdrWeightG(TxH.rgb, InExposureScale));
+
+		TxI.rgb = RGBToYCoCg(TxI.rgb * HdrWeightG(TxI.rgb, InExposureScale));
+		TxJ.rgb = RGBToYCoCg(TxJ.rgb * HdrWeightG(TxJ.rgb, InExposureScale));
+		TxK.rgb = RGBToYCoCg(TxK.rgb * HdrWeightG(TxK.rgb, InExposureScale));
+		TxL.rgb = RGBToYCoCg(TxL.rgb * HdrWeightG(TxL.rgb, InExposureScale));
+		TxM.rgb = RGBToYCoCg(TxM.rgb * HdrWeightG(TxM.rgb, InExposureScale));
+
+		TxN.rgb = RGBToYCoCg(TxN.rgb * HdrWeightG(TxN.rgb, InExposureScale));
+		TxO.rgb = RGBToYCoCg(TxO.rgb * HdrWeightG(TxO.rgb, InExposureScale));
+		TxP.rgb = RGBToYCoCg(TxP.rgb * HdrWeightG(TxP.rgb, InExposureScale));
+		TxQ.rgb = RGBToYCoCg(TxQ.rgb * HdrWeightG(TxQ.rgb, InExposureScale));
+		TxR.rgb = RGBToYCoCg(TxR.rgb * HdrWeightG(TxR.rgb, InExposureScale));
+
+		TxS.rgb = RGBToYCoCg(TxS.rgb * HdrWeightG(TxS.rgb, InExposureScale));
+		TxT.rgb = RGBToYCoCg(TxT.rgb * HdrWeightG(TxT.rgb, InExposureScale));
+		TxU.rgb = RGBToYCoCg(TxU.rgb * HdrWeightG(TxU.rgb, InExposureScale));
+
+
+		TxA2.rgb = RGBToYCoCg(TxA2.rgb * HdrWeightG(TxA2.rgb, InExposureScale));
+		TxB2.rgb = RGBToYCoCg(TxB2.rgb * HdrWeightG(TxB2.rgb, InExposureScale));
+		TxC2.rgb = RGBToYCoCg(TxC2.rgb * HdrWeightG(TxC2.rgb, InExposureScale));
+
+		TxD2.rgb = RGBToYCoCg(TxD2.rgb * HdrWeightG(TxD2.rgb, InExposureScale));
+		TxE2.rgb = RGBToYCoCg(TxE2.rgb * HdrWeightG(TxE2.rgb, InExposureScale));
+		TxF2.rgb = RGBToYCoCg(TxF2.rgb * HdrWeightG(TxF2.rgb, InExposureScale));
+		TxG2.rgb = RGBToYCoCg(TxG2.rgb * HdrWeightG(TxG2.rgb, InExposureScale));
+		TxH2.rgb = RGBToYCoCg(TxH2.rgb * HdrWeightG(TxH2.rgb, InExposureScale));
+
+		TxI2.rgb = RGBToYCoCg(TxI2.rgb * HdrWeightG(TxI2.rgb, InExposureScale));
+		TxJ2.rgb = RGBToYCoCg(TxJ2.rgb * HdrWeightG(TxJ2.rgb, InExposureScale));
+		TxK2.rgb = RGBToYCoCg(TxK2.rgb * HdrWeightG(TxK2.rgb, InExposureScale));
+		TxL2.rgb = RGBToYCoCg(TxL2.rgb * HdrWeightG(TxL2.rgb, InExposureScale));
+		TxM2.rgb = RGBToYCoCg(TxM2.rgb * HdrWeightG(TxM2.rgb, InExposureScale));
+
+		TxN2.rgb = RGBToYCoCg(TxN2.rgb * HdrWeightG(TxN2.rgb, InExposureScale));
+		TxO2.rgb = RGBToYCoCg(TxO2.rgb * HdrWeightG(TxO2.rgb, InExposureScale));
+		TxP2.rgb = RGBToYCoCg(TxP2.rgb * HdrWeightG(TxP2.rgb, InExposureScale));
+		TxQ2.rgb = RGBToYCoCg(TxQ2.rgb * HdrWeightG(TxQ2.rgb, InExposureScale));
+		TxR2.rgb = RGBToYCoCg(TxR2.rgb * HdrWeightG(TxR2.rgb, InExposureScale));
+
+		TxS2.rgb = RGBToYCoCg(TxS2.rgb * HdrWeightG(TxS2.rgb, InExposureScale));
+		TxT2.rgb = RGBToYCoCg(TxT2.rgb * HdrWeightG(TxT2.rgb, InExposureScale));
+		TxU2.rgb = RGBToYCoCg(TxU2.rgb * HdrWeightG(TxU2.rgb, InExposureScale));
+
+
+		TxA = abs(TxA - TxA2);
+		TxB = abs(TxB - TxB2);
+		TxC = abs(TxC - TxC2);
+
+		TxD = abs(TxD - TxD2);
+		TxE = abs(TxE - TxE2);
+		TxF = abs(TxF - TxF2);
+		TxG = abs(TxG - TxG2);
+		TxH = abs(TxH - TxH2);
+
+		TxI = abs(TxI - TxI2);
+		TxJ = abs(TxJ - TxJ2);
+		TxK = abs(TxK - TxK2);
+		TxL = abs(TxL - TxL2);
+		TxM = abs(TxM - TxM2);
+
+		TxN = abs(TxN - TxN2);
+		TxO = abs(TxO - TxO2);
+		TxP = abs(TxP - TxP2);
+		TxQ = abs(TxQ - TxQ2);
+		TxR = abs(TxR - TxR2);
+
+		TxS = abs(TxS - TxS2);
+		TxT = abs(TxT - TxT2);
+		TxU = abs(TxU - TxU2);
+
+		float4 TxDifAvg = (TxA+TxB+TxC+ TxD+TxE+TxF+TxG+TxH+ TxI+TxJ+TxK+TxL+TxM+ TxN+TxO+TxP+TxQ+TxR+ TxS+TxT+TxU) * (1.0/21.0);
+
+#if 0
+		// Filtered does not work as good (probably because of lower footprint).
+		float Wx1 = 1.0;
+		float Wx2 = 2.0;
+		float Wx3 = 3.0;
+		float Wx4 = 4.0;
+		float Wx5 = 5.0;
+
+		//   abc
+		//  defgh
+		//  ijklm
+		//  nopqr
+		//   stu
+
+		float4 TxDifAvg = (TxA+TxC+TxD+TxN+TxH+TxR+TxS+TxU)*Wx1 + (TxB+TxI+TxM+TxT)*Wx2 + (TxE+TxG+TxO+TxQ)*Wx3 + (TxF+TxJ+TxL+TxP)*Wx4 + TxK*Wx5;
+
+		TxDifAvg /= (Wx1*8.0)+(Wx2*4.0)+(Wx3*4.0)+(Wx4*4.0)+(Wx5);
+#endif
+
+	#endif
+
+
+
+
 	// FETCH HISTORY
 	// -------------
 	OutColor = PostprocessInput1.SampleLevel(PostprocessInput1Sampler, BackN.xy, 0);
@@ -396,6 +586,11 @@ which can be configured via defines for various temporal AA passes.
 		float DebugDiffPrior = OutColor.b;
 	#endif
 
+
+
+
+
+
 	// FIND LUMA OF CLAMPED HISTORY
 	// ----------------------------
 	// Save off luma of history before the clamp.
@@ -499,6 +694,34 @@ which can be configured via defines for various temporal AA passes.
 		#endif
 	}
 	
+	// LIMIT BLEND AMOUNT BASED ON LOW FREQUENCY SIGNAL
+	// -------------------------------------------------
+	#if AA_TEST
+
+#if 0
+		// LOWPASS
+		float TxScale = 2.0 * max(1.0, min(2.0, Velocity));
+		float TxMaxX = abs(TxAvg.x - TxAvg2.x) * 2.0 * TxScale;
+		float TxMaxY = abs(TxAvg.y - TxAvg2.y) * 4.0 * TxScale;
+		float TxMaxZ = abs(TxAvg.z - TxAvg2.z) * 8.0 * TxScale;
+#else
+		// SUM ABS DIFF
+		float TxScale = 1.0 * max(1.0, min(2.0, Velocity));
+		float TxMaxX = TxDifAvg.x * 2.0 * TxScale;
+		float TxMaxY = TxDifAvg.y * 4.0 * TxScale;
+		float TxMaxZ = TxDifAvg.z * 8.0 * TxScale;
+
+#endif
+		
+
+		float TxMax = max(TxMaxX,max(TxMaxY,TxMaxZ));
+		TxMax = min(1.0/1.0, TxMax);
+		BlendFinal = max(BlendFinal, TxMax);
+
+	#endif
+
+
+
 	// DO FINAL BLEND BETWEEN HISTORY AND FILTERED COLOR
 	// -------------------------------------------------
 	#if (AA_LERP == 0)
@@ -563,4 +786,5 @@ which can be configured via defines for various temporal AA passes.
 #undef AA_GREEN_AS_LUMA
 #undef AA_VELOCITY_WEIGHTING
 #undef AA_YCOCG
-#undef AA_DOF
\ No newline at end of file
+#undef AA_DOF
+#undef AA_TEST
\ No newline at end of file
diff --git a/Engine/Shaders/PostProcessTonemap.usf b/Engine/Shaders/PostProcessTonemap.usf
--- a/Engine/Shaders/PostProcessTonemap.usf
+++ b/Engine/Shaders/PostProcessTonemap.usf
@@ -28,6 +28,8 @@ float4 VignetteColorIntensity;
 float3 GrainRandomFull;
 float3 GrainScaleBiasJitter;
 
+float4 OverlayColor;
+
 half GrainFromUV(float2 GrainUV) 
 {
 	half Grain = frac(sin(GrainUV.x + GrainUV.y * 543.31) *  493013.0);
@@ -225,9 +227,9 @@ void MainPS(
 		#if USE_GRAIN_JITTER
 			SceneUVJitter = SceneUV.xy - UV.xy;
 		#endif
-		half3 SceneColor = Texture2DSample(PostprocessInput0, PostprocessInput0Sampler, SceneUV).rgb;
+		half3 SceneColor = Texture2DSample(PostprocessInput0, PostprocessInput0Sampler, FringeUV.zw + SceneUVJitter.xy).rgb;
 		half SceneColorG = Texture2DSample(PostprocessInput0, PostprocessInput0Sampler, FringeUV.xy + SceneUVJitter.xy).g;
-		half SceneColorB = Texture2DSample(PostprocessInput0, PostprocessInput0Sampler, FringeUV.zw + SceneUVJitter.xy).b;
+		half SceneColorB = Texture2DSample(PostprocessInput0, PostprocessInput0Sampler, SceneUV).b;
 		SceneColor.g = SceneColorG;
 		SceneColor.b = SceneColorB;
 	#else
@@ -305,7 +307,17 @@ void MainPS(
 		//TonemappedColor = pow( TonemappedColor, 2.0 / 2.2 );
 
 		//TonemappedColor = Graph( ScreenSpacePos );
-	
+
+		#if USE_PHOTO
+			// Blend with custom LDR color, used for Fade track in Matinee.
+			// This is the 101% wrong way to do this,
+			//  - It adds an extra redundant lerp.
+			//  - It is not going to work with the future-forward ES3 fast path of sRGB output.
+			//  - And it does super ugly non-linear blending.
+			// The right way is to adjust exposure instead.
+			TonemappedColor = lerp(TonemappedColor.rgb, OverlayColor.rgb, OverlayColor.a);
+		#endif
+
 		half LuminanceForPostProcessAA  = dot(TonemappedColor, half3(0.299f, 0.587f, 0.114f));
 
 		#if USE_GRAIN_QUANTIZATION
@@ -408,8 +420,6 @@ half2 CocBlendScaleBiasFine()
 	return ScaleBias;
 }
 
-float4 OverlayColor;
-
 void MainPS_ES2(
 	in float4 UVAndScreenPos : TEXCOORD0,
 	in float4 FineDofGrain : TEXCOORD1,
diff --git a/Engine/Shaders/ReflectionEnvironmentShaders.usf b/Engine/Shaders/ReflectionEnvironmentShaders.usf
--- a/Engine/Shaders/ReflectionEnvironmentShaders.usf
+++ b/Engine/Shaders/ReflectionEnvironmentShaders.usf
@@ -251,7 +251,7 @@ void FilterPS(
 
 #endif
 
-#if NORMALIZE
+#if NORMALIZE && ALLOW_STATIC_LIGHTING
 	// Remove low frequency
 	// TODO match the frequency removed to the one added. In other words divide out 2nd order SH instead of 1st.
 	float LowFrequencyBrightness = Texture2DSample(AverageBrightnessTexture, AverageBrightnessSampler, float2(0, 0)).r;
@@ -499,6 +499,7 @@ void ReflectionApplyPS(
 		// Alpha of the reflection environment texture contains coverage, where 0 the sky reflection should be applied
 		SpecularLighting = Texture2DSample( ReflectionEnvTexture, ReflectionEnvSampler, UV );
 
+#if ALLOW_STATIC_LIGHTING
 		FLATTEN
 		if( View.UseLightmaps > 0 )
 		{
@@ -510,6 +511,7 @@ void ReflectionApplyPS(
 			SpecularLighting.rgb *= GBuffer.IndirectIrradiance;
 		}
 #endif
+#endif
 
 #if APPLY_SKYLIGHT
 
@@ -518,7 +520,7 @@ void ReflectionApplyPS(
 		{
 			float3 ReflectionVector = reflect(-V, GBuffer.WorldNormal);
 			// Normalize for static skylight types which mix with lightmaps
-			bool bNormalize = SkyLightParameters.z < 1 && View.UseLightmaps > 0;
+			bool bNormalize = SkyLightParameters.z < 1 && View.UseLightmaps > 0 && !ALLOW_STATIC_LIGHTING;
 			float3 SkyLighting = GetSkyLightReflection(ReflectionVector, GBuffer.Roughness, bNormalize);
 
 			FLATTEN
diff --git a/Engine/Shaders/ScreenSpaceRayCast.usf b/Engine/Shaders/ScreenSpaceRayCast.usf
--- a/Engine/Shaders/ScreenSpaceRayCast.usf
+++ b/Engine/Shaders/ScreenSpaceRayCast.usf
@@ -297,7 +297,7 @@ float4 RayCastCheap( Texture2D Texture, SamplerState Sampler, float3 R, float Sc
 	BRANCH
 	if( MinHitTime == 1 && RayStepUVz.z < 0 )
 	{
-		float4 RayInfClip	= mul( float4( PositionTranslatedWorld + R * 4000, 1 ), View.TranslatedWorldToClip );
+		float4 RayInfClip	= mul( float4( PositionTranslatedWorld + R * 2000000, 1 ), View.TranslatedWorldToClip );
 		float3 RayInfScreen = RayInfClip.xyz / RayInfClip.w;
 		float3 RayInfUVz = float3( RayInfScreen.xy * float2( 0.5, -0.5 ) + 0.5, RayInfScreen.z );
 
diff --git a/Engine/Shaders/SkyLighting.usf b/Engine/Shaders/SkyLighting.usf
--- a/Engine/Shaders/SkyLighting.usf
+++ b/Engine/Shaders/SkyLighting.usf
@@ -48,7 +48,8 @@ void SkyLightDiffusePS(
 
 		float a = Square( GBuffer.Roughness );
 		float3 DominantDiffuseDir = lerp( N, V, saturate( ( NoV * (1.02341 * a - 1.51174) + -0.511705 * a + 0.755868 ) * a ) );
-		SkyLightingNormal = normalize( DominantDiffuseDir );
+		DominantDiffuseDir = normalize( DominantDiffuseDir );
+		SkyLightingNormal = DominantDiffuseDir;
 
 		float3 ABC = PreIntegratedGF.SampleLevel( PreIntegratedGFSampler, float2( NoV, GBuffer.Roughness ), 0 ).rgb;
 		GBuffer.DiffuseColor *= ABC.z;
@@ -66,8 +67,8 @@ void SkyLightDiffusePS(
 			// Use more bent normal in corners
 			float BentNormalWeightFactor = SkyVisibility;
 
-			SkyLightingNormal = lerp(NormalizedBentNormal, GBuffer.WorldNormal, BentNormalWeightFactor);
-			DotProductFactor = lerp(dot(NormalizedBentNormal, GBuffer.WorldNormal), 1, BentNormalWeightFactor);
+			SkyLightingNormal = lerp(NormalizedBentNormal, DominantDiffuseDir, BentNormalWeightFactor);
+			DotProductFactor = lerp(dot(NormalizedBentNormal, DominantDiffuseDir), 1, BentNormalWeightFactor);
 		}
 		#else
 		{
diff --git a/Engine/Shaders/SpeedTreeCommon.usf b/Engine/Shaders/SpeedTreeCommon.usf
--- a/Engine/Shaders/SpeedTreeCommon.usf
+++ b/Engine/Shaders/SpeedTreeCommon.usf
@@ -96,6 +96,7 @@ FSpeedTreeData GetPreviousSpeedTreeData()
 #define SPEEDTREE_WIND_TYPE_BETTER 3
 #define SPEEDTREE_WIND_TYPE_BEST 4
 #define SPEEDTREE_WIND_TYPE_PALM 5
+#define SPEEDTREE_WIND_TYPE_BEST_PLUS 6
 
 #define SPEEDTREE_GEOMETRY_TYPE_BRANCH 0
 #define SPEEDTREE_GEOMETRY_TYPE_FROND 1
@@ -573,7 +574,7 @@ float3 LeafRipple(float3 vPos,
 
 ///////////////////////////////////////////////////////////////////////  
 //  LeafTumble
-
+      
 float3 LeafTumble(FSpeedTreeData STData,
 	float3 vPos, 
 	inout float3 vDirection, 
@@ -586,7 +587,8 @@ float3 LeafTumble(FSpeedTreeData STData,
 	float fTwist, 
 	float fAdherence, 
 	float3 vTwitch, 
-	float4 vRoll, 
+	float4 vRoll,
+	int WindType,
 	bool bTwitch, 
 	bool bRoll)
 {
@@ -603,30 +605,34 @@ float3 LeafTumble(FSpeedTreeData STData,
  	float fOsc = vOscillations.x + vOscillations.y * vOscillations.y;
  	float3x3 matTumble = RotationMatrix(vGrowthDir, fScale * fTwist * fOsc);
 
-	// with wind
-	float3 vAxis = wind_cross(vGrowthDir, STData.WindVector.xyz);
-	float fDot = clamp(dot(STData.WindVector.xyz, vGrowthDir), -1.0, 1.0);
-	#ifdef SPEEDTREE_Z_UP
+	if (WindType >= SPEEDTREE_WIND_TYPE_BEST_PLUS)
+	{
+		// with wind
+		float3 vAxis = wind_cross(vGrowthDir, STData.WindVector.xyz);
+		float fDot = clamp(dot(STData.WindVector.xyz, vGrowthDir), -1.0, 1.0);
+#ifdef SPEEDTREE_Z_UP
 		vAxis.z += fDot;
-	#else
+#else
 		vAxis.y += fDot;
-	#endif
-	vAxis = normalize(vAxis);
+#endif
+		vAxis = normalize(vAxis);
 
-	float fAngle = acos(fDot);
+		float fAngle = acos(fDot);
 
-	float fAdherenceScale = 1.0;
-	//if (bRoll)
-	//	fAdherenceScale = Roll(fAdherenceScale, vRoll.x, vRoll.y, vRoll.z, vRoll.w, vAnchor.xyz, fTime/* + fOffset*/);
+		float fAdherenceScale = 1.0;
+		//if (bRoll)
+		//	fAdherenceScale = Roll(fAdherenceScale, vRoll.x, vRoll.y, vRoll.z, vRoll.w, vAnchor.xyz, fTime/* + fOffset*/);
+
+		fOsc = vOscillations.y - vOscillations.x * vOscillations.x;
 
-	fOsc = vOscillations.y - vOscillations.x * vOscillations.x;
+		float fTwitch = 0.0;
+		if (bTwitch)
+			fTwitch = Twitch(vAnchor.xyz, vTwitch.x, vTwitch.y, vTwitch.z + fOffset);
 
-	float fTwitch = 0.0;
-	if (bTwitch)
-		fTwitch = Twitch(vAnchor.xyz, vTwitch.x, vTwitch.y, vTwitch.z + fOffset);
 
-	matTumble = mul_float3x3_float3x3(matTumble, RotationMatrix(vAxis, fScale * (fAngle * fAdherence * fAdherenceScale + fOsc * fFlip + fTwitch)));
-//	matTumble = mul_float3x3_float3x3(matTumble, RotationMatrix(vAxis, fAngle));
+		matTumble = mul_float3x3_float3x3(matTumble, RotationMatrix(vAxis, fScale * (fAngle * fAdherence * fAdherenceScale + fOsc * fFlip + fTwitch)));
+		//	matTumble = mul_float3x3_float3x3(matTumble, RotationMatrix(vAxis, fAngle));
+	}
  
  	vDirection = mul_float3x3_float3(matTumble, vDirection);
 	vOriginPos = mul_float3x3_float3(matTumble, vOriginPos);
@@ -757,13 +763,17 @@ float3 LeafWind(FSpeedTreeData STData,
 	float fRippleTrigOffset,
 	int WindType)
 {
+	
+	if (WindType >= SPEEDTREE_WIND_TYPE_BEST_PLUS)
+	{
+		vPos = LeafRipple(vPos, vDirection, fScale, fPackedRippleDir,
+			(bLeaf2 ? STData.WindLeaf2Ripple.x : STData.WindLeaf1Ripple.x),
+			(bLeaf2 ? STData.WindLeaf2Ripple.y : STData.WindLeaf1Ripple.y),
+			false, fRippleTrigOffset);
+	}
 
-	vPos = LeafRipple(vPos, vDirection, fScale, fPackedRippleDir, 
-							(bLeaf2 ? STData.WindLeaf2Ripple.x : STData.WindLeaf1Ripple.x),
-							(bLeaf2 ? STData.WindLeaf2Ripple.y : STData.WindLeaf1Ripple.y),
-							false, fRippleTrigOffset);
 
-	if (WindType == SPEEDTREE_WIND_TYPE_BEST)
+	if (WindType >= SPEEDTREE_WIND_TYPE_BEST)
 	{
 		float3 vGrowthDir = UnpackNormalFromFloat(fPackedGrowthDir);
 		vPos = LeafTumble(STData, vPos, vDirection, fScale, vAnchor, vGrowthDir, fPackedGrowthDir,
@@ -773,6 +783,7 @@ float3 LeafWind(FSpeedTreeData STData,
 						  (bLeaf2 ? STData.WindLeaf2Tumble.w : STData.WindLeaf1Tumble.w),
 						  (bLeaf2 ? STData.WindLeaf2Twitch.xyz : STData.WindLeaf1Twitch.xyz), 
 						  0.0f,
+						  WindType,
 						  (bLeaf2 ? true : true), 
 						  (bLeaf2 ? true : true));
 	}
diff --git a/Engine/Shaders/TonemapCommon.usf b/Engine/Shaders/TonemapCommon.usf
--- a/Engine/Shaders/TonemapCommon.usf
+++ b/Engine/Shaders/TonemapCommon.usf
@@ -39,6 +39,11 @@ half3 FilmPostProcess(half3 LinearColor)
 			MatrixColor = LinearColor * ColorMatrixB_ColorCurveCm2.rgb;
 		#endif
 	#endif
+
+	// DON'T CHECK INTO MAIN: NEEDS A PERMUTATION TO AVOID NON-BACKWARDS COMPATIBLE CHANGES
+	// This corrects so the linear segment is linear in log space!
+	MatrixColor = log2(1.0+MatrixColor);
+
 	// Apply color curve (includes tonemapping).
 	#if USE_CONTRAST == 1
 		// Full path.
@@ -59,6 +64,11 @@ half3 FilmPostProcess(half3 LinearColor)
 		half3 MatrixColorM = min(MatrixColor, ColorCurve_Cm0Cd0_Cd2_Ch0Cm1_Ch3.zzz);
 		half3 CurveColor = (MatrixColorH*ColorCurve_Ch1_Ch2.xxx + ColorCurve_Ch1_Ch2.yyy) * rcp(MatrixColorH + ColorCurve_Cm0Cd0_Cd2_Ch0Cm1_Ch3.www) + MatrixColorM;
 	#endif
+
+	// DON'T CHECK INTO MAIN: NEEDS A PERMUTATION TO AVOID NON-BACKWARDS COMPATIBLE CHANGES
+	// This corrects so the linear segment is linear in log space!
+	CurveColor = exp2(CurveColor)-1.0;
+
 	return CurveColor;
 }
 
@@ -116,6 +126,7 @@ half3 FilmPostProcessPhoto(half3 LinearColor)
 	half Over = (PhotoTint_Over.a*OverLuma) / (PhotoOverTint_OneMinusMid.a+OverLuma);
 	PhotoColor = lerp(PhotoColor, PhotoOverTint_OneMinusMid.rgb*OverLuma, Over);
 	// Highlight compression without change of mid-tone.
-	PhotoColor *= 1.0/(PhotoOverTint_OneMinusMid.a+PhotoColor);    
+	const float WhiteScale = 1.1;
+	PhotoColor *= WhiteScale/(PhotoOverTint_OneMinusMid.a*WhiteScale+PhotoColor);    
 	return PhotoColor;
 }
diff --git a/Engine/Shaders/VertexFactoryCommon.usf b/Engine/Shaders/VertexFactoryCommon.usf
--- a/Engine/Shaders/VertexFactoryCommon.usf
+++ b/Engine/Shaders/VertexFactoryCommon.usf
@@ -12,10 +12,6 @@ float3 TransformLocalToWorld(float3 LocalPosition)
 
 float4 TransformLocalToTranslatedWorld(float3 LocalPosition)
 {
-	return float4(TransformLocalToWorld(LocalPosition) + View.PreViewTranslation.xyz,1);
-}
-
-float4 TransformLocalToPrevTranslatedWorld(float3 LocalPosition)
-{
-	return float4(TransformLocalToWorld(LocalPosition) + View.PrevPreViewTranslation.xyz,1);
+	float3 RotatedPosition = Primitive.LocalToWorld[0].xyz * LocalPosition.xxx + Primitive.LocalToWorld[1].xyz * LocalPosition.yyy + Primitive.LocalToWorld[2].xyz * LocalPosition.zzz;
+	return float4(RotatedPosition + (Primitive.LocalToWorld[3].xyz + View.PreViewTranslation.xyz),1);
 }
diff --git a/Engine/Source/Developer/Android/AndroidTargetPlatform/Private/AndroidTargetPlatform.inl b/Engine/Source/Developer/Android/AndroidTargetPlatform/Private/AndroidTargetPlatform.inl
--- a/Engine/Source/Developer/Android/AndroidTargetPlatform/Private/AndroidTargetPlatform.inl
+++ b/Engine/Source/Developer/Android/AndroidTargetPlatform/Private/AndroidTargetPlatform.inl
@@ -305,6 +305,8 @@ inline void FAndroidTargetPlatform<TPlatformProperties>::AddTextureFormatIfSuppo
 template<class TPlatformProperties>
 inline bool FAndroidTargetPlatform<TPlatformProperties>::HandleTicker( float DeltaTime )
 {
+	QUICK_SCOPE_CYCLE_COUNTER(STAT_FAndroidTargetPlatform_HandleTicker);
+
 	if (DeviceDetection == nullptr)
 	{
 		DeviceDetection = FModuleManager::LoadModuleChecked<IAndroidDeviceDetectionModule>("AndroidDeviceDetection").GetAndroidDeviceDetection();
diff --git a/Engine/Source/Developer/MeshUtilities/Private/MeshUtilities.cpp b/Engine/Source/Developer/MeshUtilities/Private/MeshUtilities.cpp
--- a/Engine/Source/Developer/MeshUtilities/Private/MeshUtilities.cpp
+++ b/Engine/Source/Developer/MeshUtilities/Private/MeshUtilities.cpp
@@ -3628,6 +3628,21 @@ bool FMeshUtilities::BuildSkeletalMesh( FStaticLODModel& LODModel, const FRefere
 #endif
 }
 
+static bool NonOpaqueMaterialPredicate(UStaticMeshComponent* InMesh)
+{
+	TArray<UMaterialInterface*> OutMaterials;
+	InMesh->GetUsedMaterials(OutMaterials);
+	for (auto Material : OutMaterials)
+	{
+		if (Material == nullptr || Material->GetBlendMode() != BLEND_Opaque)
+		{
+			return true;
+		}
+	}
+
+	return false; 
+}
+
 void FMeshUtilities::CreateProxyMesh( 
 	const TArray<AActor*>& SourceActors, 
 	const struct FMeshProxySettings& InProxySettings,
@@ -3670,6 +3685,8 @@ void FMeshUtilities::CreateProxyMesh(
 			Actor->GetComponents<UStaticMeshComponent>(Components);
 			// TODO: support instanced static meshes
 			Components.RemoveAll([](UStaticMeshComponent* Val){ return Val->IsA(UInstancedStaticMeshComponent::StaticClass()); });
+			// TODO: support non-opaque materials
+			Components.RemoveAll(&NonOpaqueMaterialPredicate);
 			//
 			ComponentsToMerge.Append(Components);
 		}
diff --git a/Engine/Source/Developer/MessageLog/Private/UserInterface/SMessageLogMessageListRow.cpp b/Engine/Source/Developer/MessageLog/Private/UserInterface/SMessageLogMessageListRow.cpp
--- a/Engine/Source/Developer/MessageLog/Private/UserInterface/SMessageLogMessageListRow.cpp
+++ b/Engine/Source/Developer/MessageLog/Private/UserInterface/SMessageLogMessageListRow.cpp
@@ -124,6 +124,8 @@ void SMessageLogMessageListRow::CreateMessage(const TSharedRef<SHorizontalBox>&
 	TSharedPtr<SWidget> Content;
 	FName IconBrushName;
 
+	TAttribute<EVisibility> TokenContentVisbility;
+	
 	switch (InMessageToken->GetType())
 	{
 	case EMessageToken::Image:
@@ -181,6 +183,8 @@ void SMessageLogMessageListRow::CreateMessage(const TSharedRef<SHorizontalBox>&
 			.ToolTipText(ActionToken->GetActionDescription())
 			.TextStyle(FEditorStyle::Get(), "MessageLog")
 			.OnNavigate(this, &SMessageLogMessageListRow::HandleActionHyperlinkNavigate, ActionToken);
+
+		TokenContentVisbility = TAttribute<EVisibility>::Create(TAttribute<EVisibility>::FGetter::CreateRaw(this, &SMessageLogMessageListRow::GetActionLinkVisibility, ActionToken));
 	}
 		break;
 
@@ -243,6 +247,7 @@ void SMessageLogMessageListRow::CreateMessage(const TSharedRef<SHorizontalBox>&
 			.Padding(Padding, 0.0f, 0.0f, 0.0f)
 			[
 				SNew(SHorizontalBox)
+				.Visibility(TokenContentVisbility.IsBound() ? TokenContentVisbility : EVisibility::Visible)
 
 				+ SHorizontalBox::Slot()
 				.AutoWidth()
diff --git a/Engine/Source/Developer/MessageLog/Private/UserInterface/SMessageLogMessageListRow.h b/Engine/Source/Developer/MessageLog/Private/UserInterface/SMessageLogMessageListRow.h
--- a/Engine/Source/Developer/MessageLog/Private/UserInterface/SMessageLogMessageListRow.h
+++ b/Engine/Source/Developer/MessageLog/Private/UserInterface/SMessageLogMessageListRow.h
@@ -64,6 +64,11 @@ protected:
 
 private:
 
+	EVisibility GetActionLinkVisibility(TSharedRef<FActionToken> ActionToken) const
+	{
+		return ActionToken->CanExecuteAction() ? EVisibility::Visible : EVisibility::Collapsed;
+	}
+
 	void HandleActionHyperlinkNavigate( TSharedRef<FActionToken> ActionToken )
 	{
 		ActionToken->ExecuteAction();
diff --git a/Engine/Source/Editor/AnimGraph/Private/AnimGraphNode_SequencePlayer.cpp b/Engine/Source/Editor/AnimGraph/Private/AnimGraphNode_SequencePlayer.cpp
--- a/Engine/Source/Editor/AnimGraph/Private/AnimGraphNode_SequencePlayer.cpp
+++ b/Engine/Source/Editor/AnimGraph/Private/AnimGraphNode_SequencePlayer.cpp
@@ -86,32 +86,33 @@ FText UAnimGraphNode_SequencePlayer::GetNodeTitle(ENodeTitleType::Type TitleType
 	{
 		return LOCTEXT("SequenceNullTitle", "Play (None)");
 	}
-	// @TODO: the bone can be altered in the property editor, so we have to 
-	//        choose to mark this dirty when that happens for this to properly work
-	else //if (!CachedNodeTitles.IsTitleCached(TitleType))
+	else
 	{
-		if(SyncGroup.GroupName == NAME_None)
-		{
-			TitleType = ENodeTitleType::ListView;
+		const bool bAdditive = Node.Sequence->IsValidAdditive();
+		const FText BasicTitle = GetTitleGivenAssetInfo(FText::FromName(Node.Sequence->GetFName()), bAdditive);
 
-			if(TitleType == ENodeTitleType::ListView || TitleType == ENodeTitleType::MenuTitle)
-			{
-				const bool bAdditive = Node.Sequence->IsValidAdditive();
-				// FText::Format() is slow, so we cache this to save on performance
-				CachedNodeTitles.SetCachedTitle(TitleType, GetTitleGivenAssetInfo(FText::FromName(Node.Sequence->GetFName()), bAdditive));
-			}
+		if (SyncGroup.GroupName == NAME_None)
+		{
+			return BasicTitle;
 		}
-		else if (TitleType == ENodeTitleType::FullTitle)
+		else
 		{
+			const FText SyncGroupName = FText::FromName(SyncGroup.GroupName);
+
 			FFormatNamedArguments Args;
-			Args.Add(TEXT("Title"), CachedNodeTitles[ENodeTitleType::ListView]);
-			Args.Add(TEXT("SyncGroup"), FText::FromName(SyncGroup.GroupName));
+			Args.Add(TEXT("Title"), BasicTitle);
+			Args.Add(TEXT("SyncGroup"), SyncGroupName);
 
-			// FText::Format() is slow, so we cache this to save on performance
-			CachedNodeTitles.SetCachedTitle(TitleType, FText::Format(LOCTEXT("SequenceNodeGroupWithSubtitle", "{Title}\nSync group {SyncGroup}"), Args));
+			if (TitleType == ENodeTitleType::FullTitle)
+			{
+				return FText::Format(LOCTEXT("SequenceNodeGroupWithSubtitleFull", "{Title}\nSync group {SyncGroup}"), Args);
+			}
+			else
+			{
+				return FText::Format(LOCTEXT("SequenceNodeGroupWithSubtitleList", "{Title} (Sync group {SyncGroup})"), Args);
+			}
 		}
 	}
-	return CachedNodeTitles[TitleType];
 }
 
 FText UAnimGraphNode_SequencePlayer::GetTitleGivenAssetInfo(const FText& AssetName, bool bKnownToBeAdditive)
diff --git a/Engine/Source/Editor/GameProjectGeneration/Private/GameProjectUtils.cpp b/Engine/Source/Editor/GameProjectGeneration/Private/GameProjectUtils.cpp
--- a/Engine/Source/Editor/GameProjectGeneration/Private/GameProjectUtils.cpp
+++ b/Engine/Source/Editor/GameProjectGeneration/Private/GameProjectUtils.cpp
@@ -1505,6 +1505,9 @@ bool GameProjectUtils::NameContainsOnlyLegalCharacters(const FString& TestName,
 
 bool GameProjectUtils::NameContainsUnderscoreAndXB1Installed(const FString& TestName)
 {
+	// disabled for now so people with the SDK installed can use the editor
+	return false;
+
 	bool bContainsIllegalCharacters = false;
 
 	// Only allow alphanumeric characters in the project name
diff --git a/Engine/Source/Editor/Persona/Private/AnimViewportShowCommands.cpp b/Engine/Source/Editor/Persona/Private/AnimViewportShowCommands.cpp
--- a/Engine/Source/Editor/Persona/Private/AnimViewportShowCommands.cpp
+++ b/Engine/Source/Editor/Persona/Private/AnimViewportShowCommands.cpp
@@ -18,6 +18,7 @@ void FAnimViewportShowCommands::RegisterCommands()
 	UI_COMMAND( ShowRetargetBasePose, "Retarget Base Pose", "Show retarget Base pose on preview mesh", EUserInterfaceActionType::ToggleButton, FInputGesture() );
 	UI_COMMAND( ShowBound, "Bound", "Show bound on preview mesh", EUserInterfaceActionType::ToggleButton, FInputGesture() );
 	UI_COMMAND( ShowPreviewMesh, "Mesh", "Show the preview mesh", EUserInterfaceActionType::ToggleButton, FInputGesture() );
+	UI_COMMAND( ShowMorphTargets, "MorphTargets", "Display Applied Morph Targets of the mesh", EUserInterfaceActionType::ToggleButton, FInputGesture() );
 
 	UI_COMMAND( ShowBones, "Bones", "Display Skeleton in viewport", EUserInterfaceActionType::ToggleButton, FInputGesture() );
 	UI_COMMAND( ShowBoneNames, "Bone Names", "Display Bone Names in Viewport", EUserInterfaceActionType::ToggleButton, FInputGesture() );
diff --git a/Engine/Source/Editor/Persona/Private/AnimViewportShowCommands.h b/Engine/Source/Editor/Persona/Private/AnimViewportShowCommands.h
--- a/Engine/Source/Editor/Persona/Private/AnimViewportShowCommands.h
+++ b/Engine/Source/Editor/Persona/Private/AnimViewportShowCommands.h
@@ -48,6 +48,9 @@ public:
 	/** Show/hide the preview mesh */
 	TSharedPtr< FUICommandInfo > ShowPreviewMesh;
 
+	/** Show Morphtarget */
+	TSharedPtr< FUICommandInfo > ShowMorphTargets;
+
 	/** Show skeletal mesh bones */	
 	TSharedPtr< FUICommandInfo > ShowBones;
 
diff --git a/Engine/Source/Editor/Persona/Private/SAnimViewportToolBar.cpp b/Engine/Source/Editor/Persona/Private/SAnimViewportToolBar.cpp
--- a/Engine/Source/Editor/Persona/Private/SAnimViewportToolBar.cpp
+++ b/Engine/Source/Editor/Persona/Private/SAnimViewportToolBar.cpp
@@ -463,6 +463,7 @@ TSharedRef<SWidget> SAnimViewportToolBar::GenerateShowMenu() const
 				ShowMenuBuilder.AddMenuEntry( Actions.ShowRetargetBasePose );
 				ShowMenuBuilder.AddMenuEntry( Actions.ShowBound );
 				ShowMenuBuilder.AddMenuEntry( Actions.ShowPreviewMesh );
+				ShowMenuBuilder.AddMenuEntry( Actions.ShowMorphTargets );
 			}
 			ShowMenuBuilder.EndSection();
 
diff --git a/Engine/Source/Editor/Persona/Private/SAnimationEditorViewport.cpp b/Engine/Source/Editor/Persona/Private/SAnimationEditorViewport.cpp
--- a/Engine/Source/Editor/Persona/Private/SAnimationEditorViewport.cpp
+++ b/Engine/Source/Editor/Persona/Private/SAnimationEditorViewport.cpp
@@ -479,6 +479,12 @@ void SAnimationEditorViewportTabBody::BindCommands()
 		FCanExecuteAction::CreateSP(this, &SAnimationEditorViewportTabBody::CanShowPreviewMesh),
 		FIsActionChecked::CreateSP(this, &SAnimationEditorViewportTabBody::IsShowPreviewMeshEnabled));
 
+	CommandList.MapAction(
+		ViewportShowMenuCommands.ShowMorphTargets,
+		FExecuteAction::CreateSP(this, &SAnimationEditorViewportTabBody::OnShowMorphTargets),
+		FCanExecuteAction(),
+		FIsActionChecked::CreateSP(this, &SAnimationEditorViewportTabBody::IsShowingMorphTargets));
+
 	CommandList.MapAction( 
 		ViewportShowMenuCommands.ShowBones,
 		FExecuteAction::CreateSP(this, &SAnimationEditorViewportTabBody::OnShowBones),
@@ -786,6 +792,17 @@ void SAnimationEditorViewportTabBody::OnShowBones()
 	}
 }
 
+void SAnimationEditorViewportTabBody::OnShowMorphTargets()
+{
+	UDebugSkelMeshComponent* PreviewComponent = PersonaPtr.Pin()->PreviewComponent;
+	if(PreviewComponent != NULL)
+	{
+		PreviewComponent->bDisableMorphTarget = !PreviewComponent->bDisableMorphTarget;
+		PreviewComponent->MarkRenderStateDirty();
+		RefreshViewport();
+	}
+}
+
 void SAnimationEditorViewportTabBody::OnShowBoneNames()
 {
 	UDebugSkelMeshComponent* PreviewComponent = PersonaPtr.Pin()->PreviewComponent;
@@ -859,6 +876,12 @@ bool SAnimationEditorViewportTabBody::IsShowingBones() const
 	return PreviewComponent != NULL && PreviewComponent->bDisplayBones;
 }
 
+bool SAnimationEditorViewportTabBody::IsShowingMorphTargets() const
+{
+	UDebugSkelMeshComponent* PreviewComponent = PersonaPtr.Pin()->PreviewComponent;
+	return PreviewComponent != NULL && PreviewComponent->bDisableMorphTarget == false;
+}
+
 bool SAnimationEditorViewportTabBody::IsShowingBoneNames() const
 {
 	UDebugSkelMeshComponent* PreviewComponent = PersonaPtr.Pin()->PreviewComponent;
diff --git a/Engine/Source/Editor/Persona/Private/SAnimationEditorViewport.h b/Engine/Source/Editor/Persona/Private/SAnimationEditorViewport.h
--- a/Engine/Source/Editor/Persona/Private/SAnimationEditorViewport.h
+++ b/Engine/Source/Editor/Persona/Private/SAnimationEditorViewport.h
@@ -202,6 +202,11 @@ private:
 
 	bool IsShowingBones() const;
 
+	/** Show Morphtarget of SkeletalMesh **/
+	void OnShowMorphTargets();
+
+	bool IsShowingMorphTargets() const;
+
 	/** Show Raw Animation on top of Compressed Animation **/
 	void OnShowRawAnimation();
 
diff --git a/Engine/Source/Editor/Persona/Private/SMorphTargetViewer.cpp b/Engine/Source/Editor/Persona/Private/SMorphTargetViewer.cpp
--- a/Engine/Source/Editor/Persona/Private/SMorphTargetViewer.cpp
+++ b/Engine/Source/Editor/Persona/Private/SMorphTargetViewer.cpp
@@ -13,6 +13,7 @@
 
 static const FName ColumnId_MorphTargetNameLabel( "MorphTargetName" );
 static const FName ColumnID_MorphTargetWeightLabel( "Weight" );
+static const FName ColumnID_MorphTargetEditLabel( "Edit" );
 static const FName ColumnID_MorphTargetVertCountLabel( "NumberOfVerts" );
 
 const float MaxMorphWeight = 5.f;
@@ -57,14 +58,25 @@ private:
 	*
 	*/
 	void OnMorphTargetWeightChanged( float NewWeight );
+	/**
+	* Called when the user types the value and enters
+	*
+	* @param NewWeight - The new number the SSpinBox is set to
+	*
+	*/
+	void OnMorphTargetWeightValueCommitted( float NewWeight, ETextCommit::Type CommitType);
 
+	/** Auto fill check call back functions */
+	void OnMorphTargetAutoFillChecked(ECheckBoxState InState);
+	ECheckBoxState IsMorphTargetAutoFillChangedChecked() const;
+	
 	/**
 	* Returns the weight of this morph target
 	*
 	* @return SearchText - The new number the SSpinBox is set to
 	*
 	*/
-	float GetWeight() const { return Item->Weight; }
+	float GetWeight() const;
 
 	/* The SMorphTargetViewer that we push the morph target weights into */
 	SMorphTargetViewer* MorphTargetViewer;
@@ -126,6 +138,23 @@ TSharedRef< SWidget > SMorphTargetListRow::GenerateWidgetForColumn( const FName&
 				.MaxValue(MaxMorphWeight)
 				.Value( this, &SMorphTargetListRow::GetWeight )
 				.OnValueChanged( this, &SMorphTargetListRow::OnMorphTargetWeightChanged )
+				.OnValueCommitted( this, &SMorphTargetListRow::OnMorphTargetWeightValueCommitted )
+			];
+	}
+	else if ( ColumnName == ColumnID_MorphTargetEditLabel )
+	{
+		return
+			SNew(SVerticalBox)
+
+			+ SVerticalBox::Slot()
+			.AutoHeight()
+			.Padding(0.0f, 1.0f)
+			.VAlign(VAlign_Center)
+			.HAlign(HAlign_Center)
+			[
+				SNew(SCheckBox)
+				.OnCheckStateChanged(this, &SMorphTargetListRow::OnMorphTargetAutoFillChecked)
+				.IsChecked(this, &SMorphTargetListRow::IsMorphTargetAutoFillChangedChecked)
 			];
 	}
 	else
@@ -152,18 +181,41 @@ TSharedRef< SWidget > SMorphTargetListRow::GenerateWidgetForColumn( const FName&
 	}
 }
 
+void SMorphTargetListRow::OnMorphTargetAutoFillChecked(ECheckBoxState InState)
+{
+	Item->bAutoFillData = InState == ECheckBoxState::Checked;
+
+	if (Item->bAutoFillData)
+	{
+		// clear value so that it can be filled up
+		MorphTargetViewer->AddMorphTargetOverride(Item->Name, 0.f, true);
+	}
+}
+
+ECheckBoxState SMorphTargetListRow::IsMorphTargetAutoFillChangedChecked() const
+{
+	return (Item->bAutoFillData)? ECheckBoxState::Checked: ECheckBoxState::Unchecked;
+}
+
 void SMorphTargetListRow::OnMorphTargetWeightChanged( float NewWeight )
 {
 	// First change this item...
-	float Delta = NewWeight - Item->Weight;
+	// the delta feature is a bit confusing when debugging morphtargets, and you're not sure why it's changing, so I'm disabling it for now. 
+	// I think in practice, you want each morph target to move independentaly. It is very unlikely you'd like to move multiple things together. 
+#if 0 
+	float Delta = NewWeight - GetWeight();
+#endif
 	Item->Weight = NewWeight;
-	MorphTargetViewer->AddMorphTargetOverride( Item->Name, Item->Weight );
+	Item->bAutoFillData = false;
+
+	MorphTargetViewer->AddMorphTargetOverride( Item->Name, Item->Weight, false );
 
 	if (PersonaPtr.IsValid())
 	{
 		PersonaPtr.Pin()->RefreshViewport();
 	}
 
+#if 0 
 	TArray< TSharedPtr< FDisplayedMorphTargetInfo > > SelectedRows = MorphTargetListView->GetSelectedItems();
 
 	// ...then any selected rows need changing by the same delta
@@ -174,11 +226,63 @@ void SMorphTargetListRow::OnMorphTargetWeightChanged( float NewWeight )
 		if ( RowItem != Item ) // Don't do "this" row again if it's selected
 		{
 			RowItem->Weight = FMath::Clamp(RowItem->Weight + Delta, -MaxMorphWeight, MaxMorphWeight);
-			MorphTargetViewer->AddMorphTargetOverride( RowItem->Name, RowItem->Weight );
+			RowItem->bAutoFillData = false;
+			MorphTargetViewer->AddMorphTargetOverride( RowItem->Name, RowItem->Weight, false );
 		}
 	}
+#endif
 }
 
+void SMorphTargetListRow::OnMorphTargetWeightValueCommitted( float NewWeight, ETextCommit::Type CommitType)
+{
+	if (CommitType == ETextCommit::OnEnter || CommitType == ETextCommit::OnUserMovedFocus)
+	{
+		float NewValidWeight = FMath::Clamp(NewWeight, -MaxMorphWeight, MaxMorphWeight);
+		Item->Weight = NewValidWeight;
+		Item->bAutoFillData = false;
+
+		MorphTargetViewer->AddMorphTargetOverride(Item->Name, Item->Weight, false);
+
+		TArray< TSharedPtr< FDisplayedMorphTargetInfo > > SelectedRows = MorphTargetListView->GetSelectedItems();
+
+		// ...then any selected rows need changing by the same delta
+		for(auto ItemIt = SelectedRows.CreateIterator(); ItemIt; ++ItemIt)
+		{
+			TSharedPtr< FDisplayedMorphTargetInfo > RowItem = (*ItemIt);
+
+			if(RowItem != Item) // Don't do "this" row again if it's selected
+			{
+				RowItem->Weight = NewValidWeight;
+				RowItem->bAutoFillData = false;
+				MorphTargetViewer->AddMorphTargetOverride(RowItem->Name, RowItem->Weight, false);
+			}
+		}
+
+		if(PersonaPtr.IsValid())
+		{
+			PersonaPtr.Pin()->RefreshViewport();
+		}
+	}
+}
+
+float SMorphTargetListRow::GetWeight() const 
+{ 
+	if (PersonaPtr.IsValid() && Item->bAutoFillData)
+	{
+		USkeletalMeshComponent* SkelComp = PersonaPtr.Pin()->GetPreviewMeshComponent();
+		if ( SkelComp && SkelComp->AnimScriptInstance )
+		{
+			// make sure if they have value that's not same as saved value
+			float* CurrentValue = SkelComp->AnimScriptInstance->MorphTargetCurves.Find(Item->Name);
+			if (CurrentValue && *CurrentValue != 0.f)
+			{
+				return *CurrentValue;
+			}
+		}
+	}
+
+	return Item->Weight; 
+}
 //////////////////////////////////////////////////////////////////////////
 // SMorphTargetViewer
 
@@ -240,6 +344,9 @@ void SMorphTargetViewer::Construct(const FArguments& InArgs)
 				+ SHeaderRow::Column( ColumnID_MorphTargetWeightLabel )
 				.DefaultLabel( LOCTEXT( "MorphTargetWeightLabel", "Weight" ) )
 
+				+ SHeaderRow::Column(ColumnID_MorphTargetEditLabel)
+				.DefaultLabel(LOCTEXT("MorphTargetEditLabel", "Auto"))
+
 				+ SHeaderRow::Column( ColumnID_MorphTargetVertCountLabel )
 				.DefaultLabel( LOCTEXT("MorphTargetVertCountLabel", "Vert Count") )
 			)
@@ -335,7 +442,7 @@ void SMorphTargetViewer::CreateMorphTargetList( const FString& SearchText )
 	MorphTargetListView->RequestListRefresh();
 }
 
-void SMorphTargetViewer::AddMorphTargetOverride( FName& Name, float Weight )
+void SMorphTargetViewer::AddMorphTargetOverride( FName& Name, float Weight, bool bRemoveZeroWeight )
 {
 	if ( PersonaPtr.IsValid() )
 	{
@@ -343,7 +450,7 @@ void SMorphTargetViewer::AddMorphTargetOverride( FName& Name, float Weight )
 
 		if ( Mesh )
 		{
-			Mesh->SetMorphTarget( Name, Weight );
+			Mesh->SetMorphTarget( Name, Weight, bRemoveZeroWeight );
 		}
 	}
 }
@@ -371,7 +478,7 @@ void SMorphTargetViewer::OnDeleteMorphTargets()
 			MorphTarget->Modify();
 
 			//Clean up override usage
-			AddMorphTargetOverride(SelectedRows[RowIndex]->Name, 0.0f);
+			AddMorphTargetOverride(SelectedRows[RowIndex]->Name, 0.0f, true);
 
 			SkeletalMesh->UnregisterMorphTarget(MorphTarget);
 		}
diff --git a/Engine/Source/Editor/Persona/Private/SMorphTargetViewer.h b/Engine/Source/Editor/Persona/Private/SMorphTargetViewer.h
--- a/Engine/Source/Editor/Persona/Private/SMorphTargetViewer.h
+++ b/Engine/Source/Editor/Persona/Private/SMorphTargetViewer.h
@@ -12,6 +12,7 @@ class FDisplayedMorphTargetInfo
 public:
 	FName Name;
 	float Weight;
+	bool bAutoFillData;
 	int32 NumberOfVerts;
 
 	/** Static function for creating a new item, but ensures that you can only have a TSharedRef to one */
@@ -25,6 +26,7 @@ protected:
 	FDisplayedMorphTargetInfo(const FName& InSource, int32 InNumberOfVerts)
 		: Name( InSource )
 		, Weight( 0 )
+		, bAutoFillData(true)
 		, NumberOfVerts(InNumberOfVerts)
 	{}
 
@@ -106,7 +108,7 @@ public:
 	* @param Weight - How much of this morph target to apply (0.0 - 1.0)
 	*
 	*/
-	void AddMorphTargetOverride( FName& Name, float Weight );
+	void AddMorphTargetOverride( FName& Name, float Weight, bool bRemoveZeroWeight );
 
 	/**
 	* Tells the AnimInstance to reset all of its morph target curves
diff --git a/Engine/Source/Editor/StaticMeshEditor/Private/StaticMeshEditorViewportClient.cpp b/Engine/Source/Editor/StaticMeshEditor/Private/StaticMeshEditorViewportClient.cpp
--- a/Engine/Source/Editor/StaticMeshEditor/Private/StaticMeshEditorViewportClient.cpp
+++ b/Engine/Source/Editor/StaticMeshEditor/Private/StaticMeshEditorViewportClient.cpp
@@ -590,6 +590,15 @@ void FStaticMeshEditorViewportClient::Draw(const FSceneView* View,FPrimitiveDraw
 				}
 			}
 		}
+
+		// The simple nav geometry is only used by dynamic obstacles for now
+		if (StaticMesh->NavCollision && StaticMesh->NavCollision->bIsDynamicObstacle)
+		{
+			// Draw the static mesh's body setup (simple collision)
+			FTransform GeomTransform(StaticMeshComponent->ComponentToWorld);
+			FColor NavCollisionColor = FColor(118, 84, 255, 255);
+			StaticMesh->NavCollision->DrawSimpleGeom(PDI, GeomTransform, FColorList::LimeGreen);
+		}
 	}
 }
 
diff --git a/Engine/Source/Editor/UnrealEd/Classes/Editor/UnrealEdEngine.h b/Engine/Source/Editor/UnrealEd/Classes/Editor/UnrealEdEngine.h
--- a/Engine/Source/Editor/UnrealEd/Classes/Editor/UnrealEdEngine.h
+++ b/Engine/Source/Editor/UnrealEd/Classes/Editor/UnrealEdEngine.h
@@ -651,6 +651,11 @@ class UNREALED_API UUnrealEdEngine : public UEditorEngine, public FNotifyHook
 	 * @return true if only editor-visible levels should be loaded in Play-In-Editor sessions
 	 */
 	virtual bool OnlyLoadEditorVisibleLevelsInPIE() const;
+
+	/**
+	 * @return true if level streaming should prefer to stream levels from disk instead of duplicating them from editor world
+	 */
+	virtual bool PreferToStreamLevelsInPIE() const override;
 	
 	/**
 	 * If all selected actors belong to the same level, that level is made the current level.
diff --git a/Engine/Source/Editor/UnrealEd/Classes/Factories/FbxAnimSequenceImportData.h b/Engine/Source/Editor/UnrealEd/Classes/Factories/FbxAnimSequenceImportData.h
--- a/Engine/Source/Editor/UnrealEd/Classes/Factories/FbxAnimSequenceImportData.h
+++ b/Engine/Source/Editor/UnrealEd/Classes/Factories/FbxAnimSequenceImportData.h
@@ -45,9 +45,17 @@ class UNREALED_API UFbxAnimSequenceImportData : public UFbxAssetImportData
 	FString SourceAnimationName;
 
 	/** Import if custom attribute as a curve within the animation **/
-	UPROPERTY(EditAnywhere, AdvancedDisplay, config, Category = ImportSettings, meta = (DisplayName = "Import Custom Attribute"))
+	UPROPERTY(EditAnywhere, AdvancedDisplay, config, Category=ImportSettings)
 	bool bImportCustomAttribute;
 
+	/** Type of asset to import from the FBX file */
+	UPROPERTY(EditAnywhere, AdvancedDisplay, config, Category=ImportSettings)
+	bool bPreserveLocalTransform;
+
+	/** Type of asset to import from the FBX file */
+	UPROPERTY(EditAnywhere, AdvancedDisplay, config, Category=ImportSettings)
+	bool bDeleteExistingMorphTargetCurves;
+
 	/** Gets or creates fbx import data for the specified anim sequence */
 	static UFbxAnimSequenceImportData* GetImportDataForAnimSequence(UAnimSequence* AnimSequence, UFbxAnimSequenceImportData* TemplateForCreation);
 
diff --git a/Engine/Source/Editor/UnrealEd/Classes/Factories/FbxImportUI.h b/Engine/Source/Editor/UnrealEd/Classes/Factories/FbxImportUI.h
--- a/Engine/Source/Editor/UnrealEd/Classes/Factories/FbxImportUI.h
+++ b/Engine/Source/Editor/UnrealEd/Classes/Factories/FbxImportUI.h
@@ -107,14 +107,6 @@ class UFbxImportUI : public UObject
 	UPROPERTY(EditAnywhere, Instanced, Category=Animation, meta=(editcondition="bImportAnimations", ImportType = "Animation"))
 	class UFbxAnimSequenceImportData* AnimSequenceImportData;
 
-	/** Type of asset to import from the FBX file */
-	UPROPERTY(EditAnywhere, AdvancedDisplay, config, Category=Animation, meta=(editcondition="bImportAnimations", ImportType = "SkeletalMesh|Animation"))
-	bool bPreserveLocalTransform;
-
-	/** Type of asset to import from the FBX file */
-	UPROPERTY(EditAnywhere, AdvancedDisplay, config, Category=Animation, meta=(editcondition="bImportAnimations", ImportType = "SkeletalMesh|Animation"))
-	bool bDeleteExistingMorphTargetCurves;
-
 	/** Import data used when importing textures */
 	UPROPERTY(EditAnywhere, Instanced, Category=Material)
 	class UFbxTextureImportData* TextureImportData;
diff --git a/Engine/Source/Editor/UnrealEd/Classes/Settings/LevelEditorPlaySettings.h b/Engine/Source/Editor/UnrealEd/Classes/Settings/LevelEditorPlaySettings.h
--- a/Engine/Source/Editor/UnrealEd/Classes/Settings/LevelEditorPlaySettings.h
+++ b/Engine/Source/Editor/UnrealEd/Classes/Settings/LevelEditorPlaySettings.h
@@ -151,6 +151,9 @@ public:
 	UPROPERTY(config)
 	uint32 bOnlyLoadVisibleLevelsInPIE:1;
 
+	UPROPERTY(config, EditAnywhere, Category = PlayInEditor, meta = (DisplayName="Stream Sub-Levels during Play in Editor", ToolTip = "Prefer to stream sub-levels from the disk instead of duplicating editor sub-levels"))
+	uint32 bPreferToStreamLevelsInPIE:1;
+
 public:
 
 	/** The width of the new view port window in pixels (0 = use the desktop's screen resolution). */
diff --git a/Engine/Source/Editor/UnrealEd/Private/EditorExporters.cpp b/Engine/Source/Editor/UnrealEd/Private/EditorExporters.cpp
--- a/Engine/Source/Editor/UnrealEd/Private/EditorExporters.cpp
+++ b/Engine/Source/Editor/UnrealEd/Private/EditorExporters.cpp
@@ -2831,13 +2831,14 @@ namespace MaterialExportUtils
 		return true;
 	}
 
-	UTexture2D* CreateTexture(UPackage* Outer, const FString& AssetLongName, FIntPoint Size, const TArray<FColor>& Samples, TextureCompressionSettings CompressionSettings, TextureGroup LODGroup, EObjectFlags Flags, bool bSRGB)
+	UTexture2D* CreateTexture(UPackage* Outer, const FString& AssetLongName, FIntPoint Size, const TArray<FColor>& Samples, TextureCompressionSettings CompressionSettings, TextureGroup LODGroup, EObjectFlags Flags, bool bSRGB, const FGuid& SourceGuidHash)
 	{
 		FCreateTexture2DParameters TexParams;
 		TexParams.bUseAlpha = false;
 		TexParams.CompressionSettings = CompressionSettings;
 		TexParams.bDeferCompression = true;
 		TexParams.bSRGB = bSRGB;
+		TexParams.SourceGuidHash = SourceGuidHash;
 
 		if (Outer == nullptr)
 		{
diff --git a/Engine/Source/Editor/UnrealEd/Private/EditorViewportClient.cpp b/Engine/Source/Editor/UnrealEd/Private/EditorViewportClient.cpp
--- a/Engine/Source/Editor/UnrealEd/Private/EditorViewportClient.cpp
+++ b/Engine/Source/Editor/UnrealEd/Private/EditorViewportClient.cpp
@@ -710,6 +710,11 @@ FSceneView* FEditorViewportClient::CalcSceneView(FSceneViewFamily* ViewFamily)
 	ViewInitOptions.OverrideLODViewOrigin = FVector::ZeroVector;
 	ViewInitOptions.bUseFauxOrthoViewPos = true;
 
+	if (bUseControllingActorViewInfo)
+	{
+		ViewInitOptions.bUseFieldOfViewForLOD = ControllingActorViewInfo.bUseFieldOfViewForLOD;
+	}
+
 	ViewInitOptions.OverrideFarClippingPlaneDistance = FarPlane;
 	ViewInitOptions.CursorPos = CurrentMousePos;
 
diff --git a/Engine/Source/Editor/UnrealEd/Private/Fbx/FbxMainImport.cpp b/Engine/Source/Editor/UnrealEd/Private/Fbx/FbxMainImport.cpp
--- a/Engine/Source/Editor/UnrealEd/Private/Fbx/FbxMainImport.cpp
+++ b/Engine/Source/Editor/UnrealEd/Private/Fbx/FbxMainImport.cpp
@@ -244,8 +244,8 @@ void ApplyImportUIToImportOptions(UFbxImportUI* ImportUI, FBXImportOptions& InOu
 	InOutImportOptions.AnimationRange.X = ImportUI->AnimSequenceImportData->StartFrame;
 	InOutImportOptions.AnimationRange.Y = ImportUI->AnimSequenceImportData->EndFrame;
 	InOutImportOptions.AnimationName = ImportUI->AnimationName;
-	InOutImportOptions.bPreserveLocalTransform = ImportUI->bPreserveLocalTransform;
-	InOutImportOptions.bDeleteExistingMorphTargetCurves = ImportUI->bDeleteExistingMorphTargetCurves;
+	InOutImportOptions.bPreserveLocalTransform = ImportUI->AnimSequenceImportData->bPreserveLocalTransform;
+	InOutImportOptions.bDeleteExistingMorphTargetCurves = ImportUI->AnimSequenceImportData->bDeleteExistingMorphTargetCurves;
 	InOutImportOptions.bImportCustomAttribute = ImportUI->AnimSequenceImportData->bImportCustomAttribute;
 }
 
diff --git a/Engine/Source/Editor/UnrealEd/Private/Fbx/FbxStaticMeshImport.cpp b/Engine/Source/Editor/UnrealEd/Private/Fbx/FbxStaticMeshImport.cpp
--- a/Engine/Source/Editor/UnrealEd/Private/Fbx/FbxStaticMeshImport.cpp
+++ b/Engine/Source/Editor/UnrealEd/Private/Fbx/FbxStaticMeshImport.cpp
@@ -1232,6 +1232,25 @@ UStaticMesh* UnFbx::FFbxImporter::ImportStaticMeshAsSingle(UObject* InParent, TA
 		StaticMesh->LODGroup = ImportOptions->StaticMeshLODGroup;
 		StaticMesh->Build(false);
 		
+		// this is damage control. After build, we'd like to absolutely sure that 
+		// all index is pointing correctly and they're all used. Otherwise we remove them
+		FMeshSectionInfoMap OldSectionInfoMap = StaticMesh->SectionInfoMap;
+		StaticMesh->SectionInfoMap.Clear();
+		// fix up section data
+		for (int32 LODIndex = 0; LODIndex<StaticMesh->RenderData->LODResources.Num(); ++LODIndex)
+		{
+			FStaticMeshLODResources& LOD = StaticMesh->RenderData->LODResources[LODIndex];
+			int32 NumSections = LOD.Sections.Num();
+			for(int32 SectionIndex = 0; SectionIndex < NumSections; ++SectionIndex)
+			{
+				FMeshSectionInfo Info = OldSectionInfoMap.Get(LODIndex, SectionIndex);
+				if (StaticMesh->Materials.IsValidIndex(Info.MaterialIndex))
+				{
+					StaticMesh->SectionInfoMap.Set(LODIndex, SectionIndex, Info);
+				}
+			}
+		}
+
 		// The code to check for bad lightmap UVs doesn't scale well with number of triangles.
 		// Skip it here because Lightmass will warn about it during a light build anyway.
 		bool bWarnOnBadLightmapUVs = false;
diff --git a/Engine/Source/Editor/UnrealEd/Private/PreviewMaterial.cpp b/Engine/Source/Editor/UnrealEd/Private/PreviewMaterial.cpp
--- a/Engine/Source/Editor/UnrealEd/Private/PreviewMaterial.cpp
+++ b/Engine/Source/Editor/UnrealEd/Private/PreviewMaterial.cpp
@@ -2,7 +2,7 @@
 
 #include "UnrealEd.h"
 #include "Editor/MaterialEditor/Public/MaterialEditorModule.h"
-#include "ComponentReregisterContext.h"
+#include "ComponentRecreateRenderStateContext.h"
 #include "Materials/MaterialInstanceConstant.h"
 
 /**
@@ -178,8 +178,8 @@ void UMaterialEditorInstanceConstant::PostEditChangeProperty(FPropertyChangedEve
 		{
 			UpdateSourceInstanceParent();
 
-			FGlobalComponentReregisterContext RecreateComponents;
-			// Fully update static parameters before re-registering the scene's components
+			FGlobalComponentRecreateRenderStateContext RecreateComponentsRenderState;
+			// Fully update static parameters before recreating render state for all components
 			SetSourceInstance(SourceInstance);
 		}
 
diff --git a/Engine/Source/Editor/UnrealEd/Private/Settings/SettingsClasses.cpp b/Engine/Source/Editor/UnrealEd/Private/Settings/SettingsClasses.cpp
--- a/Engine/Source/Editor/UnrealEd/Private/Settings/SettingsClasses.cpp
+++ b/Engine/Source/Editor/UnrealEd/Private/Settings/SettingsClasses.cpp
@@ -94,7 +94,7 @@ void UEditorExperimentalSettings::PostEditChangeProperty( struct FPropertyChange
 
 UEditorLoadingSavingSettings::UEditorLoadingSavingSettings( const FObjectInitializer& ObjectInitializer )
 	: Super(ObjectInitializer)
-	, bMonitorContentDirectories(true)
+	, bMonitorContentDirectories(false)
 	, bAutoCreateAssets(true)
 	, bAutoDeleteAssets(true)
 	, bDeleteSourceFilesWithAssets(false)
diff --git a/Engine/Source/Editor/UnrealEd/Private/UnrealEdEngine.cpp b/Engine/Source/Editor/UnrealEd/Private/UnrealEdEngine.cpp
--- a/Engine/Source/Editor/UnrealEd/Private/UnrealEdEngine.cpp
+++ b/Engine/Source/Editor/UnrealEd/Private/UnrealEdEngine.cpp
@@ -815,6 +815,10 @@ bool UUnrealEdEngine::OnlyLoadEditorVisibleLevelsInPIE() const
 	return GetDefault<ULevelEditorPlaySettings>()->bOnlyLoadVisibleLevelsInPIE;
 }
 
+bool UUnrealEdEngine::PreferToStreamLevelsInPIE() const
+{
+	return GetDefault<ULevelEditorPlaySettings>()->bPreferToStreamLevelsInPIE;
+}
 
 void UUnrealEdEngine::RedrawLevelEditingViewports(bool bInvalidateHitProxies)
 {
diff --git a/Engine/Source/Editor/UnrealEd/Public/MaterialExportUtils.h b/Engine/Source/Editor/UnrealEd/Public/MaterialExportUtils.h
--- a/Engine/Source/Editor/UnrealEd/Public/MaterialExportUtils.h
+++ b/Engine/Source/Editor/UnrealEd/Public/MaterialExportUtils.h
@@ -87,9 +87,10 @@ namespace MaterialExportUtils
 	 * @param LODGroup				LODGroup for the new texture
 	 * @param Flags					ObjectFlags for the new texture
 	 * @param bSRGB					Whether to set the bSRGB flag on the new texture
+	 * @param SourceGuidHash		(optional) Hash (stored as Guid) to use part of the texture source's DDC key.
 	 * @return						The new texture.
 	 */
-	UNREALED_API UTexture2D* CreateTexture(UPackage* Outer, const FString& AssetLongName, FIntPoint Size, const TArray<FColor>& Samples, TextureCompressionSettings CompressionSettings, TextureGroup LODGroup, EObjectFlags Flags, bool bSRGB);
+	UNREALED_API UTexture2D* CreateTexture(UPackage* Outer, const FString& AssetLongName, FIntPoint Size, const TArray<FColor>& Samples, TextureCompressionSettings CompressionSettings, TextureGroup LODGroup, EObjectFlags Flags, bool bSRGB, const FGuid& SourceGuidHash = FGuid());
 
 	/**
 	* Creates bakes textures for a ULandscapeComponent
diff --git a/Engine/Source/Editor/WorldBrowser/Private/LevelCollectionModel.cpp b/Engine/Source/Editor/WorldBrowser/Private/LevelCollectionModel.cpp
--- a/Engine/Source/Editor/WorldBrowser/Private/LevelCollectionModel.cpp
+++ b/Engine/Source/Editor/WorldBrowser/Private/LevelCollectionModel.cpp
@@ -210,12 +210,10 @@ void FLevelCollectionModel::PopulateFilteredLevelsList()
 	FilteredLevelsList.Empty();
 
 	// Filter out our flat list
-	for (auto It = AllLevelsList.CreateIterator(); It; ++It)
+	for (auto& LevelModel : AllLevelsList)
 	{
-		TSharedPtr<FLevelModel> LevelModel = (*It);
-		
 		LevelModel->SetLevelFilteredOutFlag(true);
-		if (LevelModel->IsPersistent() || PassesAllFilters(LevelModel))
+		if (LevelModel->IsPersistent() || PassesAllFilters(*LevelModel))
 		{
 			FilteredLevelsList.Add(LevelModel);
 			LevelModel->SetLevelFilteredOutFlag(false);
@@ -335,9 +333,9 @@ void FLevelCollectionModel::SetSelectedLevels(const FLevelModelList& InList)
 	SelectedLevelsList.Empty(); 
 	
 	// Set selection flag to selected levels	
-	for (auto LevelModel : InList)
+	for (auto& LevelModel : InList)
 	{
-		if (LevelModel.IsValid() && PassesAllFilters(LevelModel))
+		if (LevelModel.IsValid() && PassesAllFilters(*LevelModel))
 		{
 			LevelModel->SetLevelSelectionFlag(true);
 			SelectedLevelsList.Add(LevelModel);
@@ -656,9 +654,9 @@ TSharedPtr<FLevelDragDropOp> FLevelCollectionModel::CreateDragDropOp() const
 	return TSharedPtr<FLevelDragDropOp>();
 }
 
-bool FLevelCollectionModel::PassesAllFilters(TSharedPtr<FLevelModel> Item) const
+bool FLevelCollectionModel::PassesAllFilters(const FLevelModel& Item) const
 {
-	if (Item->IsPersistent() || Filters->PassesAllFilters(Item))
+	if (Item.IsPersistent() || Filters->PassesAllFilters(&Item))
 	{
 		return true;
 	}
diff --git a/Engine/Source/Editor/WorldBrowser/Private/LevelCollectionModel.h b/Engine/Source/Editor/WorldBrowser/Private/LevelCollectionModel.h
--- a/Engine/Source/Editor/WorldBrowser/Private/LevelCollectionModel.h
+++ b/Engine/Source/Editor/WorldBrowser/Private/LevelCollectionModel.h
@@ -6,8 +6,8 @@
 #include "IFilter.h"
 #include "FilterCollection.h"
 
-typedef IFilter< const TSharedPtr<FLevelModel>& >				LevelFilter;
-typedef TFilterCollection< const TSharedPtr<FLevelModel>& >		LevelFilterCollection;
+typedef IFilter< const FLevelModel* >				LevelFilter;
+typedef TFilterCollection< const FLevelModel* >		LevelFilterCollection;
 
 /**
  * Interface for non-UI presentation logic for a world
@@ -122,7 +122,7 @@ public:
 	virtual TSharedPtr<class FLevelDragDropOp> CreateDragDropOp() const;
 	
 	/**	@return	Whether specified level passes all filters */
-	virtual bool PassesAllFilters(TSharedPtr<FLevelModel> InLevelModel) const;
+	virtual bool PassesAllFilters(const FLevelModel& InLevelModel) const;
 	
 	/**	Builds 'hierarchy' commands menu for a selected levels */
 	virtual void BuildHierarchyMenu(FMenuBuilder& InMenuBuilder) const;
diff --git a/Engine/Source/Editor/WorldBrowser/Private/LevelModel.cpp b/Engine/Source/Editor/WorldBrowser/Private/LevelModel.cpp
--- a/Engine/Source/Editor/WorldBrowser/Private/LevelModel.cpp
+++ b/Engine/Source/Editor/WorldBrowser/Private/LevelModel.cpp
@@ -375,6 +375,11 @@ void FLevelModel::SetLevelColor(FLinearColor InColor)
 	// Currently not all base classes have the requisite support, so I've not made it pure virtual.
 }
 
+bool FLevelModel::IsVisibleInCompositionView() const
+{
+	return false;
+}
+
 bool FLevelModel::HasKismet() const
 {
 	return (GetLevelObject() != NULL);
@@ -439,14 +444,14 @@ void FLevelModel::OnFilterChanged()
 {
 	FilteredChildren.Empty();
 	
-	for (auto It = AllChildren.CreateConstIterator(); It; ++It)
+	for (const auto& LevelModel : AllChildren)
 	{
-		(*It)->OnFilterChanged();
+		LevelModel->OnFilterChanged();
 
 		// Item will pass filtering regardless of filter settings if it has children that passes filtering
-		if ((*It)->GetChildren().Num() > 0 || LevelCollectionModel.PassesAllFilters(*It))
+		if (LevelModel->GetChildren().Num() > 0 || LevelCollectionModel.PassesAllFilters(*LevelModel))
 		{
-			FilteredChildren.Add(*It);
+			FilteredChildren.Add(LevelModel);
 		}
 	}
 }
@@ -482,7 +487,7 @@ void FLevelModel::AddChild(TSharedPtr<FLevelModel> InChild)
 {
 	AllChildren.AddUnique(InChild);
 
-	if (LevelCollectionModel.PassesAllFilters(InChild))
+	if (LevelCollectionModel.PassesAllFilters(*InChild))
 	{
 		FilteredChildren.Add(InChild);
 	}
diff --git a/Engine/Source/Editor/WorldBrowser/Private/LevelModel.h b/Engine/Source/Editor/WorldBrowser/Private/LevelModel.h
--- a/Engine/Source/Editor/WorldBrowser/Private/LevelModel.h
+++ b/Engine/Source/Editor/WorldBrowser/Private/LevelModel.h
@@ -150,6 +150,9 @@ public:
 
 	/** Sets level color, used for visualization. (Show -> Advanced -> Level Coloration) */
 	virtual void SetLevelColor(FLinearColor InColor);
+
+	/** Whether level should be drawn in world composition view */
+	virtual bool IsVisibleInCompositionView() const;
 	
 	/**	@return Whether level has associated blueprint script */
 	bool HasKismet() const;
diff --git a/Engine/Source/Editor/WorldBrowser/Private/SWorldHierarchy.cpp b/Engine/Source/Editor/WorldBrowser/Private/SWorldHierarchy.cpp
--- a/Engine/Source/Editor/WorldBrowser/Private/SWorldHierarchy.cpp
+++ b/Engine/Source/Editor/WorldBrowser/Private/SWorldHierarchy.cpp
@@ -12,7 +12,7 @@
 
 #define LOCTEXT_NAMESPACE "WorldBrowser"
 
-typedef TTextFilter<const TSharedPtr<FLevelModel>& > LevelTextFilter;
+typedef TTextFilter<const FLevelModel* > LevelTextFilter;
 
 class SLevelsTreeWidget : public STreeView<TSharedPtr<FLevelModel>>
 {
@@ -428,9 +428,9 @@ private:
 	}
 
 	/** Appends the Level's name to the OutSearchStrings array if the Level is valid */
-	void TransformLevelToString(const TSharedPtr<FLevelModel>& Level, TArray<FString>& OutSearchStrings) const
+	void TransformLevelToString(const FLevelModel* Level, TArray<FString>& OutSearchStrings) const
 	{
-		if (Level.IsValid() && Level->HasValidPackage())
+		if (Level != nullptr && Level->HasValidPackage())
 		{
 			OutSearchStrings.Add(FPackageName::GetShortName(Level->GetLongPackageName()));
 		}
diff --git a/Engine/Source/Editor/WorldBrowser/Private/Tiles/SWorldComposition.cpp b/Engine/Source/Editor/WorldBrowser/Private/Tiles/SWorldComposition.cpp
--- a/Engine/Source/Editor/WorldBrowser/Private/Tiles/SWorldComposition.cpp
+++ b/Engine/Source/Editor/WorldBrowser/Private/Tiles/SWorldComposition.cpp
@@ -680,8 +680,8 @@ protected:
 		for (int32 ChildIndex=0; ChildIndex < Children.Num(); ++ChildIndex)
 		{
 			const auto Child = StaticCastSharedRef<SWorldTileItem>(Children[ChildIndex]);
-		
-			if (WorldModel->PassesAllFilters(Child->GetLevelModel()))
+			const auto LevelModel = Child->GetLevelModel();
+			if (LevelModel->IsVisibleInCompositionView())
 			{
 				FSlateRect ChildRect = Child->GetItemRect();
 				FVector2D ChildSize = ChildRect.GetSize();
diff --git a/Engine/Source/Editor/WorldBrowser/Private/Tiles/SWorldLayers.cpp b/Engine/Source/Editor/WorldBrowser/Private/Tiles/SWorldLayers.cpp
--- a/Engine/Source/Editor/WorldBrowser/Private/Tiles/SWorldLayers.cpp
+++ b/Engine/Source/Editor/WorldBrowser/Private/Tiles/SWorldLayers.cpp
@@ -193,6 +193,7 @@ void SWorldLayerButton::Construct(const FArguments& InArgs)
 						.OnCheckStateChanged(this, &SWorldLayerButton::OnCheckStateChanged)
 						.IsChecked(this, &SWorldLayerButton::IsChecked)
 						.OnGetMenuContent(this, &SWorldLayerButton::GetRightClickMenu)
+						.ToolTipText(this, &SWorldLayerButton::GetToolTipText)
 						.Padding(3)
 						[
 							SNew(STextBlock)
@@ -241,4 +242,16 @@ TSharedRef<SWidget> SWorldLayerButton::GetRightClickMenu()
 	return SNullWidget::NullWidget;
 }
 
+FText SWorldLayerButton::GetToolTipText() const
+{
+	if (WorldLayer.DistanceStreamingEnabled)
+	{
+		return FText::Format(LOCTEXT("Layer_Distance_Tooltip", "Streaming Distance: {0}"), FText::AsNumber(WorldLayer.StreamingDistance));
+	}
+	else
+	{
+		return FText(LOCTEXT("Layer_DisabledDistance_Tooltip", "Distance Streaming Disabled"));
+	}
+}
+
 #undef LOCTEXT_NAMESPACE
diff --git a/Engine/Source/Editor/WorldBrowser/Private/Tiles/SWorldLayers.h b/Engine/Source/Editor/WorldBrowser/Private/Tiles/SWorldLayers.h
--- a/Engine/Source/Editor/WorldBrowser/Private/Tiles/SWorldLayers.h
+++ b/Engine/Source/Editor/WorldBrowser/Private/Tiles/SWorldLayers.h
@@ -25,6 +25,7 @@ public:
 	FReply OnDoubleClicked();
 	FReply OnCtrlClicked();
 	TSharedRef<SWidget> GetRightClickMenu();
+	FText GetToolTipText() const;
 			
 private:
 	/** The data for this item */
diff --git a/Engine/Source/Editor/WorldBrowser/Private/Tiles/SWorldTileItem.cpp b/Engine/Source/Editor/WorldBrowser/Private/Tiles/SWorldTileItem.cpp
--- a/Engine/Source/Editor/WorldBrowser/Private/Tiles/SWorldTileItem.cpp
+++ b/Engine/Source/Editor/WorldBrowser/Private/Tiles/SWorldTileItem.cpp
@@ -348,7 +348,7 @@ TSharedPtr<IToolTip> SWorldTileItem::GetToolTip()
 FVector2D SWorldTileItem::GetDesiredSizeForMarquee() const
 {
 	// we don't want to select items in non visible layers
-	if (!WorldModel->PassesAllFilters(TileModel))
+	if (!WorldModel->PassesAllFilters(*TileModel))
 	{
 		return FVector2D::ZeroVector;
 	}
@@ -474,7 +474,14 @@ FText SWorldTileItem::GetLevelLayerNameText() const
 
 FText SWorldTileItem::GetLevelLayerDistanceText() const
 {
-	return FText::AsNumber(TileModel->TileDetails->Layer.StreamingDistance);
+	if (TileModel->TileDetails->Layer.DistanceStreamingEnabled)
+	{
+		return FText::AsNumber(TileModel->TileDetails->Layer.StreamingDistance);
+	}
+	else
+	{
+		return FText(LOCTEXT("DistanceStreamingDisabled", "Distance Streaming Disabled"));
+	}
 }
 
 bool SWorldTileItem::IsItemEditable() const
diff --git a/Engine/Source/Editor/WorldBrowser/Private/Tiles/WorldTileCollectionModel.cpp b/Engine/Source/Editor/WorldBrowser/Private/Tiles/WorldTileCollectionModel.cpp
--- a/Engine/Source/Editor/WorldBrowser/Private/Tiles/WorldTileCollectionModel.cpp
+++ b/Engine/Source/Editor/WorldBrowser/Private/Tiles/WorldTileCollectionModel.cpp
@@ -268,10 +268,10 @@ TSharedPtr<FLevelDragDropOp> FWorldTileCollectionModel::CreateDragDropOp() const
 	return FLevelCollectionModel::CreateDragDropOp();
 }
 
-bool FWorldTileCollectionModel::PassesAllFilters(TSharedPtr<FLevelModel> Item) const
+bool FWorldTileCollectionModel::PassesAllFilters(const FLevelModel& Item) const
 {
-	TSharedPtr<FWorldTileModel> Tile = StaticCastSharedPtr<FWorldTileModel>(Item);
-	if (!Tile->IsInLayersList(SelectedLayers))
+	const FWorldTileModel& Tile = static_cast<const FWorldTileModel&>(Item);
+	if (!Tile.IsInLayersList(SelectedLayers))
 	{
 		return false;
 	}
@@ -2112,13 +2112,13 @@ bool FWorldTileCollectionModel::GenerateLODLevels(FLevelModelList InLevelList, i
 			UWorld* LODWorld = UWorld::FindWorldInPackage(LODPackage);
 			if (LODWorld)
 			{
-				LODWorld->ClearFlags(RF_Standalone);
+				LODWorld->ClearFlags(RF_Public | RF_Standalone);
 				LODWorld->DestroyWorld(false);
 			}
-
+			
 			// Create a new world
 			LODWorld = UWorld::CreateWorld(EWorldType::None, false, FPackageName::GetShortFName(LODPackage->GetFName()), LODPackage);
-			LODWorld->SetFlags(RF_Standalone);
+			LODWorld->SetFlags(RF_Public | RF_Standalone);
 
 			for (FAssetToSpawnInfo& AssetInfo : AssetsToSpawn)
 			{
@@ -2143,12 +2143,12 @@ bool FWorldTileCollectionModel::GenerateLODLevels(FLevelModelList InLevelList, i
 			// Save generated level
 			if (FEditorFileUtils::PromptToCheckoutLevels(false, LODWorld->PersistentLevel))
 			{
-				FEditorFileUtils::SaveLevel(LODWorld->PersistentLevel, *LODLevelFileName);
-				FAssetRegistryModule::AssetCreated(LODWorld->PersistentLevel);
+				FEditorFileUtils::SaveLevel(LODWorld->PersistentLevel, LODLevelFileName);
+				FAssetRegistryModule::AssetCreated(LODWorld);
 			}
 			
 			// Destroy the new world we created and collect the garbage
-			LODWorld->ClearFlags(RF_Standalone);
+			LODWorld->ClearFlags(RF_Public | RF_Standalone);
 			LODWorld->DestroyWorld(false);
 			CollectGarbage(GARBAGE_COLLECTION_KEEPFLAGS);
 		}
diff --git a/Engine/Source/Editor/WorldBrowser/Private/Tiles/WorldTileCollectionModel.h b/Engine/Source/Editor/WorldBrowser/Private/Tiles/WorldTileCollectionModel.h
--- a/Engine/Source/Editor/WorldBrowser/Private/Tiles/WorldTileCollectionModel.h
+++ b/Engine/Source/Editor/WorldBrowser/Private/Tiles/WorldTileCollectionModel.h
@@ -31,7 +31,7 @@ public:
 	virtual void TranslateLevels(const FLevelModelList& InList, FVector2D InAbsoluteDelta, bool bSnapDelta = true)  override;
 	virtual FVector2D SnapTranslationDelta(const FLevelModelList& InList, FVector2D InTranslationDelta, bool bBoundsSnapping, float InSnappingValue) override;
 	virtual TSharedPtr<FLevelDragDropOp> CreateDragDropOp() const override;
-	virtual bool PassesAllFilters(TSharedPtr<FLevelModel> InLevelModel) const override;
+	virtual bool PassesAllFilters(const FLevelModel& InLevelModel) const override;
 	virtual void BuildHierarchyMenu(FMenuBuilder& InMenuBuilder) const override;
 	virtual void CustomizeFileMainMenu(FMenuBuilder& InMenuBuilder) const override;
 	virtual bool GetPlayerView(FVector& Location, FRotator& Rotation) const override;
diff --git a/Engine/Source/Editor/WorldBrowser/Private/Tiles/WorldTileDetails.cpp b/Engine/Source/Editor/WorldBrowser/Private/Tiles/WorldTileDetails.cpp
--- a/Engine/Source/Editor/WorldBrowser/Private/Tiles/WorldTileDetails.cpp
+++ b/Engine/Source/Editor/WorldBrowser/Private/Tiles/WorldTileDetails.cpp
@@ -35,6 +35,7 @@ void UWorldTileDetails::SetInfo(const FWorldTileInfo& Info, ULevel* Level)
 	Layer				= Info.Layer;
 	Bounds				= Info.Bounds;
 	ZOrder				= Info.ZOrder;
+	bHideInTileView		= Info.bHideInTileView;
 
 	// Sync LOD settings
 	NumLOD				= Info.LODList.Num();
@@ -65,6 +66,7 @@ FWorldTileInfo UWorldTileDetails::GetInfo() const
 	Info.Layer					= Layer;
 	Info.Bounds					= Bounds;
 	Info.ZOrder					= ZOrder;
+	Info.bHideInTileView		= bHideInTileView;
 
 	// Sync LOD settings
 	Info.LODList.SetNum(FMath::Clamp(NumLOD, 0, WORLDTILE_LOD_MAX_INDEX));
@@ -121,6 +123,10 @@ void UWorldTileDetails::PostEditChangeProperty( struct FPropertyChangedEvent& Pr
 	{
 		ZOrderChangedEvent.Broadcast();
 	}
+	else if (PropertyName == GET_MEMBER_NAME_CHECKED(UWorldTileDetails, bHideInTileView))
+	{
+		HideInTileViewChangedEvent.Broadcast();
+	}
 }
 
 
diff --git a/Engine/Source/Editor/WorldBrowser/Private/Tiles/WorldTileDetails.h b/Engine/Source/Editor/WorldBrowser/Private/Tiles/WorldTileDetails.h
--- a/Engine/Source/Editor/WorldBrowser/Private/Tiles/WorldTileDetails.h
+++ b/Engine/Source/Editor/WorldBrowser/Private/Tiles/WorldTileDetails.h
@@ -62,6 +62,10 @@ class UWorldTileDetails : public UObject
 	// Tile sorting order
 	UPROPERTY(Category=Tile, EditAnywhere, meta=(ClampMin = "-1000", ClampMax = "1000", UIMin = "-1000", UIMax = "1000"))
 	int32							ZOrder;
+
+	// Whether to hide tile in the world composition tile view
+	UPROPERTY(Category=Tile, EditAnywhere)
+	bool							bHideInTileView;
 	
 	// LOD entries number
 	UPROPERTY(Category=LODSettings, EditAnywhere, meta=(ClampMin = "0", ClampMax = "4", UIMin = "0", UIMax = "4"))
@@ -95,6 +99,7 @@ class UWorldTileDetails : public UObject
 	FTilePropertyChanged			ParentPackageNameChangedEvent;
 	FTilePropertyChanged			LODSettingsChangedEvent;
 	FTilePropertyChanged			ZOrderChangedEvent;
+	FTilePropertyChanged			HideInTileViewChangedEvent;
 	
 	// UObject interface
 	virtual void PostEditChangeProperty( struct FPropertyChangedEvent& PropertyChangedEvent) override;
diff --git a/Engine/Source/Editor/WorldBrowser/Private/Tiles/WorldTileDetailsCustomization.cpp b/Engine/Source/Editor/WorldBrowser/Private/Tiles/WorldTileDetailsCustomization.cpp
--- a/Engine/Source/Editor/WorldBrowser/Private/Tiles/WorldTileDetailsCustomization.cpp
+++ b/Engine/Source/Editor/WorldBrowser/Private/Tiles/WorldTileDetailsCustomization.cpp
@@ -104,6 +104,10 @@ void FWorldTileDetailsCustomization::CustomizeDetails(IDetailLayoutBuilder& Deta
 		TileCategory.AddProperty(GET_MEMBER_NAME_CHECKED(UWorldTileDetails, ZOrder))
 			.IsEnabled(IsPropertyEnabled);
 		
+		// Hide in tile view
+		TileCategory.AddProperty(GET_MEMBER_NAME_CHECKED(UWorldTileDetails, bHideInTileView))
+			.IsEnabled(IsPropertyEnabled);
+		
 		// bTileEditable (invisible property to control other properties editable state)
 		TileEditableHandle = DetailLayoutBuilder.GetProperty(GET_MEMBER_NAME_CHECKED(UWorldTileDetails, bTileEditable));
 		TileCategory.AddProperty(TileEditableHandle)
diff --git a/Engine/Source/Editor/WorldBrowser/Private/Tiles/WorldTileModel.cpp b/Engine/Source/Editor/WorldBrowser/Private/Tiles/WorldTileModel.cpp
--- a/Engine/Source/Editor/WorldBrowser/Private/Tiles/WorldTileModel.cpp
+++ b/Engine/Source/Editor/WorldBrowser/Private/Tiles/WorldTileModel.cpp
@@ -36,6 +36,7 @@ FWorldTileModel::FWorldTileModel(const TWeakObjectPtr<UEditorEngine>& InEditor,
 	TileDetails->ParentPackageNameChangedEvent.AddRaw(this, &FWorldTileModel::OnParentPackageNamePropertyChanged);
 	TileDetails->LODSettingsChangedEvent.AddRaw(this, &FWorldTileModel::OnLODSettingsPropertyChanged);
 	TileDetails->ZOrderChangedEvent.AddRaw(this, &FWorldTileModel::OnZOrderPropertyChanged);
+	TileDetails->HideInTileViewChangedEvent.AddRaw(this, &FWorldTileModel::OnHideInTileViewChanged);
 			
 	// Initialize tile details
 	if (WorldComposition->GetTilesList().IsValidIndex(TileIdx))
@@ -652,6 +653,11 @@ void FWorldTileModel::OnParentChanged()
 	OnLevelInfoUpdated();
 }
 
+bool FWorldTileModel::IsVisibleInCompositionView() const
+{
+	return !TileDetails->bHideInTileView && LevelCollectionModel.PassesAllFilters(*this);
+}
+
 void FWorldTileModel::OnLevelBoundsActorUpdated()
 {
 	Update();
@@ -759,6 +765,11 @@ void FWorldTileModel::OnZOrderPropertyChanged()
 	OnLevelInfoUpdated();
 }
 
+void FWorldTileModel::OnHideInTileViewChanged()
+{
+	OnLevelInfoUpdated();
+}
+
 bool FWorldTileModel::CreateAdjacentLandscapeProxy(ALandscapeProxy* SourceLandscape, FIntPoint SourceTileOffset, FWorldTileModel::EWorldDirections InWhere)
 {
 	if (!IsLoaded())	
diff --git a/Engine/Source/Editor/WorldBrowser/Private/Tiles/WorldTileModel.h b/Engine/Source/Editor/WorldBrowser/Private/Tiles/WorldTileModel.h
--- a/Engine/Source/Editor/WorldBrowser/Private/Tiles/WorldTileModel.h
+++ b/Engine/Source/Editor/WorldBrowser/Private/Tiles/WorldTileModel.h
@@ -83,6 +83,7 @@ public:
 	virtual void OnLevelAddedToWorld(ULevel* InLevel) override;
 	virtual void OnLevelRemovedFromWorld() override;
 	virtual void OnParentChanged() override;
+	virtual bool IsVisibleInCompositionView() const override;
 	// FLevelModel interface end
 	
 	/** Adds new streaming level*/
@@ -182,9 +183,12 @@ private:
 	/** Handler for LOD settings changes event from Tile details object  */
 	void OnLODSettingsPropertyChanged();
 	
-	/** Handler for ZOrder chnages event from Tile details object  */
+	/** Handler for ZOrder changes event from Tile details object  */
 	void OnZOrderPropertyChanged();
 
+	/** Handler for bHideInTileView changes event from Tile details object  */
+	void OnHideInTileViewChanged();
+
 	/** Set the asset name based on the passed in package name */
 	void SetAssetName(const FName& PackageName);
 
diff --git a/Engine/Source/Runtime/AIModule/Classes/AIController.h b/Engine/Source/Runtime/AIModule/Classes/AIController.h
--- a/Engine/Source/Runtime/AIModule/Classes/AIController.h
+++ b/Engine/Source/Runtime/AIModule/Classes/AIController.h
@@ -346,9 +346,8 @@ public:
 	static bool bAIIgnorePlayers;
 
 public:
-	/** Returns NavComponent subobject **/
-	UNavigationComponent* GetNavComponent() const;
 	/** Returns PathFollowingComponent subobject **/
+	UFUNCTION(BlueprintCallable, Category="AI|Navigation")
 	UPathFollowingComponent* GetPathFollowingComponent() const;
 	/** Returns ActionsComp subobject **/
 	UPawnActionsComponent* GetActionsComp() const;
diff --git a/Engine/Source/Runtime/AIModule/Classes/Navigation/CrowdFollowingComponent.h b/Engine/Source/Runtime/AIModule/Classes/Navigation/CrowdFollowingComponent.h
--- a/Engine/Source/Runtime/AIModule/Classes/Navigation/CrowdFollowingComponent.h
+++ b/Engine/Source/Runtime/AIModule/Classes/Navigation/CrowdFollowingComponent.h
@@ -21,7 +21,7 @@ namespace ECrowdAvoidanceQuality
 	};
 }
 
-UCLASS()
+UCLASS(BlueprintType)
 class AIMODULE_API UCrowdFollowingComponent : public UPathFollowingComponent, public ICrowdAgentInterface
 {
 	GENERATED_UCLASS_BODY()
@@ -61,6 +61,7 @@ class AIMODULE_API UCrowdFollowingComponent : public UPathFollowingComponent, pu
 	virtual void ApplyCrowdAgentPosition(const FVector& NewPosition);
 
 	/** master switch for crowd steering & avoidance */
+	UFUNCTION(BlueprintCallable, Category = "Crowd")
 	virtual void SuspendCrowdSteering(bool bSuspend);
 
 	/** switch between crowd simulation and parent implementation (following path segments) */
diff --git a/Engine/Source/Runtime/AIModule/Private/Perception/PawnSensingComponent.cpp b/Engine/Source/Runtime/AIModule/Private/Perception/PawnSensingComponent.cpp
--- a/Engine/Source/Runtime/AIModule/Private/Perception/PawnSensingComponent.cpp
+++ b/Engine/Source/Runtime/AIModule/Private/Perception/PawnSensingComponent.cpp
@@ -335,7 +335,7 @@ bool UPawnSensingComponent::CouldSeePawn(const APawn *Other, bool bMaySkipChecks
 
 	// check field of view
 	FVector const SelfToOtherDir = SelfToOther.GetSafeNormal();
-	FVector const MyFacingDir = Owner->GetActorRotation().Vector();
+	FVector const MyFacingDir = GetSensorRotation().Vector();
 
 // 	UE_LOG(LogPath, Warning, TEXT("DotProductFacing: %f, PeripheralVisionCosine: %f"), SelfToOtherDir | MyFacingDir, PeripheralVisionCosine);
 
diff --git a/Engine/Source/Runtime/Analytics/AnalyticsET/Private/AnalyticsET.cpp b/Engine/Source/Runtime/Analytics/AnalyticsET/Private/AnalyticsET.cpp
--- a/Engine/Source/Runtime/Analytics/AnalyticsET/Private/AnalyticsET.cpp
+++ b/Engine/Source/Runtime/Analytics/AnalyticsET/Private/AnalyticsET.cpp
@@ -216,6 +216,8 @@ FAnalyticsProviderET::FAnalyticsProviderET(const FAnalyticsET::Config& ConfigVal
 
 bool FAnalyticsProviderET::Tick(float DeltaSeconds)
 {
+	QUICK_SCOPE_CYCLE_COUNTER(STAT_FAnalyticsProviderET_Tick);
+
 	if (CachedEvents.Num() > 0)
 	{
 		// Countdown to flush
diff --git a/Engine/Source/Runtime/Core/Private/GenericPlatform/GenericPlatformMemory.cpp b/Engine/Source/Runtime/Core/Private/GenericPlatform/GenericPlatformMemory.cpp
--- a/Engine/Source/Runtime/Core/Private/GenericPlatform/GenericPlatformMemory.cpp
+++ b/Engine/Source/Runtime/Core/Private/GenericPlatform/GenericPlatformMemory.cpp
@@ -10,18 +10,17 @@ DEFINE_STAT(MCR_Physical);
 DEFINE_STAT(MCR_GPU);
 DEFINE_STAT(MCR_TexturePool);
 
-// Must match values in the MemoryProfiler2.FMemoryAllocationStatsV4
-DECLARE_MEMORY_STAT(TEXT("Total Physical"),		STAT_TotalPhysical,STATGROUP_MemoryPlatform);
-DECLARE_MEMORY_STAT(TEXT("Total Virtual"),		STAT_TotalVirtual,STATGROUP_MemoryPlatform);
-DECLARE_MEMORY_STAT(TEXT("Page Size"),			STAT_PageSize,STATGROUP_MemoryPlatform);
-DECLARE_MEMORY_STAT(TEXT("Total Physical GB"),	STAT_TotalPhysicalGB,STATGROUP_MemoryPlatform);
-
-DECLARE_MEMORY_STAT(TEXT("Available Physical"),	STAT_AvailablePhysical,STATGROUP_MemoryPlatform);
-DECLARE_MEMORY_STAT(TEXT("Available Virtual"),	STAT_AvailableVirtual,STATGROUP_MemoryPlatform);
-DECLARE_MEMORY_STAT(TEXT("Used Physical"),		STAT_UsedPhysical,STATGROUP_MemoryPlatform);
-DECLARE_MEMORY_STAT(TEXT("Peak Used Physical"),	STAT_PeakUsedPhysical,STATGROUP_MemoryPlatform);
-DECLARE_MEMORY_STAT(TEXT("Used Virtual"),		STAT_UsedVirtual,STATGROUP_MemoryPlatform);
-DECLARE_MEMORY_STAT(TEXT("Peak Used Virtual"),	STAT_PeakUsedVirtual,STATGROUP_MemoryPlatform);
+DEFINE_STAT(STAT_TotalPhysical);
+DEFINE_STAT(STAT_TotalVirtual);
+DEFINE_STAT(STAT_PageSize);
+DEFINE_STAT(STAT_TotalPhysicalGB);
+
+DEFINE_STAT(STAT_AvailablePhysical);
+DEFINE_STAT(STAT_AvailableVirtual);
+DEFINE_STAT(STAT_UsedPhysical);
+DEFINE_STAT(STAT_PeakUsedPhysical);
+DEFINE_STAT(STAT_UsedVirtual);
+DEFINE_STAT(STAT_PeakUsedVirtual);
 
 FGenericPlatformMemoryStats::FGenericPlatformMemoryStats()
 	: FGenericPlatformMemoryConstants( FPlatformMemory::GetConstants() )
@@ -71,16 +70,16 @@ void FGenericPlatformMemory::GetStatsForMallocProfiler( FGenericMemoryStats& out
 	FPlatformMemoryStats Stats = FPlatformMemory::GetStats();
 
 	// Base common stats for all platforms.
-	out_Stats.Add( GET_STATDESCRIPTION( STAT_TotalPhysical ), Stats.TotalPhysical );
-	out_Stats.Add( GET_STATDESCRIPTION( STAT_TotalVirtual ), Stats.TotalVirtual );
-	out_Stats.Add( GET_STATDESCRIPTION( STAT_PageSize ), Stats.PageSize );
-	out_Stats.Add( GET_STATDESCRIPTION( STAT_TotalPhysicalGB ), (SIZE_T)Stats.TotalPhysicalGB );
-	out_Stats.Add( GET_STATDESCRIPTION( STAT_AvailablePhysical ), Stats.AvailablePhysical );
-	out_Stats.Add( GET_STATDESCRIPTION( STAT_AvailableVirtual ), Stats.AvailableVirtual );
-	out_Stats.Add( GET_STATDESCRIPTION( STAT_UsedPhysical ), Stats.UsedPhysical );
-	out_Stats.Add( GET_STATDESCRIPTION( STAT_PeakUsedPhysical ), Stats.PeakUsedPhysical );
-	out_Stats.Add( GET_STATDESCRIPTION( STAT_UsedVirtual ), Stats.UsedVirtual );
-	out_Stats.Add( GET_STATDESCRIPTION( STAT_PeakUsedVirtual ), Stats.PeakUsedVirtual );
+	out_Stats.Add(TEXT("Total Physical"), Stats.TotalPhysical );
+	out_Stats.Add(TEXT("Total Virtual"), Stats.TotalVirtual );
+	out_Stats.Add(TEXT("Page Size"), Stats.PageSize );
+	out_Stats.Add(TEXT("Total Physical GB"), (SIZE_T)Stats.TotalPhysicalGB );
+	out_Stats.Add(TEXT("Available Physical"), Stats.AvailablePhysical );
+	out_Stats.Add(TEXT("Available Virtual"), Stats.AvailableVirtual );
+	out_Stats.Add(TEXT("Used Physical"), Stats.UsedPhysical );
+	out_Stats.Add(TEXT("Peak Used Physical"), Stats.PeakUsedPhysical );
+	out_Stats.Add(TEXT("Used Virtual"), Stats.UsedVirtual );
+	out_Stats.Add(TEXT("Peak Used Virtual"), Stats.PeakUsedVirtual );
 #endif // STATS
 }
 
diff --git a/Engine/Source/Runtime/Core/Private/HAL/ConsoleManager.cpp b/Engine/Source/Runtime/Core/Private/HAL/ConsoleManager.cpp
--- a/Engine/Source/Runtime/Core/Private/HAL/ConsoleManager.cpp
+++ b/Engine/Source/Runtime/Core/Private/HAL/ConsoleManager.cpp
@@ -2018,7 +2018,7 @@ static TAutoConsoleVariable<int32> CVarTonemapper709(
 	TEXT("r.Tonemapper709"),
 	0,
 	TEXT("0: use sRGB on PC monitor output\n")
-	TEXT("1: use Rec.709 for HDTV/projector output (requires r.TonemapperPhoto 1)"),
+	TEXT("1: use Rec.709 for HDTV/projector output"),
 	ECVF_Scalability | ECVF_RenderThreadSafe);
 
 
diff --git a/Engine/Source/Runtime/Core/Private/HAL/MemoryBase.cpp b/Engine/Source/Runtime/Core/Private/HAL/MemoryBase.cpp
--- a/Engine/Source/Runtime/Core/Private/HAL/MemoryBase.cpp
+++ b/Engine/Source/Runtime/Core/Private/HAL/MemoryBase.cpp
@@ -66,21 +66,23 @@ void FMalloc::InitializeStatsMetadata()
 
 void FMalloc::UpdateStats()
 {
+#if	STATS
 	GetCurrentFrameCalls().Update();
 
 	SET_DWORD_STAT( STAT_MallocCalls, GetCurrentFrameCalls().MallocCalls );
 	SET_DWORD_STAT( STAT_ReallocCalls, GetCurrentFrameCalls().ReallocCalls );
 	SET_DWORD_STAT( STAT_FreeCalls, GetCurrentFrameCalls().FreeCalls );
 	SET_DWORD_STAT( STAT_TotalAllocatorCalls, GetCurrentFrameCalls().AllocatorCalls );
+#endif
 }
 
 
 void FMalloc::GetAllocatorStats( FGenericMemoryStats& out_Stats )
 {
 #if	STATS
-	out_Stats.Add( GET_STATDESCRIPTION( STAT_MallocCalls ), GetCurrentFrameCalls().MallocCalls );
-	out_Stats.Add( GET_STATDESCRIPTION( STAT_ReallocCalls ), GetCurrentFrameCalls().ReallocCalls );
-	out_Stats.Add( GET_STATDESCRIPTION( STAT_FreeCalls ), GetCurrentFrameCalls().FreeCalls );
-	out_Stats.Add( GET_STATDESCRIPTION( STAT_TotalAllocatorCalls ), GetCurrentFrameCalls().AllocatorCalls );
+	out_Stats.Add( TEXT("Malloc calls"), GetCurrentFrameCalls().MallocCalls );
+	out_Stats.Add( TEXT("Realloc calls"), GetCurrentFrameCalls().ReallocCalls );
+	out_Stats.Add( TEXT("Free calls"), GetCurrentFrameCalls().FreeCalls );
+	out_Stats.Add( TEXT("Total Allocator calls"), GetCurrentFrameCalls().AllocatorCalls );
 #endif // STATS
 }
diff --git a/Engine/Source/Runtime/Core/Private/HAL/ThreadingBase.cpp b/Engine/Source/Runtime/Core/Private/HAL/ThreadingBase.cpp
--- a/Engine/Source/Runtime/Core/Private/HAL/ThreadingBase.cpp
+++ b/Engine/Source/Runtime/Core/Private/HAL/ThreadingBase.cpp
@@ -162,6 +162,8 @@ void FSingleThreadManager::RemoveThread(FFakeThread* Thread)
 
 void FSingleThreadManager::Tick()
 {
+	QUICK_SCOPE_CYCLE_COUNTER(STAT_FSingleThreadManager_Tick);
+
 	// Tick all registered threads.
 	for (int32 RunnableIndex = 0; RunnableIndex < ThreadList.Num(); ++RunnableIndex)
 	{
diff --git a/Engine/Source/Runtime/Core/Private/Misc/CoreMisc.cpp b/Engine/Source/Runtime/Core/Private/Misc/CoreMisc.cpp
--- a/Engine/Source/Runtime/Core/Private/Misc/CoreMisc.cpp
+++ b/Engine/Source/Runtime/Core/Private/Misc/CoreMisc.cpp
@@ -399,7 +399,14 @@ bool FFileHelper::CreateBitmap( const TCHAR* Pattern, int32 SourceWidth, int32 S
 		IH.biHeight             = INTEL_ORDER32((uint32) Height);
 		IH.biPlanes             = INTEL_ORDER16((uint16) 1);
 		IH.biBitCount           = INTEL_ORDER16((uint16) BytesPerPixel * 8);
-		IH.biCompression        = INTEL_ORDER32((uint32) 3); //BI_BITFIELDS
+		if(bInWriteAlpha)
+		{
+			IH.biCompression    = INTEL_ORDER32((uint32) 3); //BI_BITFIELDS
+		}
+		else
+		{
+			IH.biCompression    = INTEL_ORDER32((uint32) 0); //BI_RGB
+		}
 		IH.biSizeImage          = INTEL_ORDER32((uint32) BytesPerLine * Height);
 		IH.biXPelsPerMeter      = INTEL_ORDER32((uint32) 0);
 		IH.biYPelsPerMeter      = INTEL_ORDER32((uint32) 0);
@@ -765,46 +772,6 @@ FTicker& FTicker::GetCoreTicker()
 	return Singleton;
 }
 
-// Josh, delete this once you have a real usage of tickers
-#if 0
-
-struct FTestTicker
-{
-	FTestTicker()
-		: StartTime(FPlatformTime::Seconds())
-	{
-		FTicker::GetCoreTicker().AddTicker(FTickerDelegate::CreateRaw(this, &FTestTicker::OneSecond), 1.0f);
-		FTicker::GetCoreTicker().AddTicker(FTickerDelegate::CreateRaw(this, &FTestTicker::FiveSecond), 5.0f);
-		FTicker::GetCoreTicker().AddTicker(FTickerDelegate::CreateRaw(this, &FTestTicker::EveryFrame), 0.0f);
-		FTicker::GetCoreTicker().AddTicker(FTickerDelegate::CreateRaw(this, &FTestTicker::TwentyTimes), 0.0f);
-	}
-
-	bool OneSecond(float DeltaTime)
-	{
-		UE_LOG(LogTemp, Display, TEXT("OneSecond %f"), float(FPlatformTime::Seconds() - StartTime));
-		return true;
-	}
-	bool FiveSecond(float DeltaTime)
-	{
-		UE_LOG(LogTemp, Display, TEXT("FiveSecond %f"), float(FPlatformTime::Seconds() - StartTime));
-		return true;
-	}
-	bool EveryFrame(float DeltaTime)
-	{
-		UE_LOG(LogTemp, Display, TEXT("EveryFrame %f  (%f)"), float(FPlatformTime::Seconds() - StartTime), DeltaTime);
-		return true;
-	}
-	bool TwentyTimes(float DeltaTime)
-	{
-		static int32 Count = 0;
-		UE_LOG(LogTemp, Display, TEXT("TwentyTimes %f"), float(FPlatformTime::Seconds() - StartTime));
-		return ++Count < 20;
-	}
-
-	double StartTime;
-} TestTicker;
-
-#endif
 
 /*----------------------------------------------------------------------------
 	Runtime functions.
