Distance field AO and Ray traced shadows on PS4.  DFAO costs nearly 4ms in a mostly empty map so needs a lot of optimization.
* Added r.AOHistoryDistanceThreshold with default to 20, can be tweaked to avoid ghosting behind characters while not causing too much flickering
* Distance fields are now culled to each surface cache record before ray tracing for AO, saves a bit of GPU time

diff --git a/Engine/Shaders/DistanceFieldAOShared.usf b/Engine/Shaders/DistanceFieldAOShared.usf
--- a/Engine/Shaders/DistanceFieldAOShared.usf
+++ b/Engine/Shaders/DistanceFieldAOShared.usf
@@ -30,6 +30,8 @@ float AOMaxViewDistance;
 
 #define NUM_CULLED_OBJECT_LISTS 3
 
+#define FADE_RECORDS_OVER_TIME 0
+
 void GetPhaseParameters(uint Phase, out uint StartIndex, out uint EndIndex)
 {
 	if (Phase == 0)
@@ -97,10 +99,10 @@ void GetDownsampledGBuffer(float2 ScreenUV, out float3 OutNormal, out float OutD
 	OutNormal.z *= abs(TextureValue.z) > .75f ? -1 : 1;
 }
 
-int CurrentLevelDownsampleFactor;
+uint CurrentLevelDownsampleFactor;
 float2 AOBufferSize;
 
-int DownsampleFactorToBaseLevel;
+uint DownsampleFactorToBaseLevel;
 float2 BaseLevelTexelSize;
 
 bool SphereIntersectCone(float4 SphereCenterAndRadius, float3 ConeVertex, float3 ConeAxis, float ConeAngleCos, float ConeAngleSin)
diff --git a/Engine/Shaders/DistanceFieldGlobalIllumination.usf b/Engine/Shaders/DistanceFieldGlobalIllumination.usf
--- a/Engine/Shaders/DistanceFieldGlobalIllumination.usf
+++ b/Engine/Shaders/DistanceFieldGlobalIllumination.usf
@@ -207,7 +207,7 @@ void VPLPlacementCS(
 				float3 IntersectionNormal = ComputeDistanceFieldNormal(IntersectionPosition, ObjectIndex);
 
 				uint VPLArrayStartIndex;
-				InterlockedAdd(RWVPLParameterBuffer[1], 1, VPLArrayStartIndex);
+				InterlockedAdd(RWVPLParameterBuffer[1], 1U, VPLArrayStartIndex);
 
 				float3 DiffuseColor = .5f;
 				float CellExtent = VPLPlacementCameraRadius * InvPlacementGridSize.x;
@@ -277,7 +277,7 @@ void CullVPLsForViewCS(
 			&& ViewFrustumIntersectSphere(BoundingSphere.xyz, BoundingSphere.w + AOMaxDistance))
 		{
 			uint DestStartVPLIndex;
-			InterlockedAdd(RWCulledVPLParameterBuffer[1], 1, DestStartVPLIndex);
+			InterlockedAdd(RWCulledVPLParameterBuffer[1], 1U, DestStartVPLIndex);
 
 			uint DestIndex = DestStartVPLIndex * VPL_DATA_STRIDE;
 			RWCulledVPLData[DestIndex + 0] = BoundingSphere;
@@ -994,9 +994,7 @@ void ComputeIrradianceCS(
 	if (RecordIndex < NumRecords)
 	{
 		float3 RecordWorldNormal = IrradianceCacheNormal[RecordIndex].xyz;
-
-		float4 PositionAndRadius = IrradianceCachePositionRadius[RecordIndex];
-		float3 RecordWorldPosition = PositionAndRadius.xyz;
+		float3 RecordWorldPosition = IrradianceCachePositionRadius[RecordIndex].xyz;
 		Irradiance = GatherIrradianceFromVPLs(RecordWorldPosition, RecordWorldNormal, RelativeRecordIndex, RecordIndex, ThreadIndex);
 	}
 
diff --git a/Engine/Shaders/DistanceFieldLightingShared.usf b/Engine/Shaders/DistanceFieldLightingShared.usf
--- a/Engine/Shaders/DistanceFieldLightingShared.usf
+++ b/Engine/Shaders/DistanceFieldLightingShared.usf
@@ -200,7 +200,7 @@ float4x4 LoadObjectLocalToWorld(uint ObjectIndex)
 // x = Offset in global buffer, y = NumLOD0, z = NumSurfels (all LODs), W = instance index
 uint4 LoadObjectSurfelCoordinate(uint ObjectIndex)
 {
-	return CulledObjectData.Load(ObjectIndex * CULLED_OBJECT_DATA_STRIDE + 15);
+	return (uint4)CulledObjectData.Load(ObjectIndex * CULLED_OBJECT_DATA_STRIDE + 15);
 }
 
 void LoadObjectViewSpaceBox(uint ObjectIndex, out float3 ObjectViewSpaceMin, out float3 ObjectViewSpaceMax)
diff --git a/Engine/Shaders/DistanceFieldShadowing.usf b/Engine/Shaders/DistanceFieldShadowing.usf
--- a/Engine/Shaders/DistanceFieldShadowing.usf
+++ b/Engine/Shaders/DistanceFieldShadowing.usf
@@ -128,7 +128,7 @@ void CullObjectsForShadowCS(
 			|| ShadowBoundingSphere.w > 0 && dot(ShadowBoundingSphere.xyz - ObjectBoundingSphere.xyz, ShadowBoundingSphere.xyz - ObjectBoundingSphere.xyz) < Square(ShadowBoundingSphere.w + ObjectBoundingSphere.w))
 		{
 			uint DestIndex;
-			InterlockedAdd(RWObjectIndirectArguments[1], 1, DestIndex);
+			InterlockedAdd(RWObjectIndirectArguments[1], 1U, DestIndex);
 			CopyCulledObjectData(DestIndex, SourceIndex);
 		}
 	}
@@ -179,7 +179,6 @@ void ClearTilesCS(
 
 struct FShadowObjectCullVertexOutput
 {
-	float4 Position : SV_POSITION;
 	nointerpolation float4 PositionAndRadius : TEXCOORD0;
 	nointerpolation uint ObjectIndex : TEXCOORD1;
 };
@@ -191,7 +190,8 @@ float MinRadius;
 void ShadowObjectCullVS(
 	float4 InPosition : ATTRIBUTE0,
 	uint ObjectIndex : SV_InstanceID,
-	out FShadowObjectCullVertexOutput Output
+	out FShadowObjectCullVertexOutput Output,
+	out float4 OutPosition : SV_POSITION
 	)
 {
 	float4 ObjectPositionAndRadius = LoadObjectPositionAndRadius(ObjectIndex);
@@ -199,7 +199,7 @@ void ShadowObjectCullVS(
 	// MinRadius is for conservative rasterization
 	float EffectiveRadius = (ObjectPositionAndRadius.w + MinRadius) * ConservativeRadiusScale;
 	float3 WorldPosition = InPosition.xyz * EffectiveRadius + ObjectPositionAndRadius.xyz;
-	Output.Position = mul(float4(WorldPosition, 1), WorldToShadow);
+	OutPosition = mul(float4(WorldPosition, 1), WorldToShadow);
 	Output.PositionAndRadius = ObjectPositionAndRadius;
 	Output.ObjectIndex = ObjectIndex;
 } 
@@ -208,7 +208,7 @@ void ShadowObjectCullVS(
 void ShadowObjectCullPS(
 	FShadowObjectCullVertexOutput Input, 
 	in float4 SVPos : SV_POSITION,
-	out float4 OutColor : COLOR0)
+	out float4 OutColor : SV_Target0)
 {
 	OutColor = 0;
 	
diff --git a/Engine/Shaders/DistanceFieldSurfaceCacheLighting.usf b/Engine/Shaders/DistanceFieldSurfaceCacheLighting.usf
--- a/Engine/Shaders/DistanceFieldSurfaceCacheLighting.usf
+++ b/Engine/Shaders/DistanceFieldSurfaceCacheLighting.usf
@@ -12,18 +12,18 @@
 
 struct FObjectCullVertexOutput
 {
-	float4 Position : SV_POSITION;
 	nointerpolation float4 PositionAndRadius : TEXCOORD0;
 	nointerpolation uint ObjectIndex : TEXCOORD1;
 };
-
+ 
 float ConservativeRadiusScale;
 
 /** Used when culling objects into screenspace tile lists */
 void ObjectCullVS(
 	float4 InPosition : ATTRIBUTE0,
 	uint ObjectIndex : SV_InstanceID,
-	out FObjectCullVertexOutput Output
+	out FObjectCullVertexOutput Output,
+	out float4 OutPosition : SV_POSITION
 	)
 {
 	//@todo - implement ConservativelyBoundSphere
@@ -31,7 +31,7 @@ void ObjectCullVS(
 	//@todo - expand to handle conservative rasterization
 	float EffectiveRadius = (ObjectPositionAndRadius.w + AOMaxDistance) * ConservativeRadiusScale;
 	float3 WorldPosition = InPosition.xyz * EffectiveRadius + ObjectPositionAndRadius.xyz;
-	Output.Position = mul(float4(WorldPosition, 1), View.WorldToClip);
+	OutPosition = mul(float4(WorldPosition, 1), View.WorldToClip);
 	Output.PositionAndRadius = ObjectPositionAndRadius;
 	Output.ObjectIndex = ObjectIndex;
 } 
@@ -48,7 +48,7 @@ float2 NumGroups;
 void ObjectCullPS(
 	FObjectCullVertexOutput Input, 
 	in float4 SVPos : SV_POSITION,
-	out float4 OutColor : COLOR0)
+	out float4 OutColor : SV_Target0)
 {
 	OutColor = 0;
 	
@@ -86,7 +86,7 @@ void ObjectCullPS(
 	if (SmallestGroupIndex >= 0)
 	{
 		uint ArrayIndex;
-		InterlockedAdd(RWTileHeadDataUnpacked[TileIndex * 4 + 1 + SmallestGroupIndex], 1U, ArrayIndex);
+		InterlockedAdd(RWTileHeadDataUnpacked[TileIndex * 4 + 1 + (uint)SmallestGroupIndex], 1U, ArrayIndex);
 
 		if (ArrayIndex < MAX_OBJECTS_PER_TILE)
 		{
@@ -131,7 +131,7 @@ void ComputeDistanceFieldNormalCS(
 void WriteDownsampledDepthPS(
 	in float4 UVAndScreenPos : TEXCOORD0, 
 	out float4 OutColor : SV_Target0,
-	out float OutDepth : SV_Depth)
+	out float OutDepth : SV_DEPTH)
 {
 	float2 DistanceFieldUVs = UVAndScreenPos.xy;
 
@@ -150,7 +150,6 @@ void WriteDownsampledDepthPS(
 
 struct FIrradianceCacheSplatVertexOutput
 {
-	float4 Position : SV_POSITION;
 	nointerpolation float4 PositionRadius : TEXCOORD0;
 	nointerpolation float4 NormalAndFade : TEXCOORD1;
 #if FINAL_INTERPOLATION_PASS
@@ -178,21 +177,26 @@ void IrradianceCacheSplatVS(
 	float2 InPosition : ATTRIBUTE0,
 	float2 InUV       : ATTRIBUTE1,
 	uint JobIndex : SV_InstanceID,
-	out FIrradianceCacheSplatVertexOutput Output
+	out FIrradianceCacheSplatVertexOutput Output,
+	out float4 OutPosition : SV_POSITION
 	)
 {
-	float4 PositionAndRadius = IrradianceCachePositionRadius.Load(JobIndex);
-	float OccluderRadius = IrradianceCacheOccluderRadius.Load(JobIndex);
-	PositionAndRadius.w = min(abs(PositionAndRadius.w), OccluderRadius);
+	float4 PositionAndPackedRadius = IrradianceCachePositionRadius[JobIndex];
+
+	float3 RecordPosition = PositionAndPackedRadius.xyz;
+	float RecordRadius = abs(PositionAndPackedRadius.w);
+
+	float OccluderRadius = IrradianceCacheOccluderRadius[JobIndex];
+	RecordRadius = min(RecordRadius, OccluderRadius);
 
-	float ViewToCenterLength = length(PositionAndRadius.xyz - View.ViewOrigin.xyz);
-	float3 NormalizedViewToCenter = (PositionAndRadius.xyz - View.ViewOrigin.xyz) / ViewToCenterLength;
+	float ViewToCenterLength = length(RecordPosition - View.ViewOrigin.xyz);
+	float3 NormalizedViewToCenter = (RecordPosition - View.ViewOrigin.xyz) / ViewToCenterLength;
 
 	// Only allow full interpolation expanding for small samples, it won't be noticeable for large samples but adds a lot of splatting cost
-	float EffectiveInterpolationRadiusScale = lerp(InterpolationRadiusScale, 1, max(saturate(PositionAndRadius.w / 80), .2f));
-	PositionAndRadius.w = max(PositionAndRadius.w, ViewToCenterLength * .01f) * EffectiveInterpolationRadiusScale;
+	float EffectiveInterpolationRadiusScale = lerp(InterpolationRadiusScale, 1, max(saturate(RecordRadius / 80), .2f));
+	RecordRadius = max(RecordRadius, ViewToCenterLength * .01f) * EffectiveInterpolationRadiusScale;
 
-	float OffsetFromCenter = InterpolationBoundingDirection * PositionAndRadius.w;
+	float OffsetFromCenter = InterpolationBoundingDirection * RecordRadius;
 	// Distance from the camera position that NormalizedViewToCenter will intersect the near plane
 	float NearPlaneDistance = View.NearPlane / dot(View.ViewForward, NormalizedViewToCenter);
 
@@ -205,43 +209,60 @@ void IrradianceCacheSplatVS(
 	}
 
 	// Construct a virtual sample position that won't be near plane clipped and will have a positive z after projection
-	float3 VirtualPosition = PositionAndRadius.xyz + OffsetFromCenter * NormalizedViewToCenter;
+	float3 VirtualPosition = RecordPosition + OffsetFromCenter * NormalizedViewToCenter;
 	// Clipping the edge of the circle a bit can save perf, but introduces artifacts in the interpolation
 	float RadiusSlack = 1.0f;
 	// Compute new radius since we moved the center away from the camera, approximate as similar triangles
 	// This is really incorrect because it's not taking into account perspective correction on a sphere - the widest part in screenspace is not at the world space center
-	float VirtualRadius = RadiusSlack * PositionAndRadius.w * (ViewToCenterLength + OffsetFromCenter) / ViewToCenterLength;
+	float VirtualRadius = RadiusSlack * RecordRadius * (ViewToCenterLength + OffsetFromCenter) / ViewToCenterLength;
 
 #if !FINAL_INTERPOLATION_PASS
 	// Combat the mysterious bug where samples from last frame do not cover this frame's pixel even with no camera movement, seems to be a precision issue
 	VirtualRadius += HackExpand * ViewToCenterLength * 4;
 #endif
 
-	float3 RecordNormal = IrradianceCacheNormal.Load(JobIndex).xyz;
+	float4 RecordNormalAndFade = IrradianceCacheNormal[JobIndex];
 
 	float2 CornerUVs = InUV;
 	float3 CornerPosition = VirtualPosition + (View.ViewRight * CornerUVs.x + View.ViewUp * CornerUVs.y) * VirtualRadius;
 
-	Output.Position = mul(float4(CornerPosition.xyz, 1), View.WorldToClip);
+	OutPosition = mul(float4(CornerPosition.xyz, 1), View.WorldToClip);
 
 #if USE_SCREENVECTOR_WORLD_POSITION
-	Output.ScreenVector = mul(float4(Output.Position.xy / Output.Position.w, 1, 0), View.ScreenToWorld).xyz;
+	Output.ScreenVector = mul(float4(OutPosition.xy / OutPosition.w, 1, 0), View.ScreenToWorld).xyz;
 #endif
 
 	// Move the vertex over to compensate for the difference between the low res pixel center and the high res pixel center where shading is done
-	Output.Position.xy += NormalizedOffsetToPixelCenter * Output.Position.w;
+	OutPosition.xy += NormalizedOffsetToPixelCenter * OutPosition.w;
 
-	Output.PositionRadius = PositionAndRadius; 
+	Output.PositionRadius = float4(RecordPosition, RecordRadius); 
 
+#if FADE_RECORDS_OVER_TIME
+	float Fade = 1 - (abs(RecordNormalAndFade.w) - 1);
+	/*
+	FLATTEN
+	if (PositionAndPackedRadius.w < 0)
+	{
+		// Fading out
+		Fade = 1 - Fade;
+	}
+	else
+	{
+		Fade = max(Fade, .001f);
+	}
+	*/
+#else
 	float Fade = 1;
-	Output.NormalAndFade = float4(RecordNormal, Fade);
+#endif
+
+	Output.NormalAndFade = float4(RecordNormalAndFade.xyz, Fade);
 
 #if FINAL_INTERPOLATION_PASS
-	Output.BentNormal = IrradianceCacheBentNormal.Load(JobIndex).xyz;
+	Output.BentNormal = IrradianceCacheBentNormal[JobIndex].xyz;
 	#if SUPPORT_IRRADIANCE
-		float ViewDistance = length(View.ViewOrigin.xyz - PositionAndRadius.xyz);
+		float ViewDistance = length(View.ViewOrigin.xyz - RecordPosition);
 		float FadeAlpha = saturate(.001f * (.95f * AOMaxViewDistance - ViewDistance));
-		Output.Irradiance = IrradianceCacheIrradiance.Load(JobIndex).xyz * FadeAlpha;
+		Output.Irradiance = IrradianceCacheIrradiance[JobIndex].xyz * FadeAlpha;
 	#endif
 #endif
 } 
@@ -253,9 +274,9 @@ float InvMinCosPointBehindPlane;
 void IrradianceCacheSplatPS(
 	FIrradianceCacheSplatVertexOutput Input, 
 	in float4 SVPos : SV_POSITION
-	, out float4 OutBentNormal : COLOR0
+	,out float4 OutBentNormal : SV_Target0
 #if FINAL_INTERPOLATION_PASS && SUPPORT_IRRADIANCE
-	, out float4 OutIrradiance : COLOR1
+	, out float4 OutIrradiance : SV_Target1
 #endif
 	)
 {
@@ -500,6 +521,7 @@ SamplerState BentNormalHistorySampler;
 Texture2D IrradianceHistoryTexture;
 SamplerState IrradianceHistorySampler;
 float HistoryWeight;
+float HistoryDistanceThreshold;
 
 Texture2D VelocityTexture;
 SamplerState VelocityTextureSampler;
@@ -573,7 +595,7 @@ void UpdateHistoryPS(
 	if (any(OldUV > MaxUV) 
 		|| any(OldUV < MinUV)
 		// Discard history if we are shading a new position (newly revealed)
-		|| length(PrevWorldPosition - WorldPosition) > 100)
+		|| length(PrevWorldPosition - WorldPosition) > HistoryDistanceThreshold)
 	{
 		EffectiveHistoryWeight = 0;
 	}
@@ -743,6 +765,7 @@ void AOUpsamplePS(
 
 #if OUTPUT_AO
 	OutBentNormal = float4(length(BentNormal).xxx, 1);
+	//OutBentNormal = float4(BentNormal, 1);
 #else
 	#if OUTPUT_BENT_NORMAL
 		OutBentNormal = float4(BentNormal, 1);
diff --git a/Engine/Shaders/DistanceFieldSurfaceCacheLightingCompute.usf b/Engine/Shaders/DistanceFieldSurfaceCacheLightingCompute.usf
--- a/Engine/Shaders/DistanceFieldSurfaceCacheLightingCompute.usf
+++ b/Engine/Shaders/DistanceFieldSurfaceCacheLightingCompute.usf
@@ -202,7 +202,7 @@ void CullObjectsForViewCS(
 			&& ViewFrustumIntersectSphere(ObjectBoundingSphere.xyz, ObjectBoundingSphere.w + AOMaxDistance))
 		{
 			uint DestIndex;
-			InterlockedAdd(NumGroupObjects, 1, DestIndex);
+			InterlockedAdd(NumGroupObjects, 1U, DestIndex);
 			GroupObjectIndices[DestIndex] = SourceIndex;
 		}
 	}
@@ -536,7 +536,7 @@ void DistanceFieldAOBuildTileListMain(
 		uint NumObjectsIntersecting = SmallTileNumObjects + MediumTileNumObjects + LargeTileNumObjects;
 		InterlockedAdd(RWTileArrayNextAllocation[0], NumObjectsIntersecting, ArrayStart);
 		TileArrayDataStart = ArrayStart;
-		RWTileHeadData[GroupId.y * NumGroups.x + GroupId.x] = uint4(TileArrayDataStart, SmallTileNumObjects, MediumTileNumObjects, LargeTileNumObjects);
+		RWTileHeadData[GroupId.y * (uint)NumGroups.x + GroupId.x] = uint4(TileArrayDataStart, SmallTileNumObjects, MediumTileNumObjects, LargeTileNumObjects);
 	}
 	
 	GroupMemoryBarrierWithGroupSync();
@@ -638,14 +638,16 @@ void PopulateCacheCS(
 			//if (all(DispatchThreadId.xy == uint2(4,3)))
 			{
 				// Allocate a new record and store off attributes of the created record
-				int NextSampleIndex;
-				InterlockedAdd(NumQueuedIrradianceCacheRecords, 1, NextSampleIndex);
+				uint NextSampleIndex;
+				InterlockedAdd(NumQueuedIrradianceCacheRecords, 1u, NextSampleIndex);
 				// W stores max allowed radius, used to limit overdraw from nearby samples placed in the high resolution passes
 				float MaxRadiusScale = bHasDistanceFieldRepresentation > 0 ? .005f : .0005f;
-				// Sign of W stores whether the sample was placed on geometry with a distance field representation or not
-				float RadiusSign = bHasDistanceFieldRepresentation > 0 ? 1 : -1;
+				// Sign of W stores whether the sample is fading in or out
+				float RadiusSign = 1;
 				CachedIrradianceCachePositionRadius[NextSampleIndex] = float4(WorldShadingPosition, RadiusSign * SceneDepth * CurrentLevelDownsampleFactor * MaxRadiusScale);
-				CachedIrradianceCacheNormal[NextSampleIndex] = float4(WorldNormal, bHasHeightfieldRepresentation ? 1 : 0);
+				// abs(w) - 1 stores fade in amount (shifted away from 0 to retain sign when fade amount is 0)
+				// sign(w) stores bHasHeightfieldRepresentation
+				CachedIrradianceCacheNormal[NextSampleIndex] = float4(WorldNormal, 1 * (bHasHeightfieldRepresentation ? 1 : -1));
 				CachedIrradianceCacheTileCoordinate[NextSampleIndex] = TileCoordinate;
 			}
 		}
@@ -703,6 +705,12 @@ groupshared uint SharedMinOcclusionDistance;
 
 groupshared float3 SharedGatheredLighting[FINAL_GATHER_THREADGROUP_SIZE];
 
+#define MAX_RECORD_CULLED_OBJECTS 128
+groupshared uint SharedCulledObjectList[MAX_RECORD_CULLED_OBJECTS];
+groupshared uint NumRecordCulledObjects;
+
+#define CULL_OBJECTS_TO_RECORD 1
+
 /** Computes ambient occlusion for a surface cache record by cone stepping through the nearby object distance fields. */
 [numthreads(FINAL_GATHER_THREADGROUP_SIZE, 1, 1)]
 void ConeTraceOcclusionCS(
@@ -729,6 +737,7 @@ void ConeTraceOcclusionCS(
 		}
 
 		SharedMinOcclusionDistance = asuint(AOMaxDistance);
+		NumRecordCulledObjects = 0;
 	}
 
 	GroupMemoryBarrierWithGroupSync();
@@ -749,6 +758,54 @@ void ConeTraceOcclusionCS(
 
 		uint4 TileHead = GetTileHead(TileCoordinate);
 		uint NumObjectsAffectingTile = TileHead.y + TileHead.z + TileHead.w;
+		uint NumCulledObjects = NumObjectsAffectingTile;
+
+		#if CULL_OBJECTS_TO_RECORD
+
+			LOOP 
+			for (uint ListObjectIndex = ThreadIndex; ListObjectIndex < NumObjectsAffectingTile; ListObjectIndex += FINAL_GATHER_THREADGROUP_SIZE)
+			{
+				if (ListObjectIndex < NumObjectsAffectingTile)
+				{
+					uint ListIndex = 0;
+					uint ArrayIndex = ListObjectIndex;
+
+					FLATTEN
+					if (ListObjectIndex >= TileHead.y + TileHead.z)
+					{
+						ListIndex = 2;
+						ArrayIndex = ListObjectIndex - TileHead.y - TileHead.z;
+					}
+					else if (ListObjectIndex >= TileHead.y)
+					{
+						ListIndex = 1;
+						ArrayIndex = ListObjectIndex - TileHead.y;
+					}
+
+					uint ObjectIndex = TileArrayData.Load((ArrayIndex * TileListGroupSize.x * TileListGroupSize.y + TileHead.x) * NUM_CULLED_OBJECT_LISTS + ListIndex);
+
+					float4 ObjectPositionAndRadius = LoadObjectPositionAndRadius(ObjectIndex);
+					float ObjectDistanceSq = dot(ObjectPositionAndRadius.xyz - WorldShadingPosition, ObjectPositionAndRadius.xyz - WorldShadingPosition);
+
+					BRANCH
+					// Skip tracing objects with a small projected angle 
+					if (ObjectPositionAndRadius.w * ObjectPositionAndRadius.w / ObjectDistanceSq > Square(.25f)
+						// Skip tracing objects outside the max occlusion distance
+						//@todo - box distance
+						&& ObjectDistanceSq < Square(ObjectPositionAndRadius.w + AOMaxDistance))
+					{
+						uint DestIndex;
+						InterlockedAdd(NumRecordCulledObjects, 1U, DestIndex);
+						SharedCulledObjectList[DestIndex] = ObjectIndex;
+					}
+				}
+			}
+
+			GroupMemoryBarrierWithGroupSync();
+
+			NumCulledObjects = NumRecordCulledObjects;
+
+		#endif
 
 		FindBestAxisVectors2(WorldNormal, TangentX, TangentY);
 
@@ -778,30 +835,35 @@ void ConeTraceOcclusionCS(
 #endif
 
 		LOOP 
-		for (uint ListObjectIndex = 0; ListObjectIndex < NumObjectsAffectingTile; ListObjectIndex += SIMULTANEOUSLY_TRACED_OBJECTS)
+		for (uint ListObjectIndex = 0; ListObjectIndex < NumCulledObjects; ListObjectIndex += SIMULTANEOUSLY_TRACED_OBJECTS)
 		{
 			uint EffectiveListObjectIndex = ListObjectIndex + ObjectOffsetIndex;
 
-			if (EffectiveListObjectIndex < NumObjectsAffectingTile 
+			if (EffectiveListObjectIndex < NumCulledObjects 
 				// Ignore extra threads
 				&& ObjectOffsetIndex < SIMULTANEOUSLY_TRACED_OBJECTS)
 			{
-				uint ListIndex = 0;
-				uint ArrayIndex = EffectiveListObjectIndex;
+				#if CULL_OBJECTS_TO_RECORD
+					uint ObjectIndex = SharedCulledObjectList[EffectiveListObjectIndex];
+					{
+				#else
 
-				FLATTEN
-				if (EffectiveListObjectIndex >= TileHead.y + TileHead.z)
-				{
-					ListIndex = 2;
-					ArrayIndex = EffectiveListObjectIndex - TileHead.y - TileHead.z;
-				}
-				else if (EffectiveListObjectIndex >= TileHead.y)
-				{
-					ListIndex = 1;
-					ArrayIndex = EffectiveListObjectIndex - TileHead.y;
-				}
+					uint ListIndex = 0;
+					uint ArrayIndex = EffectiveListObjectIndex;
 
-				uint ObjectIndex = TileArrayData.Load((ArrayIndex * TileListGroupSize.x * TileListGroupSize.y + TileHead.x) * NUM_CULLED_OBJECT_LISTS + ListIndex);
+					FLATTEN
+					if (EffectiveListObjectIndex >= TileHead.y + TileHead.z)
+					{
+						ListIndex = 2;
+						ArrayIndex = EffectiveListObjectIndex - TileHead.y - TileHead.z;
+					}
+					else if (EffectiveListObjectIndex >= TileHead.y)
+					{
+						ListIndex = 1;
+						ArrayIndex = EffectiveListObjectIndex - TileHead.y;
+					}
+
+					uint ObjectIndex = TileArrayData.Load((ArrayIndex * TileListGroupSize.x * TileListGroupSize.y + TileHead.x) * NUM_CULLED_OBJECT_LISTS + ListIndex);
 
 				float4 ObjectPositionAndRadius = LoadObjectPositionAndRadius(ObjectIndex);
 				float ObjectDistanceSq = dot(ObjectPositionAndRadius.xyz - WorldShadingPosition, ObjectPositionAndRadius.xyz - WorldShadingPosition);
@@ -810,6 +872,8 @@ void ConeTraceOcclusionCS(
 				// Skip tracing objects with a small projected angle 
 				if (ObjectPositionAndRadius.w * ObjectPositionAndRadius.w / ObjectDistanceSq > Square(.25f))
 				{
+
+				#endif
 					float3 LocalPositionExtent = LoadObjectLocalPositionExtent(ObjectIndex);
 					float4x4 WorldToVolume = LoadObjectWorldToVolume(ObjectIndex);
 					bool bGeneratedAsTwoSided;
@@ -826,15 +890,17 @@ void ConeTraceOcclusionCS(
 
 						uint StartStepIndex = 0;
 			
-						FLATTEN
-						if (EffectiveListObjectIndex >= TileHead.y + TileHead.z)
-						{
-							StartStepIndex = 8;
-						}
-						else if (EffectiveListObjectIndex >= TileHead.y)
-						{
-							StartStepIndex = 5;
-						}
+						#if !CULL_OBJECTS_TO_RECORD
+							FLATTEN
+							if (EffectiveListObjectIndex >= TileHead.y + TileHead.z)
+							{
+								StartStepIndex = 8;
+							}
+							else if (EffectiveListObjectIndex >= TileHead.y)
+							{
+								StartStepIndex = 5;
+							}
+						#endif
 
 						float WorldStepOffset = GetStepOffset(StartStepIndex);
 
@@ -1005,12 +1071,26 @@ void SetupCopyIndirectArgumentsCS(
 	uint3 DispatchThreadId : SV_DispatchThreadID,
     uint3 GroupThreadId : SV_GroupThreadID) 
 {
+#if FADE_RECORDS_OVER_TIME
+	// Spawn a thread per record
+	RWDispatchParameters[0] = (DrawParameters[1] + COPY_THREADGROUP_SIZE - 1) / COPY_THREADGROUP_SIZE;
+#else
 	uint NumRecords = (1 - TrimFraction) * DrawParameters[1];
 
 	// Make sure we spawn at least one group so RWScatterDrawParameters gets written to in the next pass
 	RWDispatchParameters[0] = max((NumRecords + COPY_THREADGROUP_SIZE - 1) / COPY_THREADGROUP_SIZE, 1);
+#endif
+
 	RWDispatchParameters[1] = 1;
 	RWDispatchParameters[2] = 1;
+
+#if FADE_RECORDS_OVER_TIME
+	if (DispatchThreadId.x == 0)
+	{
+		// Clear to 0 to prepare for accumulation
+		RWScatterDrawParameters[1] = 0;
+	}
+#endif
 }
 
 RWBuffer<float4> RWCopyIrradianceCachePositionRadius;
@@ -1030,6 +1110,52 @@ void CopyIrradianceCacheSamplesCS(
 {
 	uint NumRecords = DrawParameters[1];
 
+#if FADE_RECORDS_OVER_TIME
+
+	uint RecordIndex = DispatchThreadId.x;
+
+	if (RecordIndex < NumRecords)
+	{
+		float4 PositionAndPackedRadius = IrradianceCachePositionRadius[RecordIndex];
+		float4 NormalAndFade = IrradianceCacheNormal[RecordIndex];
+
+		// Update fade
+		float NewFade = abs(NormalAndFade.w) - 1 + clamp(View.GeneralPurposeTweak, .001f, 1);
+
+		// Only write out if still alive
+		if (NewFade < 1 /*RecordIndex > (uint)(NumRecords * TrimFraction)*/)
+		{
+			/*
+			// Always trim at least one to handle dynamic scene changes not accounted for with irradiance cache misses
+			if (RecordIndex <= (uint)(NumRecords * TrimFraction) && PositionAndPackedRadius.w > 0)
+			{
+				// Mark as fading out
+				PositionAndPackedRadius.w *= -1;
+
+				// Reset fade to 0
+				NewFade = 0;
+			}*/
+
+			NormalAndFade.w = (NewFade + 1) * sign(NormalAndFade.w);
+
+			uint DestIndex;
+			InterlockedAdd(RWScatterDrawParameters[1], 1U, DestIndex);
+
+			RWCopyIrradianceCachePositionRadius[DestIndex] = PositionAndPackedRadius;
+			RWCopyIrradianceCacheNormal[DestIndex] = NormalAndFade;
+			RWCopyOccluderRadius[DestIndex] = OccluderRadius[RecordIndex];
+			RWCopyIrradianceCacheBentNormal[DestIndex] = IrradianceCacheBentNormal[RecordIndex];
+
+			#if SUPPORT_IRRADIANCE
+				RWCopyIrradianceCacheIrradiance[DestIndex] = IrradianceCacheIrradiance[RecordIndex];
+			#endif
+
+			RWCopyIrradianceCacheTileCoordinate[DestIndex] = IrradianceCacheTileCoordinate[RecordIndex];
+		}
+	}
+
+#else
+
 	// Always trim at least one to handle dynamic scene changes not accounted for with irradiance cache misses
 	uint StartIndex = max(NumRecords * TrimFraction, 1);
 	StartIndex = min(StartIndex, NumRecords);
@@ -1042,9 +1168,11 @@ void CopyIrradianceCacheSamplesCS(
 		RWCopyIrradianceCacheNormal[DestIndex] = IrradianceCacheNormal[SourceIndex];
 		RWCopyOccluderRadius[DestIndex] = OccluderRadius[SourceIndex];
 		RWCopyIrradianceCacheBentNormal[DestIndex] = IrradianceCacheBentNormal[SourceIndex];
-#if SUPPORT_IRRADIANCE
-		RWCopyIrradianceCacheIrradiance[DestIndex] = IrradianceCacheIrradiance[SourceIndex];
-#endif
+
+		#if SUPPORT_IRRADIANCE
+			RWCopyIrradianceCacheIrradiance[DestIndex] = IrradianceCacheIrradiance[SourceIndex];
+		#endif
+
 		RWCopyIrradianceCacheTileCoordinate[DestIndex] = IrradianceCacheTileCoordinate[SourceIndex];
 	}
 
@@ -1052,4 +1180,6 @@ void CopyIrradianceCacheSamplesCS(
 	{
 		RWScatterDrawParameters[1] = NumRecords - StartIndex;
 	}
+
+#endif
 }
diff --git a/Engine/Shaders/HeightfieldLighting.usf b/Engine/Shaders/HeightfieldLighting.usf
--- a/Engine/Shaders/HeightfieldLighting.usf
+++ b/Engine/Shaders/HeightfieldLighting.usf
@@ -29,11 +29,11 @@ Buffer<float4> HeightfieldDescriptions;
 
 float4 GetHeightfieldScaleBias(uint Index, out bool bHasSubsections)
 {
-	float4 Vector = HeightfieldDescriptions[Index * HEIGHTFIELD_DATA_STRIDE + 0];
+	float4 HeightfieldVector = HeightfieldDescriptions[Index * HEIGHTFIELD_DATA_STRIDE + 0];
 	// Sign bit stored presence of subsections
-	bHasSubsections = Vector.x < 0;
-	Vector.x = abs(Vector.x);
-	return Vector;
+	bHasSubsections = HeightfieldVector.x < 0;
+	HeightfieldVector.x = abs(HeightfieldVector.x);
+	return HeightfieldVector;
 }
 
 float4 GetHeightfieldMinMaxUV(uint Index)
@@ -452,16 +452,15 @@ void CalculateHeightfieldOcclusionCS(
 	uint NumRecords = ScatterDrawParameters[1];
 	uint RecordIndex = StartIndex + DispatchThreadId.x;
 
-	float4 PositionAndRadius = IrradianceCachePositionRadius[RecordIndex];
+	float3 WorldPosition = IrradianceCachePositionRadius[RecordIndex].xyz;
 	float4 WorldNormalAndHasHeightfieldRepresentation = IrradianceCacheNormal[RecordIndex];
 
 	// Only compute heightfield occlusion contribution for samples created on a distance field, 
 	// Otherwise we are potentially self-shadowing, the heightfield is not accurate enough for that
 	BRANCH
-	if (RecordIndex < NumRecords && WorldNormalAndHasHeightfieldRepresentation.w == 0)
+	if (RecordIndex < NumRecords && WorldNormalAndHasHeightfieldRepresentation.w < 0)
 	{
 		float3 WorldNormal = WorldNormalAndHasHeightfieldRepresentation.xyz;
-		float3 WorldPosition = PositionAndRadius.xyz;
 
 		float3 TangentX;
 		float3 TangentY;
@@ -670,14 +669,12 @@ void CalculateHeightfieldIrradianceCS(
 	uint RelativeRecordIndex = GroupId.x;
 	uint ThreadIndex = GroupThreadId.x;
 
-	float4 PositionAndRadius = IrradianceCachePositionRadius[RecordIndex];
+	float3 WorldPosition = IrradianceCachePositionRadius[RecordIndex].xyz;
 	float3 Lighting = 0;
 	
 	BRANCH
 	if (RecordIndex < NumRecords)
 	{
-		float3 WorldPosition = PositionAndRadius.xyz;
-
 		//@todo - constant time grid lookup
 		LOOP
 		for (uint HeightfieldIndex = 0; HeightfieldIndex < NumHeightfields; HeightfieldIndex++)
diff --git a/Engine/Source/Runtime/Engine/Private/DistanceFieldAtlas.cpp b/Engine/Source/Runtime/Engine/Private/DistanceFieldAtlas.cpp
--- a/Engine/Source/Runtime/Engine/Private/DistanceFieldAtlas.cpp
+++ b/Engine/Source/Runtime/Engine/Private/DistanceFieldAtlas.cpp
@@ -120,13 +120,20 @@ void FDistanceFieldVolumeTextureAtlas::UpdateAllocations()
 
 			FRHIResourceCreateInfo CreateInfo;
 
+			int32 SizeX = BlockAllocator.GetSizeX();
+			int32 SizeY = BlockAllocator.GetSizeY();
+
+#if PLATFORM_PS4
+			SizeX = SizeY = FMath::Max(SizeX, SizeY);
+#endif
+
 			VolumeTextureRHI = RHICreateTexture3D(
-				BlockAllocator.GetSizeX(), 
-				BlockAllocator.GetSizeY(), 
+				SizeX, 
+				SizeY, 
 				BlockAllocator.GetSizeZ(), 
 				Format,
 				1,
-				TexCreate_ShaderResource,
+				TexCreate_ShaderResource | TexCreate_NoTiling,
 				CreateInfo);
 
 			UE_LOG(LogStaticMesh,Log,TEXT("Allocated %s"), *GetSizeString());
diff --git a/Engine/Source/Runtime/Engine/Public/SceneManagement.h b/Engine/Source/Runtime/Engine/Public/SceneManagement.h
--- a/Engine/Source/Runtime/Engine/Public/SceneManagement.h
+++ b/Engine/Source/Runtime/Engine/Public/SceneManagement.h
@@ -595,7 +595,7 @@ public:
 inline bool DoesPlatformSupportDistanceFieldShadowing(EShaderPlatform Platform)
 {
 	// Hasn't been tested elsewhere yet
-	return Platform == SP_PCD3D_SM5;
+	return Platform == SP_PCD3D_SM5 || Platform == SP_PS4;
 }
 
 /** Represents a USkyLightComponent to the rendering thread. */
diff --git a/Engine/Source/Runtime/Renderer/Private/ConvertToUniformMesh.cpp b/Engine/Source/Runtime/Renderer/Private/ConvertToUniformMesh.cpp
--- a/Engine/Source/Runtime/Renderer/Private/ConvertToUniformMesh.cpp
+++ b/Engine/Source/Runtime/Renderer/Private/ConvertToUniformMesh.cpp
@@ -36,7 +36,7 @@ protected:
 	static bool ShouldCache(EShaderPlatform Platform,const FMaterial* Material,const FVertexFactoryType* VertexFactoryType)
 	{
 		return IsFeatureLevelSupported(Platform, ERHIFeatureLevel::SM5) 
-			&& DoesPlatformSupportDistanceFieldAO(Platform)
+			&& DoesPlatformSupportDistanceFieldGI(Platform)
 			&& (FCString::Strstr(VertexFactoryType->GetName(), TEXT("LocalVertexFactory")) != NULL
 				|| FCString::Strstr(VertexFactoryType->GetName(), TEXT("InstancedStaticMeshVertexFactory")) != NULL);
 	}
@@ -119,7 +119,7 @@ protected:
 	static bool ShouldCache(EShaderPlatform Platform,const FMaterial* Material,const FVertexFactoryType* VertexFactoryType)
 	{
 		return IsFeatureLevelSupported(Platform, ERHIFeatureLevel::SM5) 
-			&& DoesPlatformSupportDistanceFieldAO(Platform)
+			&& DoesPlatformSupportDistanceFieldGI(Platform)
 			&& (FCString::Strstr(VertexFactoryType->GetName(), TEXT("LocalVertexFactory")) != NULL
 				|| FCString::Strstr(VertexFactoryType->GetName(), TEXT("InstancedStaticMeshVertexFactory")) != NULL);
 	}
@@ -482,7 +482,7 @@ public:
 	static bool ShouldCache(EShaderPlatform Platform, const FMaterial* Material)
 	{
 		//@todo - lit materials only 
-		return IsFeatureLevelSupported(Platform, ERHIFeatureLevel::SM5) && DoesPlatformSupportDistanceFieldAO(Platform);
+		return IsFeatureLevelSupported(Platform, ERHIFeatureLevel::SM5) && DoesPlatformSupportDistanceFieldGI(Platform);
 	}
 
 	static void ModifyCompilationEnvironment(EShaderPlatform Platform, const FMaterial* Material, FShaderCompilerEnvironment& OutEnvironment)
diff --git a/Engine/Source/Runtime/Renderer/Private/DistanceFieldGlobalIllumination.cpp b/Engine/Source/Runtime/Renderer/Private/DistanceFieldGlobalIllumination.cpp
--- a/Engine/Source/Runtime/Renderer/Private/DistanceFieldGlobalIllumination.cpp
+++ b/Engine/Source/Runtime/Renderer/Private/DistanceFieldGlobalIllumination.cpp
@@ -109,7 +109,7 @@ public:
 
 	static bool ShouldCache(EShaderPlatform Platform)
 	{
-		return IsFeatureLevelSupported(Platform, ERHIFeatureLevel::SM5) && DoesPlatformSupportDistanceFieldAO(Platform);
+		return IsFeatureLevelSupported(Platform, ERHIFeatureLevel::SM5) && DoesPlatformSupportDistanceFieldGI(Platform);
 	}
 
 	static void ModifyCompilationEnvironment(EShaderPlatform Platform, FShaderCompilerEnvironment& OutEnvironment)
@@ -225,7 +225,7 @@ public:
 
 	static bool ShouldCache(EShaderPlatform Platform)
 	{
-		return IsFeatureLevelSupported(Platform, ERHIFeatureLevel::SM5) && DoesPlatformSupportDistanceFieldAO(Platform);
+		return IsFeatureLevelSupported(Platform, ERHIFeatureLevel::SM5) && DoesPlatformSupportDistanceFieldGI(Platform);
 	}
 
 	static void ModifyCompilationEnvironment(EShaderPlatform Platform, FShaderCompilerEnvironment& OutEnvironment)
@@ -283,7 +283,7 @@ public:
 
 	static bool ShouldCache(EShaderPlatform Platform)
 	{
-		return IsFeatureLevelSupported(Platform, ERHIFeatureLevel::SM5) && DoesPlatformSupportDistanceFieldAO(Platform);
+		return IsFeatureLevelSupported(Platform, ERHIFeatureLevel::SM5) && DoesPlatformSupportDistanceFieldGI(Platform);
 	}
 
 	static void ModifyCompilationEnvironment(EShaderPlatform Platform, FShaderCompilerEnvironment& OutEnvironment)
@@ -552,7 +552,7 @@ public:
 
 	static bool ShouldCache(EShaderPlatform Platform)
 	{
-		return IsFeatureLevelSupported(Platform, ERHIFeatureLevel::SM5) && DoesPlatformSupportDistanceFieldAO(Platform);
+		return IsFeatureLevelSupported(Platform, ERHIFeatureLevel::SM5) && DoesPlatformSupportDistanceFieldGI(Platform);
 	}
 
 	static void ModifyCompilationEnvironment(EShaderPlatform Platform, FShaderCompilerEnvironment& OutEnvironment)
@@ -614,7 +614,7 @@ public:
 
 	static bool ShouldCache(EShaderPlatform Platform)
 	{
-		return IsFeatureLevelSupported(Platform, ERHIFeatureLevel::SM5) && DoesPlatformSupportDistanceFieldAO(Platform);
+		return IsFeatureLevelSupported(Platform, ERHIFeatureLevel::SM5) && DoesPlatformSupportDistanceFieldGI(Platform);
 	}
 
 	static void ModifyCompilationEnvironment(EShaderPlatform Platform, FShaderCompilerEnvironment& OutEnvironment)
@@ -884,7 +884,7 @@ public:
 
 	static bool ShouldCache(EShaderPlatform Platform)
 	{
-		return IsFeatureLevelSupported(Platform, ERHIFeatureLevel::SM5) && DoesPlatformSupportDistanceFieldAO(Platform);
+		return IsFeatureLevelSupported(Platform, ERHIFeatureLevel::SM5) && DoesPlatformSupportDistanceFieldGI(Platform);
 	}
 
 	static void ModifyCompilationEnvironment(EShaderPlatform Platform, FShaderCompilerEnvironment& OutEnvironment)
@@ -958,7 +958,7 @@ public:
 
 	static bool ShouldCache(EShaderPlatform Platform)
 	{
-		return IsFeatureLevelSupported(Platform, ERHIFeatureLevel::SM5) && DoesPlatformSupportDistanceFieldAO(Platform);
+		return IsFeatureLevelSupported(Platform, ERHIFeatureLevel::SM5) && DoesPlatformSupportDistanceFieldGI(Platform);
 	}
 
 	static void ModifyCompilationEnvironment(EShaderPlatform Platform, FShaderCompilerEnvironment& OutEnvironment)
@@ -1068,7 +1068,7 @@ public:
 
 	static bool ShouldCache(EShaderPlatform Platform)
 	{
-		return IsFeatureLevelSupported(Platform, ERHIFeatureLevel::SM5) && DoesPlatformSupportDistanceFieldAO(Platform);
+		return IsFeatureLevelSupported(Platform, ERHIFeatureLevel::SM5) && DoesPlatformSupportDistanceFieldGI(Platform);
 	}
 
 	static void ModifyCompilationEnvironment(EShaderPlatform Platform, FShaderCompilerEnvironment& OutEnvironment)
@@ -1216,7 +1216,7 @@ public:
 
 	static bool ShouldCache(EShaderPlatform Platform)
 	{
-		return IsFeatureLevelSupported(Platform, ERHIFeatureLevel::SM5) && DoesPlatformSupportDistanceFieldAO(Platform);
+		return IsFeatureLevelSupported(Platform, ERHIFeatureLevel::SM5) && DoesPlatformSupportDistanceFieldGI(Platform);
 	}
 
 	static void ModifyCompilationEnvironment(EShaderPlatform Platform, FShaderCompilerEnvironment& OutEnvironment)
diff --git a/Engine/Source/Runtime/Renderer/Private/DistanceFieldLightingShared.h b/Engine/Source/Runtime/Renderer/Private/DistanceFieldLightingShared.h
--- a/Engine/Source/Runtime/Renderer/Private/DistanceFieldLightingShared.h
+++ b/Engine/Source/Runtime/Renderer/Private/DistanceFieldLightingShared.h
@@ -21,6 +21,18 @@ extern uint32 UpdateObjectsGroupSize;
 
 extern int32 GDistanceFieldGI;
 
+inline bool DoesPlatformSupportDistanceFieldGI(EShaderPlatform Platform)
+{
+	return Platform == SP_PCD3D_SM5;
+}
+
+inline bool SupportsDistanceFieldGI(ERHIFeatureLevel::Type FeatureLevel, EShaderPlatform ShaderPlatform)
+{
+	return GDistanceFieldGI 
+		&& FeatureLevel >= ERHIFeatureLevel::SM5
+		&& DoesPlatformSupportDistanceFieldGI(ShaderPlatform);
+}
+
 extern FIntPoint GetBufferSizeForAO();
 
 class FDistanceFieldObjectBuffers
@@ -382,7 +394,7 @@ public:
 		if (MaxElements > 0 && Stride > 0)
 		{
 			FRHIResourceCreateInfo CreateInfo;
-			Buffer = RHICreateVertexBuffer(MaxElements * Stride * GPixelFormats[Format].BlockBytes, BUF_Volatile | BUF_ShaderResource, CreateInfo);
+			Buffer = RHICreateVertexBuffer(MaxElements * Stride * GPixelFormats[Format].BlockBytes, BUF_Dynamic | BUF_ShaderResource, CreateInfo);
 			BufferSRV = RHICreateShaderResourceView(Buffer, GPixelFormats[Format].BlockBytes, Format);
 		}
 	}
diff --git a/Engine/Source/Runtime/Renderer/Private/DistanceFieldObjectManagement.cpp b/Engine/Source/Runtime/Renderer/Private/DistanceFieldObjectManagement.cpp
--- a/Engine/Source/Runtime/Renderer/Private/DistanceFieldObjectManagement.cpp
+++ b/Engine/Source/Runtime/Renderer/Private/DistanceFieldObjectManagement.cpp
@@ -5,6 +5,15 @@
 =============================================================================*/
 
 #include "RendererPrivate.h"
+#include "RendererPrivate.h"
+#include "ScenePrivate.h"
+#include "UniformBuffer.h"
+#include "ShaderParameters.h"
+#include "DistanceFieldLightingShared.h"
+#include "DistanceFieldSurfaceCacheLighting.h"
+#include "DistanceFieldGlobalIllumination.h"
+#include "RHICommandList.h"
+#include "SceneUtils.h"
 
 float GAOMaxObjectBoundingRadius = 50000;
 FAutoConsoleVariableRef CVarAOMaxObjectBoundingRadius(
@@ -22,16 +31,6 @@ FAutoConsoleVariableRef CVarAOLogObjectBufferReallocation(
 	ECVF_Cheat | ECVF_RenderThreadSafe
 	);
 
-#include "RendererPrivate.h"
-#include "ScenePrivate.h"
-#include "UniformBuffer.h"
-#include "ShaderParameters.h"
-#include "DistanceFieldLightingShared.h"
-#include "DistanceFieldSurfaceCacheLighting.h"
-#include "DistanceFieldGlobalIllumination.h"
-#include "RHICommandList.h"
-#include "SceneUtils.h"
-
 // Must match equivalent shader defines
 int32 FDistanceFieldObjectBuffers::ObjectDataStride = 16;
 int32 FDistanceFieldCulledObjectBuffers::ObjectDataStride = 16;
@@ -253,7 +252,7 @@ public:
 
 	static bool ShouldCache(EShaderPlatform Platform)
 	{
-		return IsFeatureLevelSupported(Platform, ERHIFeatureLevel::SM5) && DoesPlatformSupportDistanceFieldAO(Platform);
+		return IsFeatureLevelSupported(Platform, ERHIFeatureLevel::SM5) && DoesPlatformSupportDistanceFieldGI(Platform);
 	}
 
 	static void ModifyCompilationEnvironment(EShaderPlatform Platform, FShaderCompilerEnvironment& OutEnvironment)
@@ -739,6 +738,7 @@ void FDeferredShadingSceneRenderer::UpdateGlobalDistanceFieldObjectBuffers(FRHIC
 
 		TArray<uint32> UploadObjectIndices;
 		TArray<FVector4> UploadObjectData;
+		const bool bPrepareForDistanceFieldGI = GVPLMeshGlobalIllumination && SupportsDistanceFieldGI(Scene->GetFeatureLevel(), Scene->GetShaderPlatform());
 
 		if (DistanceFieldSceneData.PendingAddOperations.Num() > 0 || DistanceFieldSceneData.PendingUpdateOperations.Num() > 0)
 		{
@@ -758,7 +758,7 @@ void FDeferredShadingSceneRenderer::UpdateGlobalDistanceFieldObjectBuffers(FRHIC
 			int32 OriginalNumInstancedSurfels = DistanceFieldSceneData.InstancedSurfelAllocations.GetNumSurfelsInBuffer();
 			extern int32 GVPLMeshGlobalIllumination;
 
-			if (GDistanceFieldGI && GVPLMeshGlobalIllumination)
+			if (bPrepareForDistanceFieldGI)
 			{
 				for (int32 UploadPrimitiveIndex = 0; UploadPrimitiveIndex < DistanceFieldSceneData.PendingAddOperations.Num(); UploadPrimitiveIndex++)
 				{
@@ -879,7 +879,7 @@ void FDeferredShadingSceneRenderer::UpdateGlobalDistanceFieldObjectBuffers(FRHIC
 						FPrimitiveSurfelAllocation InstancedAllocation;
 						extern int32 GVPLMeshGlobalIllumination;
 
-						if (GDistanceFieldGI && GVPLMeshGlobalIllumination)
+						if (bPrepareForDistanceFieldGI)
 						{
 							const FPrimitiveSurfelAllocation* AllocationPtr = Scene->DistanceFieldSceneData.SurfelAllocations.FindAllocation(PrimitiveSceneInfo);
 							const FPrimitiveSurfelAllocation* InstancedAllocationPtr = Scene->DistanceFieldSceneData.InstancedSurfelAllocations.FindAllocation(PrimitiveSceneInfo);
diff --git a/Engine/Source/Runtime/Renderer/Private/DistanceFieldSurfaceCacheLighting.cpp b/Engine/Source/Runtime/Renderer/Private/DistanceFieldSurfaceCacheLighting.cpp
--- a/Engine/Source/Runtime/Renderer/Private/DistanceFieldSurfaceCacheLighting.cpp
+++ b/Engine/Source/Runtime/Renderer/Private/DistanceFieldSurfaceCacheLighting.cpp
@@ -30,7 +30,8 @@ FAutoConsoleVariableRef CVarDistanceFieldAO(
 
 bool IsDistanceFieldGIAllowed(const FViewInfo& View)
 {
-	return View.Family->EngineShowFlags.VisualizeDistanceFieldGI || (View.Family->EngineShowFlags.DistanceFieldGI && GDistanceFieldGI && View.Family->EngineShowFlags.GlobalIllumination);
+	return DoesPlatformSupportDistanceFieldGI(View.GetShaderPlatform())
+		&& (View.Family->EngineShowFlags.VisualizeDistanceFieldGI || (View.Family->EngineShowFlags.DistanceFieldGI && GDistanceFieldGI && View.Family->EngineShowFlags.GlobalIllumination));
 }
 
 int32 GAOPowerOfTwoBetweenLevels = 2;
@@ -101,6 +102,14 @@ FAutoConsoleVariableRef CVarAOHistoryWeight(
 	ECVF_Cheat | ECVF_RenderThreadSafe
 	);
 
+float GAOHistoryDistanceThreshold = 20;
+FAutoConsoleVariableRef CVarAOHistoryDistanceThreshold(
+	TEXT("r.AOHistoryDistanceThreshold"),
+	GAOHistoryDistanceThreshold,
+	TEXT("World space distance threshold needed to discard last frame's DFAO results.  Lower values reduce ghosting from characters when near a wall but increase flickering artifacts."),
+	ECVF_Cheat | ECVF_RenderThreadSafe
+	);
+
 float GAORecordRadiusScale = .3f;
 FAutoConsoleVariableRef CVarAORecordRadiusScale(
 	TEXT("r.AORecordRadiusScale"),
@@ -1094,6 +1103,7 @@ public:
 	{
 		DrawParameters.Bind(Initializer.ParameterMap, TEXT("DrawParameters"));
 		DispatchParameters.Bind(Initializer.ParameterMap, TEXT("DispatchParameters"));
+		ScatterDrawParameters.Bind(Initializer.ParameterMap, TEXT("ScatterDrawParameters")); 
 		TrimFraction.Bind(Initializer.ParameterMap, TEXT("TrimFraction"));
 	}
 
@@ -1112,12 +1122,15 @@ public:
 		SetSRVParameter(RHICmdList, ShaderRHI, DrawParameters, SurfaceCacheResources.Level[DepthLevel]->ScatterDrawParameters.SRV);
 
 		DispatchParameters.SetBuffer(RHICmdList, ShaderRHI, SurfaceCacheResources.DispatchParameters);
+		ScatterDrawParameters.SetBuffer(RHICmdList, ShaderRHI, SurfaceCacheResources.TempResources->ScatterDrawParameters);
+
 		SetShaderValue(RHICmdList, ShaderRHI, TrimFraction, GAOTrimOldRecordsFraction);
 	}
 
 	void UnsetParameters(FRHICommandList& RHICmdList)
 	{
 		DispatchParameters.UnsetUAV(RHICmdList, GetComputeShader());
+		ScatterDrawParameters.UnsetUAV(RHICmdList, GetComputeShader());
 	}
 
 	virtual bool Serialize(FArchive& Ar)
@@ -1125,6 +1138,7 @@ public:
 		bool bShaderHasOutdatedParameters = FGlobalShader::Serialize(Ar);
 		Ar << DrawParameters;
 		Ar << DispatchParameters;
+		Ar << ScatterDrawParameters;
 		Ar << TrimFraction;
 		return bShaderHasOutdatedParameters;
 	}
@@ -1133,6 +1147,7 @@ private:
 
 	FShaderResourceParameter DrawParameters;
 	FRWShaderParameter DispatchParameters;
+	FRWShaderParameter ScatterDrawParameters;
 	FShaderParameter TrimFraction;
 };
 
@@ -2335,6 +2350,7 @@ public:
 		IrradianceHistorySampler.Bind(Initializer.ParameterMap, TEXT("IrradianceHistorySampler"));
 		CameraMotion.Bind(Initializer.ParameterMap);
 		HistoryWeight.Bind(Initializer.ParameterMap, TEXT("HistoryWeight"));
+		HistoryDistanceThreshold.Bind(Initializer.ParameterMap, TEXT("HistoryDistanceThreshold"));
 		VelocityTexture.Bind(Initializer.ParameterMap, TEXT("VelocityTexture"));
 		VelocityTextureSampler.Bind(Initializer.ParameterMap, TEXT("VelocityTextureSampler"));
 	}
@@ -2398,6 +2414,7 @@ public:
 		CameraMotion.Set(RHICmdList, View, ShaderRHI);
 
 		SetShaderValue(RHICmdList, ShaderRHI, HistoryWeight, GAOHistoryWeight);
+		SetShaderValue(RHICmdList, ShaderRHI, HistoryDistanceThreshold, GAOHistoryDistanceThreshold);
 
 		SetTextureParameter(
 			RHICmdList,
@@ -2423,6 +2440,7 @@ public:
 		Ar << IrradianceHistorySampler;
 		Ar << CameraMotion;
 		Ar << HistoryWeight;
+		Ar << HistoryDistanceThreshold;
 		Ar << VelocityTexture;
 		Ar << VelocityTextureSampler;
 		return bShaderHasOutdatedParameters;
@@ -2441,6 +2459,7 @@ private:
 	FShaderResourceParameter IrradianceHistorySampler;
 	FCameraMotionParameters CameraMotion;
 	FShaderParameter HistoryWeight;
+	FShaderParameter HistoryDistanceThreshold;
 	FShaderResourceParameter VelocityTexture;
 	FShaderResourceParameter VelocityTextureSampler;
 };
diff --git a/Engine/Source/Runtime/Renderer/Private/DistanceFieldSurfaceCacheLighting.h b/Engine/Source/Runtime/Renderer/Private/DistanceFieldSurfaceCacheLighting.h
--- a/Engine/Source/Runtime/Renderer/Private/DistanceFieldSurfaceCacheLighting.h
+++ b/Engine/Source/Runtime/Renderer/Private/DistanceFieldSurfaceCacheLighting.h
@@ -212,7 +212,7 @@ private:
 
 inline bool DoesPlatformSupportDistanceFieldAO(EShaderPlatform Platform)
 {
-	return Platform == SP_PCD3D_SM5;
+	return Platform == SP_PCD3D_SM5 || Platform == SP_PS4;
 }
 
 template<bool bOneGroupPerRecord>
diff --git a/Engine/Source/Runtime/Renderer/Private/HeightfieldLighting.cpp b/Engine/Source/Runtime/Renderer/Private/HeightfieldLighting.cpp
--- a/Engine/Source/Runtime/Renderer/Private/HeightfieldLighting.cpp
+++ b/Engine/Source/Runtime/Renderer/Private/HeightfieldLighting.cpp
@@ -288,7 +288,7 @@ public:
 
 	static bool ShouldCache(EShaderPlatform Platform)
 	{
-		return IsFeatureLevelSupported(Platform, ERHIFeatureLevel::SM5) && DoesPlatformSupportDistanceFieldAO(Platform);
+		return IsFeatureLevelSupported(Platform, ERHIFeatureLevel::SM5) && DoesPlatformSupportDistanceFieldGI(Platform);
 	}
 
 	/** Default constructor. */
@@ -331,7 +331,7 @@ public:
 
 	static bool ShouldCache(EShaderPlatform Platform)
 	{
-		return IsFeatureLevelSupported(Platform, ERHIFeatureLevel::SM5) && DoesPlatformSupportDistanceFieldAO(Platform);
+		return IsFeatureLevelSupported(Platform, ERHIFeatureLevel::SM5) && DoesPlatformSupportDistanceFieldGI(Platform);
 	}
 
 	/** Default constructor. */
@@ -409,7 +409,7 @@ TGlobalResource<FQuadVertexBuffer> GQuadVertexBuffer;
 bool SupportsHeightfieldLighting(ERHIFeatureLevel::Type FeatureLevel, EShaderPlatform ShaderPlatform)
 {
 	return FeatureLevel >= ERHIFeatureLevel::SM5
-		&& DoesPlatformSupportDistanceFieldAO(ShaderPlatform);
+		&& DoesPlatformSupportDistanceFieldGI(ShaderPlatform);
 }
 
 bool AllowHeightfieldGI(const FViewInfo& View)
@@ -701,7 +701,7 @@ public:
 
 	static bool ShouldCache(EShaderPlatform Platform)
 	{
-		return IsFeatureLevelSupported(Platform, ERHIFeatureLevel::SM5) && DoesPlatformSupportDistanceFieldAO(Platform);
+		return IsFeatureLevelSupported(Platform, ERHIFeatureLevel::SM5) && DoesPlatformSupportDistanceFieldGI(Platform);
 	}
 
 	/** Default constructor. */
@@ -744,7 +744,7 @@ public:
 
 	static bool ShouldCache(EShaderPlatform Platform)
 	{
-		return IsFeatureLevelSupported(Platform, ERHIFeatureLevel::SM5) && DoesPlatformSupportDistanceFieldAO(Platform);
+		return IsFeatureLevelSupported(Platform, ERHIFeatureLevel::SM5) && DoesPlatformSupportDistanceFieldGI(Platform);
 	}
 
 	/** Default constructor. */
@@ -889,7 +889,7 @@ public:
 
 	static bool ShouldCache(EShaderPlatform Platform)
 	{
-		return IsFeatureLevelSupported(Platform, ERHIFeatureLevel::SM5) && DoesPlatformSupportDistanceFieldAO(Platform);
+		return IsFeatureLevelSupported(Platform, ERHIFeatureLevel::SM5) && DoesPlatformSupportDistanceFieldGI(Platform);
 	}
 
 	/** Default constructor. */
@@ -1049,7 +1049,7 @@ public:
 
 	static bool ShouldCache(EShaderPlatform Platform, const FMaterial* Material)
 	{
-		return Material->IsLightFunction() && IsFeatureLevelSupported(Platform, ERHIFeatureLevel::SM5) && DoesPlatformSupportDistanceFieldAO(Platform);
+		return Material->IsLightFunction() && IsFeatureLevelSupported(Platform, ERHIFeatureLevel::SM5) && DoesPlatformSupportDistanceFieldGI(Platform);
 	}
 
 	static void ModifyCompilationEnvironment(EShaderPlatform Platform, const FMaterial* Material, FShaderCompilerEnvironment& OutEnvironment)
@@ -1219,7 +1219,7 @@ public:
 
 	static bool ShouldCache(EShaderPlatform Platform)
 	{
-		return IsFeatureLevelSupported(Platform, ERHIFeatureLevel::SM5) && DoesPlatformSupportDistanceFieldAO(Platform);
+		return IsFeatureLevelSupported(Platform, ERHIFeatureLevel::SM5) && DoesPlatformSupportDistanceFieldGI(Platform);
 	}
 
 	static void ModifyCompilationEnvironment(EShaderPlatform Platform, FShaderCompilerEnvironment& OutEnvironment)
@@ -1415,7 +1415,7 @@ public:
 
 	static bool ShouldCache(EShaderPlatform Platform)
 	{
-		return IsFeatureLevelSupported(Platform, ERHIFeatureLevel::SM5) && DoesPlatformSupportDistanceFieldAO(Platform);
+		return IsFeatureLevelSupported(Platform, ERHIFeatureLevel::SM5) && DoesPlatformSupportDistanceFieldGI(Platform);
 	}
 
 	static void ModifyCompilationEnvironment(EShaderPlatform Platform, FShaderCompilerEnvironment& OutEnvironment)
diff --git a/Engine/Source/Runtime/Renderer/Private/SurfelTree.cpp b/Engine/Source/Runtime/Renderer/Private/SurfelTree.cpp
--- a/Engine/Source/Runtime/Renderer/Private/SurfelTree.cpp
+++ b/Engine/Source/Runtime/Renderer/Private/SurfelTree.cpp
@@ -48,7 +48,7 @@ public:
 
 	static bool ShouldCache(EShaderPlatform Platform)
 	{
-		return IsFeatureLevelSupported(Platform, ERHIFeatureLevel::SM5) && DoesPlatformSupportDistanceFieldAO(Platform);
+		return IsFeatureLevelSupported(Platform, ERHIFeatureLevel::SM5) && DoesPlatformSupportDistanceFieldGI(Platform);
 	}
 
 	static void ModifyCompilationEnvironment(EShaderPlatform Platform, FShaderCompilerEnvironment& OutEnvironment)
@@ -120,7 +120,7 @@ public:
 
 	static bool ShouldCache(EShaderPlatform Platform)
 	{
-		return IsFeatureLevelSupported(Platform, ERHIFeatureLevel::SM5) && DoesPlatformSupportDistanceFieldAO(Platform);
+		return IsFeatureLevelSupported(Platform, ERHIFeatureLevel::SM5) && DoesPlatformSupportDistanceFieldGI(Platform);
 	}
 
 	static void ModifyCompilationEnvironment(EShaderPlatform Platform, FShaderCompilerEnvironment& OutEnvironment)
@@ -190,7 +190,7 @@ public:
 
 	static bool ShouldCache(EShaderPlatform Platform)
 	{
-		return IsFeatureLevelSupported(Platform, ERHIFeatureLevel::SM5) && DoesPlatformSupportDistanceFieldAO(Platform);
+		return IsFeatureLevelSupported(Platform, ERHIFeatureLevel::SM5) && DoesPlatformSupportDistanceFieldGI(Platform);
 	}
 
 	static void ModifyCompilationEnvironment(EShaderPlatform Platform, FShaderCompilerEnvironment& OutEnvironment)
