Fix PS4 volumetexture updates

diff --git a/Engine/Shaders/UpdateTextureShaders.usf b/Engine/Shaders/UpdateTextureShaders.usf
--- a/Engine/Shaders/UpdateTextureShaders.usf
+++ b/Engine/Shaders/UpdateTextureShaders.usf
@@ -6,9 +6,9 @@
 
 #include "Common.usf"
 
-Buffer<int4> SrcBuffer;
-RWTexture2D<int4> DestTexture;  
-int SrcPitch;
+Buffer<int> SrcBuffer;
+RWTexture2D<int4> DestTexture;
+uint SrcPitch; //number of SrcBuffer elements in a row of the source area
 
 
 uint4 DestPosSize; // xy = pos, zw = Size of the destination sub rect
@@ -18,21 +18,29 @@ void UpdateTexture2DSubresourceCS( uint2 ThreadId : SV_DispatchThreadID )
 {
 	if( all( ThreadId.xy < DestPosSize.zw ) )
 	{
-		DestTexture[ThreadId.xy + DestPosSize.xy] = SrcBuffer[ThreadId.x + ThreadId.y * SrcPitch];
+		int2 TexturePixelOffset = ThreadId.xy + DestPosSize.xy;
+
+		//supports textures with 1-4 components.  Textures with  < 4 components will simply ignore writes to the extra components via channel masking.		
+		int BufferOffset = ThreadId.x + (ThreadId.y * SrcPitch);
+		DestTexture[TexturePixelOffset] = int4(SrcBuffer[BufferOffset].x, SrcBuffer[BufferOffset + 1].x, SrcBuffer[BufferOffset + 2].x, SrcBuffer[BufferOffset + 3].x);
 	}
 }
 
 RWTexture3D<int4> DestTexture3D;
-int SrcDepthPitch;
-uint4 DestPos;
-uint4 DestSize;
+uint SrcDepthPitch; //number of SrcBuffer entries between z slices of the source volume
+uint4 DestPos; // xyz = starting offset in destination volume texture
+uint4 DestSize; //syz = size of the volume to update
 
-[numthreads(8, 8, 8)]
+[numthreads(8, 8, 1)]
 void UpdateTexture3DSubresourceCS(uint3 ThreadId : SV_DispatchThreadID)
 {
-	if (all(ThreadId.xyz < DestSize.xyz))
+	if (all(ThreadId.xyz < DestSize.xyz))  
 	{
-		DestTexture3D[ThreadId.xyz + DestPos.xyz] = SrcBuffer[ThreadId.x + (ThreadId.y * SrcPitch) + (ThreadId.z * SrcDepthPitch)];
+		int3 TexturePixelOffset = ThreadId.xyz + (int3)DestPos.xyz;
+
+		//supports textures with 1-4 components.  Textures with  < 4 components will simply ignore writes to the extra components via channel masking.
+		int BufferOffset = ThreadId.x + (ThreadId.y * SrcPitch) + (ThreadId.z * SrcDepthPitch);		
+		DestTexture3D[TexturePixelOffset] = int4(SrcBuffer[BufferOffset].x, SrcBuffer[BufferOffset + 1].x, SrcBuffer[BufferOffset + 2].x, SrcBuffer[BufferOffset + 3].x);
 	}
 }
 
diff --git a/Engine/Source/Runtime/Engine/Private/DistanceFieldAtlas.cpp b/Engine/Source/Runtime/Engine/Private/DistanceFieldAtlas.cpp
--- a/Engine/Source/Runtime/Engine/Private/DistanceFieldAtlas.cpp
+++ b/Engine/Source/Runtime/Engine/Private/DistanceFieldAtlas.cpp
@@ -123,10 +123,6 @@ void FDistanceFieldVolumeTextureAtlas::UpdateAllocations()
 			int32 SizeX = BlockAllocator.GetSizeX();
 			int32 SizeY = BlockAllocator.GetSizeY();
 
-#if PLATFORM_PS4
-			SizeX = SizeY = FMath::Max(SizeX, SizeY);
-#endif
-
 			VolumeTextureRHI = RHICreateTexture3D(
 				SizeX, 
 				SizeY, 
@@ -163,9 +159,9 @@ void FDistanceFieldVolumeTextureAtlas::UpdateAllocations()
 
 		CurrentAllocations.Append(PendingAllocations);
 		PendingAllocations.Empty();
-	}
+	}	
 }
-
+	
 void FDistanceFieldVolumeTexture::Initialize()
 {
 	if (IsValidDistanceFieldVolume())
