UE4 - Fix issue with > 75 bones on Metal (UE-13255)

diff --git a/Engine/Shaders/GpuSkinVertexFactory.usf b/Engine/Shaders/GpuSkinVertexFactory.usf
--- a/Engine/Shaders/GpuSkinVertexFactory.usf
+++ b/Engine/Shaders/GpuSkinVertexFactory.usf
@@ -159,7 +159,7 @@ half2 GetShadowMapCoordinate(FVertexFactoryInterpolantsVSToPS Interpolants)
 
 #define FBoneMatrix float3x4
 
-#if FEATURE_LEVEL >= FEATURE_LEVEL_SM4
+#if FEATURE_LEVEL >= FEATURE_LEVEL_ES3_1
 
 // The bone matrix buffer stored as 4x3 (3 float4 texels behind each other), all chunks of a skeletal mesh in one
 Buffer<float4> BoneMatrices;
@@ -278,7 +278,7 @@ float3 MorphPosition( FVertexFactoryInput Input, FVertexFactoryIntermediates Int
 
 FBoneMatrix GetBoneMatrix(int Index)
 {
-#if FEATURE_LEVEL >= FEATURE_LEVEL_SM4
+#if FEATURE_LEVEL >= FEATURE_LEVEL_ES3_1
 	float4 A = BoneMatrices[Index * 3];
 	float4 B = BoneMatrices[Index * 3 + 1];
 	float4 C = BoneMatrices[Index * 3 + 2];
@@ -307,7 +307,7 @@ FBoneMatrix CalcBoneMatrix( FVertexFactoryInput Input )
 
 FBoneMatrix GetPreviousBoneMatrix(int Index)
 {
-#if FEATURE_LEVEL >= FEATURE_LEVEL_SM4
+#if FEATURE_LEVEL >= FEATURE_LEVEL_ES3_1
 	float4 A = PreviousBoneMatrices[Index * 3 + BoneIndexOffset.x];
 	float4 B = PreviousBoneMatrices[Index * 3 + BoneIndexOffset.y];
 	float4 C = PreviousBoneMatrices[Index * 3 + BoneIndexOffset.z];
diff --git a/Engine/Source/Runtime/Engine/Private/GPUSkinVertexFactory.cpp b/Engine/Source/Runtime/Engine/Private/GPUSkinVertexFactory.cpp
--- a/Engine/Source/Runtime/Engine/Private/GPUSkinVertexFactory.cpp
+++ b/Engine/Source/Runtime/Engine/Private/GPUSkinVertexFactory.cpp
@@ -116,7 +116,7 @@ void FGPUBaseSkinVertexFactory::ShaderDataType::UpdateBoneData(ERHIFeatureLevel:
 {
 	uint32 NumBones = BoneMatrices.Num();
 	check(NumBones <= MaxGPUSkinBones);
-	if (FeatureLevel >= ERHIFeatureLevel::SM4)
+	if (FeatureLevel >= ERHIFeatureLevel::ES3_1)
 	{
 		static FSharedPoolPolicyData PoolPolicy;
 		uint32 NumVectors = NumBones*3;
@@ -391,7 +391,7 @@ public:
 				ShaderData.MeshExtension
 				);
 			
-			if (FeatureLevel >= ERHIFeatureLevel::SM4)
+			if (FeatureLevel >= ERHIFeatureLevel::ES3_1)
 			{
 				if(BoneMatrices.IsBound())
 				{
