UE-2400: Implements the new material node "Decal Mipmap Level" for the decal material domain

diff --git a/Engine/Documentation/Source/Engine/Actors/DecalActor/DecalUserGuide.INT.udn b/Engine/Documentation/Source/Engine/Actors/DecalActor/DecalUserGuide.INT.udn
--- a/Engine/Documentation/Source/Engine/Actors/DecalActor/DecalUserGuide.INT.udn
+++ b/Engine/Documentation/Source/Engine/Actors/DecalActor/DecalUserGuide.INT.udn
@@ -116,7 +116,30 @@ The darker color indicates a higher performance cost of those pixels. This infor
 
 * We currently only support deferred decals and they only work on static objects.
 * Normal blending is currently not wrapping around the object.
-* Mip map computation is not done yet so you might see 2x2 block artifacts on object borders.
 * Streaming is not yet hooked up so make sure the texture is not streamed.
 * Masking decals (not affecting other object) is not fully implemented.
 
+## The 2x2 block artifacts fix
+
+Decals may have 2x2 pixels block artifacts on edges as shown in the screenshot below.
+
+![](2x2_artifacts.bmp)
+
+This is where the node "Decal Mipmap Level" comes in. This node must be used with  has a very big performance impact.
+
+![](decal_mipmap_material_1.png)
+
+And using it fixes the artifacts on the decal that benefits from it.
+
+![](2x2_no_artifacts.bmp)
+
+### Missing SampleTexture UV
+
+Now this node is so expensive from a computing point of view, that it doesn't have UV input as the SampleTexture. But in case of a tilling, you can still do so by patching up the mipmap level manualy. Here is an example to have a 2^4 = 16 tiling repetition.
+
+![](decal_mipmap_material_2.png)
+
+### Current limitations
+
+* DecalMipmapLevel doesn't support UV
+* DecalMipmapLevel doesn't work with anysotropic filtering
diff --git a/Engine/Documentation/Source/Shared/GraphNodes/Material/Material.INT.udn b/Engine/Documentation/Source/Shared/GraphNodes/Material/Material.INT.udn
--- a/Engine/Documentation/Source/Shared/GraphNodes/Material/Material.INT.udn
+++ b/Engine/Documentation/Source/Shared/GraphNodes/Material/Material.INT.udn
@@ -312,6 +312,14 @@ Engine/Rendering/Materials/ExpressionReference/Custom#customtexture
 [/VAR]
 [/EXCERPT:UMaterialExpressionCustomTexture]
 
+[EXCERPT:UMaterialExpressionDecalMipmapLevel]
+The DecalMipmapLevel is used to explitily compute a decal's texture's mipmap level assuming there is no custom uv, to avoid the 2x2 pixel mipmap artefacts on the edges where the decal is project to.
+
+[VAR:ToolTipFullLink]
+Engine/Actors/DecalActor/DecalUserGuide
+[/VAR]
+[/EXCERPT:UMaterialExpressionDecalMipmapLevel]
+
 [EXCERPT:UMaterialExpressionDepthFade]
 The DepthFade expression is used to hide unsightly seams that take place when translucent objects intersect with opaque ones. 
 
diff --git a/Engine/Shaders/MaterialTemplate.usf b/Engine/Shaders/MaterialTemplate.usf
--- a/Engine/Shaders/MaterialTemplate.usf
+++ b/Engine/Shaders/MaterialTemplate.usf
@@ -260,6 +260,62 @@ struct FMaterialVertexParameters
 	FMaterialParticleParameters Particle;
 };
 
+
+#ifdef DECAL_RENDERSTAGE
+	/*
+	 * Material node DecalMipmapLevel's code designed to avoid the 2x2 pixels artefacts on the edges around where the decal
+	 * is projected to. The technique is fetched from (http://www.humus.name/index.php?page=3D&ID=84).
+	 *
+	 * The problem around edges of the meshes, is that the hardware computes the mipmap level according to ddx(uv) and ddy(uv),
+	 * but since the pixel shader are invocated by group of 2x2 pixels, then on edges some pixel might be getting the
+	 * current depth of an differet mesh that the other pixel of the same groups. If this mesh is very far from the other
+	 * mesh of the same group of pixel, then one of the delta might be very big, leading to choosing a low mipmap level for this
+	 * group of 4 pixels, causing the artefacts.
+	 */
+	float2 ComputeDecalUVFromScreenPosition(float4 ScreenPosition)
+	{
+		half SceneW = PreviousDepth(ScreenPosition);
+		float4 DecalVector = mul(float4(ScreenPosition.xy / ScreenPosition.w * SceneW, SceneW, 1), ScreenToDecal);
+		DecalVector.xyz /= DecalVector.w;
+		DecalVector = DecalVector * 0.5f + 0.5f;
+		DecalVector.xyz = DecalVector.zyx;
+		return DecalVector.xy;
+	}
+
+	float ComputeDecalMipmapLevel(FMaterialPixelParameters Parameters, float2 TextureSize)
+	{
+		/*
+		 * Assuming where in a pixel shader invocation, then we compute manualy compute two d(uv)/d(x)
+		 * with the pixels's left and right neighbours.
+		 */
+		float4 ScreenDeltaX = ddx(Parameters.ScreenPosition);
+		float2 UvDiffX0 = Parameters.TexCoords[0] - ComputeDecalUVFromScreenPosition(Parameters.ScreenPosition - ScreenDeltaX);
+		float2 UvDiffX1 = ComputeDecalUVFromScreenPosition(Parameters.ScreenPosition + ScreenDeltaX) - Parameters.TexCoords[0];
+
+		/*
+		 * So we have two diff on the X axis, we want the one that has the smallest length
+		 * to avoid the 2x2 pixels mipmap artefacts on the edges. 
+		 */
+		float2 UvPixelDiffX = (dot(UvDiffX0, UvDiffX0) < dot(UvDiffX1, UvDiffX1) ? UvDiffX0 : UvDiffX1) * TextureSize;
+	
+		// do same for the Y axis
+		float4 ScreenDeltaY = ddy(Parameters.ScreenPosition);
+		float2 UvDiffY0 = Parameters.TexCoords[0] - ComputeDecalUVFromScreenPosition(Parameters.ScreenPosition - ScreenDeltaY);
+		float2 UvDiffY1 = ComputeDecalUVFromScreenPosition(Parameters.ScreenPosition + ScreenDeltaY) - Parameters.TexCoords[0];
+		float2 UvPixelDiffY = (dot(UvDiffY0, UvDiffY0) < dot(UvDiffY1, UvDiffY1) ? UvDiffY0 : UvDiffY1) * TextureSize;
+
+		// Computes the mipmap level
+		float MaxDiff = max(dot(UvPixelDiffX, UvPixelDiffX), dot(UvPixelDiffY, UvPixelDiffY));
+		return 0.5 * log2(MaxDiff);
+	}
+#else
+	float ComputeDecalMipmapLevel(FMaterialPixelParameters Parameters, float2 TextureSize)
+	{
+		return 0.0f;
+	}
+#endif //DECAL_RENDERSTAGE
+
+
 #if HAS_PRIMITIVE_UNIFORM_BUFFER
 
 /** Transforms a vector from tangent space to world space, prescaling by an amount calculated previously */
diff --git a/Engine/Source/Runtime/Engine/Classes/Materials/MaterialExpressionDecalMipmapLevel.h b/Engine/Source/Runtime/Engine/Classes/Materials/MaterialExpressionDecalMipmapLevel.h
new file mode 100644
--- /dev/null
+++ b/Engine/Source/Runtime/Engine/Classes/Materials/MaterialExpressionDecalMipmapLevel.h
@@ -0,0 +1,29 @@
+// Copyright 1998-2015 Epic Games, Inc. All Rights Reserved.
+
+
+#pragma once
+#include "Materials/MaterialExpression.h"
+#include "MaterialExpressionDecalMipmapLevel.generated.h"
+
+
+UCLASS(collapsecategories, hidecategories = Object)
+class UMaterialExpressionDecalMipmapLevel : public UMaterialExpression
+{
+	GENERATED_UCLASS_BODY()
+
+	/** The texture's size */
+	UPROPERTY(meta = (RequiredInput = "false", ToolTip = "Defaults to '(Const Width, Const Height)' if not specified"))
+	FExpressionInput TextureSize;
+
+	/** only used if TextureSize is not hooked up */
+	UPROPERTY(EditAnywhere, Category = MaterialExpressionDecalMipmapLevel, meta=(OverridingInputProperty = "TextureSize"))
+	float ConstWidth;
+	UPROPERTY(EditAnywhere, Category = MaterialExpressionDecalMipmapLevel, meta=(OverridingInputProperty = "TextureSize"))
+	float ConstHeight;
+
+
+	// Begin UMaterialExpression Interface
+	virtual int32 Compile(class FMaterialCompiler* Compiler, int32 OutputIndex, int32 MultiplexIndex) override;
+	virtual void GetCaption(TArray<FString>& OutCaptions) const override;
+	// End UMaterialExpression Interface
+};
diff --git a/Engine/Source/Runtime/Engine/Private/Materials/HLSLMaterialTranslator.h b/Engine/Source/Runtime/Engine/Private/Materials/HLSLMaterialTranslator.h
--- a/Engine/Source/Runtime/Engine/Private/Materials/HLSLMaterialTranslator.h
+++ b/Engine/Source/Runtime/Engine/Private/Materials/HLSLMaterialTranslator.h
@@ -2386,6 +2386,25 @@ protected:
 			);
 	}
 
+	virtual int32 TextureDecalMipmapLevel(int32 TextureSizeInput) override
+	{
+		EMaterialValueType TextureSizeType = GetParameterType(TextureSizeInput);
+
+		if (TextureSizeType != MCT_Float2)
+		{
+			Errorf(TEXT("Unmatching conversion %s -> float2"), DescribeType(TextureSizeType));
+			return INDEX_NONE;
+		}
+
+		FString TextureSize = CoerceParameter(TextureSizeInput, MCT_Float2);
+
+		return AddCodeChunk(
+			MCT_Float1,
+			TEXT("ComputeDecalMipmapLevel(Parameters,%s)"),
+			*TextureSize
+			);
+	}
+
 	virtual int32 PixelDepth() override
 	{
 		if (ShaderFrequency != SF_Pixel && ShaderFrequency != SF_Compute)
diff --git a/Engine/Source/Runtime/Engine/Private/Materials/MaterialExpressions.cpp b/Engine/Source/Runtime/Engine/Private/Materials/MaterialExpressions.cpp
--- a/Engine/Source/Runtime/Engine/Private/Materials/MaterialExpressions.cpp
+++ b/Engine/Source/Runtime/Engine/Private/Materials/MaterialExpressions.cpp
@@ -31,6 +31,7 @@
 #include "Materials/MaterialExpressionCustom.h"
 #include "Materials/MaterialExpressionDDX.h"
 #include "Materials/MaterialExpressionDDY.h"
+#include "Materials/MaterialExpressionDecalMipmapLevel.h"
 #include "Materials/MaterialExpressionDepthFade.h"
 #include "Materials/MaterialExpressionDepthOfFieldFunction.h"
 #include "Materials/MaterialExpressionDeriveNormalZ.h"
@@ -8485,6 +8486,60 @@ void UMaterialExpressionAntialiasedTextureMask::SetDefaultTexture()
 	Texture = LoadObject<UTexture2D>(NULL, TEXT("/Engine/EngineResources/DefaultTexture.DefaultTexture"), NULL, LOAD_None, NULL);
 }
 
+//
+//	UMaterialExpressionDecalMipmapLevel
+//
+UMaterialExpressionDecalMipmapLevel::UMaterialExpressionDecalMipmapLevel(const FObjectInitializer& ObjectInitializer)
+	: Super(ObjectInitializer)
+	, ConstWidth(256.0f)
+	, ConstHeight(ConstWidth)
+{
+	// Structure to hold one-time initialization
+	struct FConstructorStatics
+	{
+		FString NAME_Vectors;
+		FConstructorStatics()
+			: NAME_Vectors(LOCTEXT("Utils", "Utils").ToString())
+		{
+		}
+	};
+	static FConstructorStatics ConstructorStatics;
+
+	MenuCategories.Add(ConstructorStatics.NAME_Vectors);
+	bCollapsed = true;
+}
+
+int32 UMaterialExpressionDecalMipmapLevel::Compile(class FMaterialCompiler* Compiler, int32 OutputIndex, int32 MultiplexIndex)
+{
+	if (Material->MaterialDomain != MD_DeferredDecal)
+	{
+		return CompilerError(Compiler, TEXT("Node only works for the deferred decal material domain."));
+	}
+
+	int32 TextureSizeInput = INDEX_NONE;
+
+	if (TextureSize.Expression)
+	{
+		TextureSizeInput = TextureSize.Compile(Compiler);
+	}
+	else
+	{
+		TextureSizeInput = Compiler->Constant2(ConstWidth, ConstHeight);
+	}
+
+	if (TextureSizeInput == INDEX_NONE)
+	{
+		return INDEX_NONE;
+	}
+
+	return Compiler->TextureDecalMipmapLevel(TextureSizeInput);
+}
+
+void UMaterialExpressionDecalMipmapLevel::GetCaption(TArray<FString>& OutCaptions) const
+{
+	OutCaptions.Add(TEXT("Decal Mipmap Level"));
+}
+
 UMaterialExpressionDepthFade::UMaterialExpressionDepthFade(const FObjectInitializer& ObjectInitializer)
 	: Super(ObjectInitializer)
 {
diff --git a/Engine/Source/Runtime/Engine/Public/MaterialCompiler.h b/Engine/Source/Runtime/Engine/Public/MaterialCompiler.h
--- a/Engine/Source/Runtime/Engine/Public/MaterialCompiler.h
+++ b/Engine/Source/Runtime/Engine/Public/MaterialCompiler.h
@@ -102,6 +102,8 @@ public:
 	virtual int32 TextureCoordinate(uint32 CoordinateIndex, bool UnMirrorU, bool UnMirrorV) = 0;
 	virtual int32 TextureSample(int32 Texture,int32 Coordinate,enum EMaterialSamplerType SamplerType,int32 MipValueIndex=INDEX_NONE,ETextureMipValueMode MipValueMode=TMVM_None,ESamplerSourceMode SamplerSource=SSM_FromTextureAsset) = 0;
 
+	virtual int32 TextureDecalMipmapLevel(int32 TextureSizeInput) = 0;
+
 	virtual int32 Texture(UTexture* Texture,ESamplerSourceMode SamplerSource=SSM_FromTextureAsset) = 0;
 	virtual int32 TextureParameter(FName ParameterName,UTexture* DefaultTexture,ESamplerSourceMode SamplerSource=SSM_FromTextureAsset) = 0;
 
@@ -259,6 +261,8 @@ public:
 
 	virtual int32 TextureCoordinate(uint32 CoordinateIndex, bool UnMirrorU, bool UnMirrorV) override { return Compiler->TextureCoordinate(CoordinateIndex, UnMirrorU, UnMirrorV); }
 
+	virtual int32 TextureDecalMipmapLevel(int32 TextureSizeInput) override { return Compiler->TextureDecalMipmapLevel(TextureSizeInput); }
+
 	virtual int32 Texture(UTexture* InTexture,ESamplerSourceMode SamplerSource=SSM_FromTextureAsset) override { return Compiler->Texture(InTexture,SamplerSource); }
 	virtual int32 TextureParameter(FName ParameterName,UTexture* DefaultValue,ESamplerSourceMode SamplerSource=SSM_FromTextureAsset) override { return Compiler->TextureParameter(ParameterName,DefaultValue,SamplerSource); }
 
