removed unused parameters hiding d3ddebug errors,  fixed d3ddebug errors

diff --git a/Engine/Shaders/EyeAdaptationCommon.usf b/Engine/Shaders/EyeAdaptationCommon.usf
new file mode 100644
--- /dev/null
+++ b/Engine/Shaders/EyeAdaptationCommon.usf
@@ -0,0 +1,22 @@
+// Copyright 1998-2015 Epic Games, Inc. All Rights Reserved.
+
+/*=============================================================================
+	EyeAdaptationCommon.usf: PostProcessing shared functions and structures.
+=============================================================================*/
+
+#ifndef __EYE_ADAPTATION_COMMON__
+#define __EYE_ADAPTATION_COMMON__
+
+
+// Single pixel RT used for EyeAdaptation, use function EyeAdaptationLookup() to access.
+Texture2D EyeAdaptation;
+#if FEATURE_LEVEL >= FEATURE_LEVEL_SM5
+//Provides access to the EyeAdaptation RT. Only available in SM5.
+float EyeAdaptationLookup()
+{
+	return EyeAdaptation.Load(int3(0, 0, 0)).r;
+}
+#endif
+
+
+#endif // __EYE_ADAPTATION_COMMON__
diff --git a/Engine/Shaders/MaterialTemplate.usf b/Engine/Shaders/MaterialTemplate.usf
--- a/Engine/Shaders/MaterialTemplate.usf
+++ b/Engine/Shaders/MaterialTemplate.usf
@@ -935,7 +935,12 @@ MaterialFloat2 GetLightmapUVs(FMaterialPixelParameters Parameters)
 #endif
 }
 
-#if NEEDS_SCENE_TEXTURES || USES_EYE_ADAPTATION
+#if USES_EYE_ADAPTATION
+#include "EyeAdaptationCommon.usf"
+#endif
+
+#if POST_PROCESS_MATERIAL
+// only Postprocess material bind PostprocessInput0, PostprocessInput0Size, PostprocessInput0MinMax ...
 #include "PostProcessCommon.usf"			// PostprocessInput0
 #endif
 
@@ -970,6 +975,34 @@ float3 GetShadingModelColor(FGBufferData GBuffer)
 #endif
 }
 
+float4 GetPostProcessInputSize(uint Index)
+{
+#if POST_PROCESS_MATERIAL
+		 if (Index == 0)	return PostprocessInput0Size;
+	else if (Index == 1)	return PostprocessInput1Size;
+	else if (Index == 2)	return PostprocessInput2Size;
+	else if (Index == 3)	return PostprocessInput3Size;
+	else if (Index == 4)	return PostprocessInput4Size;
+	else if (Index == 5)	return PostprocessInput5Size;
+	else if (Index == 6)	return PostprocessInput6Size;
+#endif
+	return 1.0f;
+}
+
+float4 GetPostProcessInputMinMax(uint Index)
+{
+#if POST_PROCESS_MATERIAL
+		 if (Index == 0)	return PostprocessInput0MinMax;
+	else if (Index == 1)	return PostprocessInput1MinMax;
+	else if (Index == 2)	return PostprocessInput2MinMax;
+	else if (Index == 3)	return PostprocessInput3MinMax;
+	else if (Index == 4)	return PostprocessInput4MinMax;
+	else if (Index == 5)	return PostprocessInput5MinMax;
+	else if (Index == 6)	return PostprocessInput6MinMax;
+#endif
+	return 1.0f;
+}
+
 /** Applies an offset to the scene texture lookup and decodes the HDR linear space color. */
 float4 SceneTextureLookup(float2 UV, int SceneTextureIndex, bool bFiltered)
 {
@@ -992,6 +1025,7 @@ float4 SceneTextureLookup(float2 UV, int SceneTextureIndex, bool bFiltered)
 	else if (SceneTextureIndex == 11)	return ScreenSpaceData.GBuffer.Roughness;
 	else if (SceneTextureIndex == 12)	return ScreenSpaceData.GBuffer.GBufferAO;
 	else if (SceneTextureIndex == 13)	return ScreenSpaceData.GBuffer.CustomDepth;
+#if POST_PROCESS_MATERIAL
 	else if (SceneTextureIndex == 14)	if (bFiltered) { return Texture2DSample(PostprocessInput0, BilinearTextureSampler0, UV);} else { return Texture2DSample(PostprocessInput0, PostprocessInput0Sampler, UV);} 
 	else if (SceneTextureIndex == 15)	if (bFiltered) { return Texture2DSample(PostprocessInput1, BilinearTextureSampler0, UV);} else { return Texture2DSample(PostprocessInput1, PostprocessInput1Sampler, UV);} 
 	else if (SceneTextureIndex == 16)	if (bFiltered) { return Texture2DSample(PostprocessInput2, BilinearTextureSampler0, UV);} else { return Texture2DSample(PostprocessInput2, PostprocessInput2Sampler, UV);} 
@@ -999,6 +1033,7 @@ float4 SceneTextureLookup(float2 UV, int SceneTextureIndex, bool bFiltered)
 	else if (SceneTextureIndex == 18)	if (bFiltered) { return Texture2DSample(PostprocessInput4, BilinearTextureSampler0, UV);} else { return Texture2DSample(PostprocessInput4, PostprocessInput4Sampler, UV);} 
 	else if (SceneTextureIndex == 19)	if (bFiltered) { return Texture2DSample(PostprocessInput5, BilinearTextureSampler0, UV);} else { return Texture2DSample(PostprocessInput5, PostprocessInput5Sampler, UV);} 
 	else if (SceneTextureIndex == 20)	if (bFiltered) { return Texture2DSample(PostprocessInput6, BilinearTextureSampler0, UV);} else { return Texture2DSample(PostprocessInput6, PostprocessInput6Sampler, UV);} 
+#endif // POST_PROCESS_MATERIAL
 	else if (SceneTextureIndex == 21)	return ScreenSpaceData.GBuffer.DecalMask;
 	else if (SceneTextureIndex == 22)	return float4(GetShadingModelColor(ScreenSpaceData.GBuffer), 1);
 	else if (SceneTextureIndex == 23)	return ScreenSpaceData.AmbientOcclusion;
@@ -1037,6 +1072,7 @@ float4 SceneTextureLookup(float2 UV, int SceneTextureIndex, bool bFiltered)
 		case 12: return ScreenSpaceData.GBuffer.GBufferAO;
 		// PPI_CustomDepth
 		case 13: return ScreenSpaceData.GBuffer.CustomDepth;
+#if POST_PROCESS_MATERIAL
 		// PPI_PostprocessInput0
 		case 14: return Texture2DSample(PostprocessInput0, bFiltered ? BilinearTextureSampler0 : PostprocessInput0Sampler, UV);
 		// PPI_PostprocessInput1
@@ -1051,6 +1087,7 @@ float4 SceneTextureLookup(float2 UV, int SceneTextureIndex, bool bFiltered)
 		case 19: return Texture2DSample(PostprocessInput5, bFiltered ? BilinearTextureSampler0 : PostprocessInput5Sampler, UV);
 		// PPI_PostprocessInput6
 		case 20: return Texture2DSample(PostprocessInput6, bFiltered ? BilinearTextureSampler0 : PostprocessInput6Sampler, UV);
+#endif // __POST_PROCESS_COMMON__
 		// PPI_DecalMask
 		case 21: return ScreenSpaceData.GBuffer.DecalMask;
 		// PPI_ShadingModel
diff --git a/Engine/Shaders/PostProcessBloom.usf b/Engine/Shaders/PostProcessBloom.usf
--- a/Engine/Shaders/PostProcessBloom.usf
+++ b/Engine/Shaders/PostProcessBloom.usf
@@ -6,6 +6,7 @@
 
 #include "Common.usf"
 #include "PostProcessCommon.usf"
+#include "EyeAdaptationCommon.usf"
 
 // for VisualizeBloomOverlayPS
 // xyz:Bloom1Tint.rgb, w:unused
diff --git a/Engine/Shaders/PostProcessCommon.usf b/Engine/Shaders/PostProcessCommon.usf
--- a/Engine/Shaders/PostProcessCommon.usf
+++ b/Engine/Shaders/PostProcessCommon.usf
@@ -48,16 +48,6 @@ float4 PostprocessInput4MinMax;
 float4 PostprocessInput5MinMax;
 float4 PostprocessInput6MinMax;
 
-// Single pixel RT used for EyeAdaptation, use function EyeAdaptationLookup() to access.
-Texture2D EyeAdaptation;
-#if FEATURE_LEVEL >= FEATURE_LEVEL_SM5
-//Provides access to the EyeAdaptation RT. Only available in SM5.
-float EyeAdaptationLookup()
-{
-	return EyeAdaptation.Load(int3(0, 0, 0)).r;
-}
-#endif
-
 // viewport width, height, 1/width, 1/height (scaled to the current rendertarget resolution), depends on SetViewportAndCallRHI() call
 float4 ViewportSize;
 // in pixels (scaled to the current rendertarget resolution), float4(minx,miny,maxx,maxy), depends on SetViewportAndCallRHI() call
diff --git a/Engine/Shaders/PostProcessTonemap.usf b/Engine/Shaders/PostProcessTonemap.usf
--- a/Engine/Shaders/PostProcessTonemap.usf
+++ b/Engine/Shaders/PostProcessTonemap.usf
@@ -8,6 +8,7 @@
 #include "PostProcessCommon.usf"
 #include "TonemapCommon.usf"
 #include "GammaCorrectionCommon.usf"
+#include "EyeAdaptationCommon.usf"
 //
 // 64x64 tiling noise texture, optimized to have mostly high frequency content
 Texture2D NoiseTexture;
diff --git a/Engine/Shaders/ShaderVersion.usf b/Engine/Shaders/ShaderVersion.usf
--- a/Engine/Shaders/ShaderVersion.usf
+++ b/Engine/Shaders/ShaderVersion.usf
@@ -2,4 +2,4 @@
 // Each time the console command is executed it generates a new GUID. As this file is included
 // in common.usf (which should be included in any shader) it allows to invalidate the shader DDC.
 // 
-// GUID = 7D700E574A85DE267C385C9E997BCC23
\ No newline at end of file
+// GUID = 7D700E574A85DE267C385C9E997BCC23  ds
\ No newline at end of file
diff --git a/Engine/Source/Runtime/Engine/Private/Materials/HLSLMaterialTranslator.h b/Engine/Source/Runtime/Engine/Private/Materials/HLSLMaterialTranslator.h
--- a/Engine/Source/Runtime/Engine/Private/Materials/HLSLMaterialTranslator.h
+++ b/Engine/Source/Runtime/Engine/Private/Materials/HLSLMaterialTranslator.h
@@ -2498,14 +2498,7 @@ protected:
 		{
 			int Index = SceneTextureId - PPI_PostProcessInput0;
 
-			if(bInvert)
-			{
-				return AddCodeChunk(MCT_Float2, TEXT("PostprocessInput%dSize.zw"), Index);
-			}
-			else
-			{
-				return AddCodeChunk(MCT_Float2, TEXT("PostprocessInput%dSize.xy"), Index);
-			}
+			return AddCodeChunk(MCT_Float2, TEXT("GetPostProcessInputSize(%d).%s"), Index, bInvert ? TEXT("zw") : TEXT("xy"));
 		}
 		else
 		{
@@ -2538,7 +2531,7 @@ protected:
 		{
 			int Index = SceneTextureId - PPI_PostProcessInput0;
 
-			return AddCodeChunk(MCT_Float2,TEXT("PostprocessInput%dMinMax.xy"), Index);
+			return AddCodeChunk(MCT_Float2, TEXT("GetPostProcessInputMinMax(%d).xy"), Index);
 		}
 		else
 		{			
@@ -2562,7 +2555,7 @@ protected:
 		{
 			int Index = SceneTextureId - PPI_PostProcessInput0;
 
-			return AddCodeChunk(MCT_Float2,TEXT("PostprocessInput%dMinMax.zw"), Index);
+			return AddCodeChunk(MCT_Float2, TEXT("GetPostProcessInputMinMax(%d).zw"), Index);
 		}
 		else
 		{			
diff --git a/Engine/Source/Runtime/Renderer/Private/PostProcess/PostProcessMaterial.cpp b/Engine/Source/Runtime/Renderer/Private/PostProcess/PostProcessMaterial.cpp
--- a/Engine/Source/Runtime/Renderer/Private/PostProcess/PostProcessMaterial.cpp
+++ b/Engine/Source/Runtime/Renderer/Private/PostProcess/PostProcessMaterial.cpp
@@ -26,24 +26,38 @@ public:
 		return (Material->GetMaterialDomain() == MD_PostProcess) && IsFeatureLevelSupported(Platform, ERHIFeatureLevel::SM4);
 	}
 
+	static void ModifyCompilationEnvironment(EShaderPlatform Platform, const class FMaterial* Material, FShaderCompilerEnvironment& OutEnvironment)
+	{
+		FMaterialShader::ModifyCompilationEnvironment(Platform, OutEnvironment);
+
+		OutEnvironment.SetDefine(TEXT("POST_PROCESS_MATERIAL"), 1);
+	}
+
+
 	FPostProcessMaterialVS( )	{ }
 	FPostProcessMaterialVS(const ShaderMetaType::CompiledShaderInitializerType& Initializer)
 		: FMaterialShader(Initializer)
 	{
+		PostprocessParameter.Bind(Initializer.ParameterMap);
 	}
 
 	void SetParameters(FRHICommandList& RHICmdList, const FRenderingCompositePassContext& Context )
 	{
-		FMaterialShader::SetParameters(RHICmdList, GetVertexShader(), Context.View);
+		const FVertexShaderRHIParamRef ShaderRHI = GetVertexShader();
+		FMaterialShader::SetParameters(RHICmdList, ShaderRHI, Context.View);
+		PostprocessParameter.SetVS(ShaderRHI, Context, TStaticSamplerState<SF_Point, AM_Clamp, AM_Clamp, AM_Clamp>::GetRHI());
 	}
 
 	// Begin FShader interface
 	virtual bool Serialize(FArchive& Ar) override
 	{
 		bool bShaderHasOutdatedParameters = FMaterialShader::Serialize(Ar);
+		Ar << PostprocessParameter;
 		return bShaderHasOutdatedParameters;
 	}
 	//  End FShader interface 
+private:
+	FPostProcessPassParameters PostprocessParameter;
 };
 
 IMPLEMENT_MATERIAL_SHADER_TYPE(,FPostProcessMaterialVS,TEXT("PostProcessMaterialShaders"),TEXT("MainVS"),SF_Vertex);
@@ -64,6 +78,13 @@ public:
 		return (Material->GetMaterialDomain() == MD_PostProcess) && IsFeatureLevelSupported(Platform, ERHIFeatureLevel::SM4);
 	}
 
+	static void ModifyCompilationEnvironment(EShaderPlatform Platform, const class FMaterial* Material, FShaderCompilerEnvironment& OutEnvironment)
+	{
+		FMaterialShader::ModifyCompilationEnvironment(Platform, OutEnvironment);
+
+		OutEnvironment.SetDefine(TEXT("POST_PROCESS_MATERIAL"), 1);
+	}
+
 	FPostProcessMaterialPS() {}
 	FPostProcessMaterialPS(const ShaderMetaType::CompiledShaderInitializerType& Initializer):
 		FMaterialShader(Initializer)
diff --git a/Engine/Source/Runtime/Renderer/Private/ShaderBaseClasses.cpp b/Engine/Source/Runtime/Renderer/Private/ShaderBaseClasses.cpp
--- a/Engine/Source/Runtime/Renderer/Private/ShaderBaseClasses.cpp
+++ b/Engine/Source/Runtime/Renderer/Private/ShaderBaseClasses.cpp
@@ -66,7 +66,6 @@ FMaterialShader::FMaterialShader(const FMaterialShaderType::CompiledShaderInitia
 	LightAttenuation.Bind(Initializer.ParameterMap, TEXT("LightAttenuationTexture"));
 	LightAttenuationSampler.Bind(Initializer.ParameterMap, TEXT("LightAttenuationTextureSampler"));
 	AtmosphericFogTextureParameters.Bind(Initializer.ParameterMap);
-	PostprocessParameter.Bind(Initializer.ParameterMap);
 	EyeAdaptation.Bind(Initializer.ParameterMap, TEXT("EyeAdaptation"));
 	// only used it Material has expression that requires PerlinNoiseGradientTexture
 	PerlinNoiseGradientTexture.Bind(Initializer.ParameterMap,TEXT("PerlinNoiseGradientTexture"));
@@ -348,7 +347,6 @@ bool FMaterialShader::Serialize(FArchive& Ar)
 		Ar << DebugDescription;
 	}
 	Ar << AtmosphericFogTextureParameters;
-	Ar << PostprocessParameter;
 	Ar << EyeAdaptation;
 	Ar << PerlinNoiseGradientTexture;
 	Ar << PerlinNoiseGradientTextureSampler;
diff --git a/Engine/Source/Runtime/Renderer/Public/MaterialShader.h b/Engine/Source/Runtime/Renderer/Public/MaterialShader.h
--- a/Engine/Source/Runtime/Renderer/Public/MaterialShader.h
+++ b/Engine/Source/Runtime/Renderer/Public/MaterialShader.h
@@ -148,7 +148,6 @@ private:
 
 	// For materials using atmospheric fog color 
 	FAtmosphereShaderTextureParameters AtmosphericFogTextureParameters;
-	FPostProcessPassParameters PostprocessParameter;
 
 	//Use of the eye adaptation texture here is experimental and potentially dangerous as it can introduce a feedback loop. May be removed.
 	FShaderResourceParameter EyeAdaptation;
