Merge Morpheus 120hz and reprojection support.
Also fixes for NGS2A3D

diff --git a/Engine/Plugins/Experimental/SimpleHMD/Source/SimpleHMD/Private/SimpleHMD.cpp b/Engine/Plugins/Experimental/SimpleHMD/Source/SimpleHMD/Private/SimpleHMD.cpp
--- a/Engine/Plugins/Experimental/SimpleHMD/Source/SimpleHMD/Private/SimpleHMD.cpp
+++ b/Engine/Plugins/Experimental/SimpleHMD/Source/SimpleHMD/Private/SimpleHMD.cpp
@@ -368,7 +368,7 @@ void FSimpleHMD::PreRenderView_RenderThread(FSceneView& View)
 	check(IsInRenderingThread());
 }
 
-void FSimpleHMD::PreRenderViewFamily_RenderThread(FSceneViewFamily& ViewFamily)
+void FSimpleHMD::PreRenderViewFamily_RenderThread(FRHICommandListImmediate& RHICmdList, FSceneViewFamily& ViewFamily)
 {
 	check(IsInRenderingThread());
 }
diff --git a/Engine/Plugins/Experimental/SimpleHMD/Source/SimpleHMD/Private/SimpleHMD.h b/Engine/Plugins/Experimental/SimpleHMD/Source/SimpleHMD/Private/SimpleHMD.h
--- a/Engine/Plugins/Experimental/SimpleHMD/Source/SimpleHMD/Private/SimpleHMD.h
+++ b/Engine/Plugins/Experimental/SimpleHMD/Source/SimpleHMD/Private/SimpleHMD.h
@@ -80,7 +80,7 @@ public:
 	virtual void ModifyShowFlags(FEngineShowFlags& ShowFlags) override;
 	virtual void SetupView(FSceneViewFamily& InViewFamily, FSceneView& InView) override;
 	virtual void PreRenderView_RenderThread(FSceneView& InView) override;
-	virtual void PreRenderViewFamily_RenderThread(FSceneViewFamily& InViewFamily) override;
+	virtual void PreRenderViewFamily_RenderThread(FRHICommandListImmediate& RHICmdList, FSceneViewFamily& InViewFamily) override;
 
 public:
 	/** Constructor */
diff --git a/Engine/Plugins/Runtime/GearVR/Source/GearVR/Private/GearVR.cpp b/Engine/Plugins/Runtime/GearVR/Source/GearVR/Private/GearVR.cpp
--- a/Engine/Plugins/Runtime/GearVR/Source/GearVR/Private/GearVR.cpp
+++ b/Engine/Plugins/Runtime/GearVR/Source/GearVR/Private/GearVR.cpp
@@ -1289,7 +1289,7 @@ void FGearVR::PreRenderView_RenderThread(FSceneView& View)
 	}
 }
 
-void FGearVR::PreRenderViewFamily_RenderThread(FSceneViewFamily& ViewFamily)
+void FGearVR::PreRenderViewFamily_RenderThread(FRHICommandListImmediate& RHICmdList, FSceneViewFamily& ViewFamily)
 {
 	check(IsInRenderingThread());
 
diff --git a/Engine/Plugins/Runtime/GearVR/Source/GearVR/Private/GearVR.h b/Engine/Plugins/Runtime/GearVR/Source/GearVR/Private/GearVR.h
--- a/Engine/Plugins/Runtime/GearVR/Source/GearVR/Private/GearVR.h
+++ b/Engine/Plugins/Runtime/GearVR/Source/GearVR/Private/GearVR.h
@@ -96,7 +96,7 @@ public:
     virtual void ModifyShowFlags(FEngineShowFlags& ShowFlags) override;
     virtual void SetupView(FSceneViewFamily& InViewFamily, FSceneView& InView) override;
     virtual void PreRenderView_RenderThread(FSceneView& InView) override;
-	virtual void PreRenderViewFamily_RenderThread(FSceneViewFamily& InViewFamily) override;
+	virtual void PreRenderViewFamily_RenderThread(FRHICommandListImmediate& RHICmdList, FSceneViewFamily& InViewFamily) override;
 
 	/** Positional tracking control methods */
 	virtual bool IsPositionalTrackingEnabled() const override;
diff --git a/Engine/Plugins/Runtime/OculusRift/Source/OculusRift/Private/OculusRiftHMD.h b/Engine/Plugins/Runtime/OculusRift/Source/OculusRift/Private/OculusRiftHMD.h
--- a/Engine/Plugins/Runtime/OculusRift/Source/OculusRift/Private/OculusRiftHMD.h
+++ b/Engine/Plugins/Runtime/OculusRift/Source/OculusRift/Private/OculusRiftHMD.h
@@ -130,7 +130,7 @@ public:
     virtual void ModifyShowFlags(FEngineShowFlags& ShowFlags) override;
     virtual void SetupView(FSceneViewFamily& InViewFamily, FSceneView& InView) override;
     virtual void PreRenderView_RenderThread(FSceneView& InView) override;
-	virtual void PreRenderViewFamily_RenderThread(FSceneViewFamily& InViewFamily) override;
+	virtual void PreRenderViewFamily_RenderThread(FRHICommandListImmediate& RHICmdList, FSceneViewFamily& InViewFamily) override;
 
 	/** Positional tracking control methods */
 	virtual bool IsPositionalTrackingEnabled() const override;
diff --git a/Engine/Plugins/Runtime/OculusRift/Source/OculusRift/Private/OculusRiftRender.cpp b/Engine/Plugins/Runtime/OculusRift/Source/OculusRift/Private/OculusRiftRender.cpp
--- a/Engine/Plugins/Runtime/OculusRift/Source/OculusRift/Private/OculusRiftRender.cpp
+++ b/Engine/Plugins/Runtime/OculusRift/Source/OculusRift/Private/OculusRiftRender.cpp
@@ -145,7 +145,7 @@ void FOculusRiftHMD::GetTimewarpMatrices_RenderThread(EStereoscopicPass StereoPa
 #endif
 }
 
-void FOculusRiftHMD::PreRenderViewFamily_RenderThread(FSceneViewFamily& ViewFamily)
+void FOculusRiftHMD::PreRenderViewFamily_RenderThread(FRHICommandListImmediate& RHICmdList, FSceneViewFamily& ViewFamily)
 {
 	check(IsInRenderingThread());
 
diff --git a/Engine/Plugins/Runtime/Steam/SteamVR/Source/SteamVR/Private/SteamVRHMD.cpp b/Engine/Plugins/Runtime/Steam/SteamVR/Source/SteamVR/Private/SteamVRHMD.cpp
--- a/Engine/Plugins/Runtime/Steam/SteamVR/Source/SteamVR/Private/SteamVRHMD.cpp
+++ b/Engine/Plugins/Runtime/Steam/SteamVR/Source/SteamVR/Private/SteamVRHMD.cpp
@@ -649,7 +649,7 @@ void FSteamVRHMD::PreRenderView_RenderThread(FSceneView& View)
  	View.UpdateViewMatrix();
 }
 
-void FSteamVRHMD::PreRenderViewFamily_RenderThread(FSceneViewFamily& ViewFamily)
+void FSteamVRHMD::PreRenderViewFamily_RenderThread(FRHICommandListImmediate& RHICmdList, FSceneViewFamily& ViewFamily)
 {
 	check(IsInRenderingThread());
 	GetActiveRHIBridgeImpl()->BeginRendering();
diff --git a/Engine/Plugins/Runtime/Steam/SteamVR/Source/SteamVR/Private/SteamVRHMD.h b/Engine/Plugins/Runtime/Steam/SteamVR/Source/SteamVR/Private/SteamVRHMD.h
--- a/Engine/Plugins/Runtime/Steam/SteamVR/Source/SteamVR/Private/SteamVRHMD.h
+++ b/Engine/Plugins/Runtime/Steam/SteamVR/Source/SteamVR/Private/SteamVRHMD.h
@@ -112,7 +112,7 @@ public:
 	virtual void ModifyShowFlags(FEngineShowFlags& ShowFlags) override;
 	virtual void SetupView(FSceneViewFamily& InViewFamily, FSceneView& InView) override;
 	virtual void PreRenderView_RenderThread(FSceneView& InView) override;
-	virtual void PreRenderViewFamily_RenderThread(FSceneViewFamily& InViewFamily) override;
+	virtual void PreRenderViewFamily_RenderThread(FRHICommandListImmediate& RHICmdList, FSceneViewFamily& InViewFamily) override;
 
 	virtual void UpdateViewport(bool bUseSeparateRenderTarget, const FViewport& Viewport, SViewport*) override;
 
diff --git a/Engine/Shaders/PostProcessTemporalCommon.usf b/Engine/Shaders/PostProcessTemporalCommon.usf
--- a/Engine/Shaders/PostProcessTemporalCommon.usf
+++ b/Engine/Shaders/PostProcessTemporalCommon.usf
@@ -1,7 +1,7 @@
 // Copyright 1998-2015 Epic Games, Inc. All Rights Reserved.
 
 /*=============================================================================
-	   PostProcessTemporalCommon.usf: Common Temporal AA Functionality
+PostProcessTemporalCommon.usf: Common Temporal AA Functionality
 -------------------------------------------------------------------------------
 This is the common bits of the temporal AA shader
 which can be configured via defines for various temporal AA passes.
@@ -10,27 +10,27 @@ which can be configured via defines for various temporal AA passes.
 // 1 = Use tighter AABB clamp for history.
 // 0 = Use simple min/max clamp.
 #ifndef AA_AABB
-	#define AA_AABB 1
+#define AA_AABB 1
 #endif
 
 // 0 = Anti-alias the alpha channel also (not getting used currently). 
 // 1 = Use alpha channel to improve quality (required for primary AA).
 //     Leverages dead code removal to work in RGB instead of RGBA.
 #ifndef AA_ALPHA
-	#define AA_ALPHA 1
+#define AA_ALPHA 1
 #endif
 
 // Cross distance in pixels used in depth search X pattern.
 // 0 = Turn this feature off.
 // 2 = Is required for standard temporal AA pass.
 #ifndef AA_CROSS
-	#define AA_CROSS 2
+#define AA_CROSS 2
 #endif
-	
+
 // 1 = Render in blue, with green = diff between frames, red = alpha channel.
 // 0 = Non-debug.
 #ifndef AA_DEBUG
-	#define AA_DEBUG 0
+#define AA_DEBUG 0
 #endif
 
 // 2 = Dilate in cross pattern by 2 pixels in distance (this can be larger than 2 if required).
@@ -39,736 +39,543 @@ which can be configured via defines for various temporal AA passes.
 //     This is only valid for AA_ALPHA == 1
 // 0 = Turn off.
 #ifndef AA_DILATE
-	#define AA_DILATE AA_ALPHA
+#define AA_DILATE AA_ALPHA
 #endif
 
 // 1 = Use dynamic motion.
 // 0 = Skip dynamic motion, currently required for half resolution passes.
 #ifndef AA_DYNAMIC
-	#define AA_DYNAMIC 1
+#define AA_DYNAMIC 1
 #endif
 
 // 1 = Use filtered sample.
 // 0 = Use center sample.
 #ifndef AA_FILTERED
-	#define AA_FILTERED 1
+#define AA_FILTERED 1
 #endif
 
 // 1 = Improve quality on converged edges.
 // 0 = Default.
 #ifndef AA_GRAIN
-	#define AA_GRAIN 0
+#define AA_GRAIN 0
 #endif
 
 // 1 = Use higher quality HDR weighted filtering.
 // 0 = Don't use.
 #ifndef AA_HDR
-	#define AA_HDR 1
+#define AA_HDR 1
 #endif
 
 // 1 = Use manual HDR weighted bilinear filtering for history.
 #if AA_HDR_HISTORY
-	#define AA_HDR_HISTORY 1
+#define AA_HDR_HISTORY 1
 #endif
 
 // 0 = Dynamic motion based lerp value (default).
 // non-zero = Use 1/LERP fixed lerp value (used for reflections).
 #ifndef AA_LERP
-	#define AA_LERP 0
+#define AA_LERP 0
 #endif
 
 // 1 = Use extra lowpass filter for quality bump.
 // 0 = Don't use.
 #ifndef AA_LOWPASS
-	#define AA_LOWPASS 1
+#define AA_LOWPASS 1
 #endif
 
 // 1 = Use higher quality round clamp.
 // 0 = Use lower quality but faster box clamp.
 #ifndef AA_ROUND
-	#define AA_ROUND 1
+#define AA_ROUND 1
 #endif
 
 // 1 = Use extra clamp to avoid NANs
 // 0 = Don't use.
 #ifndef AA_NAN
-	#define AA_NAN 1
+#define AA_NAN 1
 #endif
 
 // Fix for lack of borders during current frame filter.
 #ifndef AA_BORDER
-	#define AA_BORDER 0
+#define AA_BORDER 0
 #endif
 
 // Force clamp on alpha.
 #ifndef AA_FORCE_ALPHA_CLAMP
-	#define AA_FORCE_ALPHA_CLAMP 0
+#define AA_FORCE_ALPHA_CLAMP 0
 #endif
 
 // Use velocity weighting.
 #ifndef AA_VELOCITY_WEIGHTING
-	#define AA_VELOCITY_WEIGHTING 0
+#define AA_VELOCITY_WEIGHTING 0
 #endif
 
 // Use YCoCg path.
 #ifndef AA_YCOCG
-	#define AA_YCOCG 0
+#define AA_YCOCG 0
 #endif
 
 // Use green as luma.
 #ifndef AA_GREEN_AS_LUMA
-	#define AA_GREEN_AS_LUMA AA_YCOCG
+#define AA_GREEN_AS_LUMA AA_YCOCG
 #endif
 
 // Special adjustments for DOF.
 #ifndef AA_DOF
-	#define AA_DOF 0
+#define AA_DOF 0
 #endif
 
 // Completely non-optimized prototype for solving noise problem.
 #ifndef AA_TEST
-	#define AA_TEST 0
-#endif
-
-
-	// FIND MOTION OF PIXEL AND NEAREST IN NEIGHBORHOOD
-    // ------------------------------------------------
-	float3 PosN; // Position of this pixel, possibly later nearest pixel in neighborhood.
-	PosN.xy = UVAndScreenPos.zw * float2(0.5, -0.5) + 0.5; // View position [0 to 1] flipped in Y.
-	PosN.z = SceneDepthTexture.SampleLevel(SceneDepthTextureSampler, UVAndScreenPos.xy, 0).r;
-	// Screen position of minimum depth.
-	float2 VelocityOffset = float2(0.0, 0.0);
-	#if AA_CROSS
-		// For motion vector, use camera/dynamic motion from min depth pixel in pattern around pixel.
-		// This enables better quality outline on foreground against different motion background.
-		// Larger 2 pixel distance "x" works best (because AA diolates surface).
-		float4 Depths;
-		#if FEATURE_LEVEL >= FEATURE_LEVEL_SM5
-			Depths = SceneDepthTexture.GatherRed(SceneDepthTextureSampler, UVAndScreenPos.xy, int2(-AA_CROSS, -AA_CROSS), int2(AA_CROSS, -AA_CROSS), int2(-AA_CROSS, AA_CROSS), int2(AA_CROSS, AA_CROSS));
-		#else
-			Depths.x = SceneDepthTexture.SampleLevel(SceneDepthTextureSampler, UVAndScreenPos.xy, 0, int2(-AA_CROSS, -AA_CROSS)).r;
-			Depths.y = SceneDepthTexture.SampleLevel(SceneDepthTextureSampler, UVAndScreenPos.xy, 0, int2( AA_CROSS, -AA_CROSS)).r;
-			Depths.z = SceneDepthTexture.SampleLevel(SceneDepthTextureSampler, UVAndScreenPos.xy, 0, int2(-AA_CROSS,  AA_CROSS)).r;
-			Depths.w = SceneDepthTexture.SampleLevel(SceneDepthTextureSampler, UVAndScreenPos.xy, 0, int2( AA_CROSS,  AA_CROSS)).r;
-		#endif
-		float2 DepthOffset = float2(AA_CROSS, AA_CROSS);
-		float DepthOffsetXx = float(AA_CROSS);
+#define AA_TEST 0
+#endif
+
+
+// FIND MOTION OF PIXEL AND NEAREST IN NEIGHBORHOOD
+// ------------------------------------------------
+float3 PosN; // Position of this pixel, possibly later nearest pixel in neighborhood.
+PosN.xy = UVAndScreenPos.zw * float2(0.5, -0.5) + 0.5; // View position [0 to 1] flipped in Y.
+PosN.z = SceneDepthTexture.SampleLevel(SceneDepthTextureSampler, UVAndScreenPos.xy, 0).r;
+// Screen position of minimum depth.
+float2 VelocityOffset = float2(0.0, 0.0);
+#if AA_CROSS
+// For motion vector, use camera/dynamic motion from min depth pixel in pattern around pixel.
+// This enables better quality outline on foreground against different motion background.
+// Larger 2 pixel distance "x" works best (because AA dilates surface).
+float4 Depths;
+#if FEATURE_LEVEL >= FEATURE_LEVEL_SM5
+Depths = SceneDepthTexture.GatherRed(SceneDepthTextureSampler, UVAndScreenPos.xy, int2(-AA_CROSS, -AA_CROSS), int2(AA_CROSS, -AA_CROSS), int2(-AA_CROSS, AA_CROSS), int2(AA_CROSS, AA_CROSS));
+#else
+Depths.x = SceneDepthTexture.SampleLevel(SceneDepthTextureSampler, UVAndScreenPos.xy, 0, int2(-AA_CROSS, -AA_CROSS)).r;
+Depths.y = SceneDepthTexture.SampleLevel(SceneDepthTextureSampler, UVAndScreenPos.xy, 0, int2(AA_CROSS, -AA_CROSS)).r;
+Depths.z = SceneDepthTexture.SampleLevel(SceneDepthTextureSampler, UVAndScreenPos.xy, 0, int2(-AA_CROSS, AA_CROSS)).r;
+Depths.w = SceneDepthTexture.SampleLevel(SceneDepthTextureSampler, UVAndScreenPos.xy, 0, int2(AA_CROSS, AA_CROSS)).r;
+#endif
+float2 DepthOffset = float2(AA_CROSS, AA_CROSS);
+float DepthOffsetXx = float(AA_CROSS);
 #if HAS_INVERTED_Z_BUFFER
-		//@todo-briank
-		// Nearest depth is the largest depth (depth surface 0=far, 1=near).
-		if(Depths.x > Depths.y) 
-		{
-			DepthOffsetXx = -AA_CROSS;
-		}
-		if(Depths.z > Depths.w) 
-		{
-			DepthOffset.x = -AA_CROSS;
-		}
-		float DepthsXY = max(Depths.x, Depths.y);
-		float DepthsZW = max(Depths.z, Depths.w);
-		if(DepthsXY > DepthsZW) 
-		{
-			DepthOffset.y = -AA_CROSS;
-			DepthOffset.x = DepthOffsetXx; 
-		}
-		float DepthsXYZW = max(DepthsXY, DepthsZW);
-		if(DepthsXYZW > PosN.z) 
-		{
-			// This is offset for reading from velocity texture.
-			// This supports half or fractional resolution velocity textures.
-			// With the assumption that UV position scales between velocity and color.
-			VelocityOffset = DepthOffset * PostprocessInput0Size.zw;
-			// This is [0 to 1] flipped in Y.
-			PosN.xy = (UVAndScreenPos.zw + DepthOffset * ViewportSize.zw * 2.0) * float2(0.5, -0.5) + 0.5;
-			PosN.z = DepthsXYZW;
-		}
+//@todo-briank
+// Nearest depth is the largest depth (depth surface 0=far, 1=near).
+if (Depths.x > Depths.y)
+{
+	DepthOffsetXx = -AA_CROSS;
+}
+if (Depths.z > Depths.w)
+{
+	DepthOffset.x = -AA_CROSS;
+}
+float DepthsXY = max(Depths.x, Depths.y);
+float DepthsZW = max(Depths.z, Depths.w);
+if (DepthsXY > DepthsZW)
+{
+	DepthOffset.y = -AA_CROSS;
+	DepthOffset.x = DepthOffsetXx;
+}
+float DepthsXYZW = max(DepthsXY, DepthsZW);
+if (DepthsXYZW > PosN.z)
+{
+	// This is offset for reading from velocity texture.
+	// This supports half or fractional resolution velocity textures.
+	// With the assumption that UV position scales between velocity and color.
+	VelocityOffset = DepthOffset * PostprocessInput0Size.zw;
+	// This is [0 to 1] flipped in Y.
+	PosN.xy = (UVAndScreenPos.zw + DepthOffset * ViewportSize.zw * 2.0) * float2(0.5, -0.5) + 0.5;
+	PosN.z = DepthsXYZW;
+}
 #else
 #error Fix me!
 #endif // HAS_INVERTED_Z_BUFFER
-	#endif	// AA_CROSS
-	// Camera motion for pixel or nearest pixel (in ScreenPos space).
-	float ScaleM = 1.0 / (dot(PosN, CameraMotion[0].xyz) + CameraMotion[0].w);
-	float2 BackN;
-	BackN.x = -2.0 * ((PosN.x * ((CameraMotion[1].x * PosN.y) + (CameraMotion[1].y * PosN.z) + CameraMotion[1].z)) + (CameraMotion[1].w * PosN.y) + (CameraMotion[2].x * PosN.x * PosN.x) + (CameraMotion[2].y * PosN.z) + CameraMotion[2].z) * ScaleM;
-	BackN.y =  2.0 * ((PosN.y * ((CameraMotion[3].x * PosN.x) + (CameraMotion[3].y * PosN.z) + CameraMotion[3].z)) + (CameraMotion[3].w * PosN.x) + (CameraMotion[4].x * PosN.y * PosN.y) + (CameraMotion[4].y * PosN.z) + CameraMotion[4].z) * ScaleM;
-	float2 BackTemp = BackN * ViewportSize.xy;
-	#if AA_DYNAMIC
-		float2 VelocityN;
-		#if AA_CROSS
-			VelocityN = PostprocessInput3.SampleLevel(PostprocessInput3Sampler, UVAndScreenPos.xy + VelocityOffset, 0).xy;
-		#else
-			VelocityN = PostprocessInput3.SampleLevel(PostprocessInput3Sampler, UVAndScreenPos.xy, 0).xy;
-		#endif
-		bool DynamicN = VelocityN.x > 0.0;
-		if(DynamicN)
-		{
-			BackN = DecodeVelocityFromTexture(VelocityN);
-		}
-		BackTemp = BackN * ViewportSize.xy;
-	#endif
-	// Save the amount of pixel offset of just camera motion, used later as the amount of blur introduced by history.
-	float HistoryBlurAmp = 2.0;
-	float HistoryBlur = saturate(abs(BackTemp.x) * HistoryBlurAmp + abs(BackTemp.y) * HistoryBlurAmp);
-	float Velocity = sqrt(dot(BackTemp, BackTemp));
-	// Easier to do off screen check before conversion.
-	// BackN is in units of 2pixels/viewportWidthInPixels
-	// This converts back projection vector to [-1 to 1] offset in viewport.
-	BackN = UVAndScreenPos.zw - BackN;
-	bool OffScreen = max(abs(BackN.x), abs(BackN.y)) >= 1.0;
-	// Also clamp to be on screen (fixes problem with DOF).
-	// The .z and .w is the 1/width and 1/height.
-	// This clamps to be a pixel inside the viewport.
-	BackN.x = clamp(BackN.x, -1.0 + ViewportSize.z, 1.0 - ViewportSize.z);
-	BackN.y = clamp(BackN.y, -1.0 + ViewportSize.w, 1.0 - ViewportSize.w);
-	// Convert from [-1 to 1] to view rectangle which is somewhere in [0 to 1].
-	// The extra +0.5 factor is because ScreenPosToPixel.zw is incorrectly computed
-	// as the upper left of the pixel instead of the center of the pixel.
-	BackN = (BackN * ScreenPosToPixel.xy + ScreenPosToPixel.zw + 0.5) * PostprocessInput0Size.zw;
-
-
-
-	// FILTER PIXEL (RESAMPLE TO REMOVE JITTER OFFSET) AND GET NEIGHBORHOOD
-    // --------------------------------------------------------------------
-	// 012
-	// 345
-	// 678
-	#if AA_YCOCG	
-		// Special case, only using 5 taps.
-		float4 Neighbor1 = PostprocessInput0.SampleLevel(PostprocessInput0Sampler, UVAndScreenPos.xy, 0, int2( 0, -1));
-		float4 Neighbor3 = PostprocessInput0.SampleLevel(PostprocessInput0Sampler, UVAndScreenPos.xy, 0, int2(-1,  0));
-		float4 Neighbor4 = PostprocessInput0.SampleLevel(PostprocessInput0Sampler, UVAndScreenPos.xy, 0);
-		float4 Neighbor5 = PostprocessInput0.SampleLevel(PostprocessInput0Sampler, UVAndScreenPos.xy, 0, int2( 1,  0));
-		float4 Neighbor7 = PostprocessInput0.SampleLevel(PostprocessInput0Sampler, UVAndScreenPos.xy, 0, int2( 0,  1));
-		#if 1
-			Neighbor1.rgb *= HdrWeightG(Neighbor1.rgb, InExposureScale);
-			Neighbor3.rgb *= HdrWeightG(Neighbor3.rgb, InExposureScale);
-			Neighbor4.rgb *= HdrWeightG(Neighbor4.rgb, InExposureScale);
-			Neighbor5.rgb *= HdrWeightG(Neighbor5.rgb, InExposureScale);
-			Neighbor7.rgb *= HdrWeightG(Neighbor7.rgb, InExposureScale);
-		#endif
-		Neighbor1.rgb = RGBToYCoCg(Neighbor1.rgb);
-		Neighbor3.rgb = RGBToYCoCg(Neighbor3.rgb);
-		Neighbor4.rgb = RGBToYCoCg(Neighbor4.rgb);
-		Neighbor5.rgb = RGBToYCoCg(Neighbor5.rgb);
-		Neighbor7.rgb = RGBToYCoCg(Neighbor7.rgb);
-		#if 0
-			Neighbor1.xyz *= HdrWeightY(Neighbor1.x, InExposureScale);
-			Neighbor3.xyz *= HdrWeightY(Neighbor3.x, InExposureScale);
-			Neighbor4.xyz *= HdrWeightY(Neighbor4.x, InExposureScale);
-			Neighbor5.xyz *= HdrWeightY(Neighbor5.x, InExposureScale);
-			Neighbor7.xyz *= HdrWeightY(Neighbor7.x, InExposureScale);
-		#endif
-		#if AA_FILTERED
-			float4 Filtered = 
-				Neighbor1 * PlusWeights[0] +
-				Neighbor3 * PlusWeights[1] +
-				Neighbor4 * PlusWeights[2] +
-				Neighbor5 * PlusWeights[3] +
-				Neighbor7 * PlusWeights[4];
-			#if AA_BORDER
-				// Use unfiltered for 1 pixel border.
-				float2 TestPos = abs(UVAndScreenPos.zw);
-				// Add 1 pixel and check if off screen.
-				TestPos += ViewportSize.zw * 2.0;
-				bool FilteredOffScreen = max(TestPos.x, TestPos.y) >= 1.0;
-				if(FilteredOffScreen) 
-				{
-					Filtered = Neighbor4;
-				}
-			#endif
-		#else
-			// Unfiltered.
-			float4 Filtered = Neighbor4;
-		#endif
-		float4 FilteredLow = Filtered;
-		// Neighborhood seems to only need the "+" pattern.
-		float4 NeighborMin = min(min(min(Neighbor1, Neighbor3), min(Neighbor4, Neighbor5)), Neighbor7);		
-		float4 NeighborMax = max(max(max(Neighbor1, Neighbor3), max(Neighbor4, Neighbor5)), Neighbor7);		
-	#else
-		float4 Neighbor0 = PostprocessInput0.SampleLevel(PostprocessInput0Sampler, UVAndScreenPos.xy, 0, int2(-1, -1));
-		float4 Neighbor1 = PostprocessInput0.SampleLevel(PostprocessInput0Sampler, UVAndScreenPos.xy, 0, int2( 0, -1));
-		float4 Neighbor2 = PostprocessInput0.SampleLevel(PostprocessInput0Sampler, UVAndScreenPos.xy, 0, int2( 1, -1));
-		float4 Neighbor3 = PostprocessInput0.SampleLevel(PostprocessInput0Sampler, UVAndScreenPos.xy, 0, int2(-1,  0));
-		float4 Neighbor4 = PostprocessInput0.SampleLevel(PostprocessInput0Sampler, UVAndScreenPos.xy, 0);
-		float4 Neighbor5 = PostprocessInput0.SampleLevel(PostprocessInput0Sampler, UVAndScreenPos.xy, 0, int2( 1,  0));
-		float4 Neighbor6 = PostprocessInput0.SampleLevel(PostprocessInput0Sampler, UVAndScreenPos.xy, 0, int2(-1,  1));
-		float4 Neighbor7 = PostprocessInput0.SampleLevel(PostprocessInput0Sampler, UVAndScreenPos.xy, 0, int2( 0,  1));
-		float4 Neighbor8 = PostprocessInput0.SampleLevel(PostprocessInput0Sampler, UVAndScreenPos.xy, 0, int2( 1,  1));
-		#if AA_GREEN_AS_LUMA
-			Neighbor0.rgb *= HdrWeightG(Neighbor0.rgb, InExposureScale);
-			Neighbor1.rgb *= HdrWeightG(Neighbor1.rgb, InExposureScale);
-			Neighbor2.rgb *= HdrWeightG(Neighbor2.rgb, InExposureScale);
-			Neighbor3.rgb *= HdrWeightG(Neighbor3.rgb, InExposureScale);
-			Neighbor4.rgb *= HdrWeightG(Neighbor4.rgb, InExposureScale);
-			Neighbor5.rgb *= HdrWeightG(Neighbor5.rgb, InExposureScale);
-			Neighbor6.rgb *= HdrWeightG(Neighbor6.rgb, InExposureScale);
-			Neighbor7.rgb *= HdrWeightG(Neighbor7.rgb, InExposureScale);
-			Neighbor8.rgb *= HdrWeightG(Neighbor8.rgb, InExposureScale);
-		#else
-			Neighbor0.rgb *= HdrWeight4(Neighbor0.rgb, InExposureScale);
-			Neighbor1.rgb *= HdrWeight4(Neighbor1.rgb, InExposureScale);
-			Neighbor2.rgb *= HdrWeight4(Neighbor2.rgb, InExposureScale);
-			Neighbor3.rgb *= HdrWeight4(Neighbor3.rgb, InExposureScale);
-			Neighbor4.rgb *= HdrWeight4(Neighbor4.rgb, InExposureScale);
-			Neighbor5.rgb *= HdrWeight4(Neighbor5.rgb, InExposureScale);
-			Neighbor6.rgb *= HdrWeight4(Neighbor6.rgb, InExposureScale);
-			Neighbor7.rgb *= HdrWeight4(Neighbor7.rgb, InExposureScale);
-			Neighbor8.rgb *= HdrWeight4(Neighbor8.rgb, InExposureScale);
-		#endif
-		#if AA_FILTERED
-			float4 Filtered = 
-				Neighbor0 * SampleWeights[0] +
-				Neighbor1 * SampleWeights[1] +
-				Neighbor2 * SampleWeights[2] +
-				Neighbor3 * SampleWeights[3] +
-				Neighbor4 * SampleWeights[4] +
-				Neighbor5 * SampleWeights[5] +
-				Neighbor6 * SampleWeights[6] +
-				Neighbor7 * SampleWeights[7] +
-				Neighbor8 * SampleWeights[8];
-			#if AA_LOWPASS
-				float4 FilteredLow = 
-					Neighbor0 * LowpassWeights[0] +
-					Neighbor1 * LowpassWeights[1] +
-					Neighbor2 * LowpassWeights[2] +
-					Neighbor3 * LowpassWeights[3] +
-					Neighbor4 * LowpassWeights[4] +
-					Neighbor5 * LowpassWeights[5] +
-					Neighbor6 * LowpassWeights[6] +
-					Neighbor7 * LowpassWeights[7] +
-					Neighbor8 * LowpassWeights[8];
-			#else
-				float4 FilteredLow = Filtered;
-			#endif
-			#if AA_BORDER
-				// Use unfiltered for 1 pixel border.
-				float2 TestPos = abs(UVAndScreenPos.zw);
-				// Add 1 pixel and check if off screen.
-				TestPos += ViewportSize.zw * 2.0;
-				bool FilteredOffScreen = max(TestPos.x, TestPos.y) >= 1.0;
-				if(FilteredOffScreen) 
-				{
-					Filtered = Neighbor4;
-					FilteredLow = Neighbor4;
-				}
-			#endif
-		#else
-			// Unfiltered.
-			float4 Filtered = Neighbor4;
-			float4 FilteredLow = Neighbor4;
-		#endif
-		#if AA_ROUND
-			float4 NeighborMin2 = min(min(Neighbor0, Neighbor2), min(Neighbor6, Neighbor8));		
-			float4 NeighborMax2 = max(max(Neighbor0, Neighbor2), max(Neighbor6, Neighbor8));		
-			float4 NeighborMin = min(min(min(Neighbor1, Neighbor3), min(Neighbor4, Neighbor5)), Neighbor7);		
-			float4 NeighborMax = max(max(max(Neighbor1, Neighbor3), max(Neighbor4, Neighbor5)), Neighbor7);		
-			NeighborMin2 = min(NeighborMin2, NeighborMin);
-			NeighborMax2 = max(NeighborMax2, NeighborMax);
-			NeighborMin = NeighborMin * 0.5 + NeighborMin2 * 0.5;
-			NeighborMax = NeighborMax * 0.5 + NeighborMax2 * 0.5;
-		#else
-			float4 NeighborMin = min(min(
-				min(min(Neighbor0, Neighbor1), min(Neighbor2, Neighbor3)),
-				min(min(Neighbor4, Neighbor5), min(Neighbor6, Neighbor7))), Neighbor8);
-			float4 NeighborMax = max(max(
-				max(max(Neighbor0, Neighbor1), max(Neighbor2, Neighbor3)),
-				max(max(Neighbor4, Neighbor5), max(Neighbor6, Neighbor7))), Neighbor8);
-		#endif
-	#endif
-
-
-	// Prototype low frequency test.
-	// This has to be re-done to do the reduction/etc prior at non-full res/etc.
-	#if AA_TEST
-
-		// Get low resolution version of current frame.
-
-		//   abc
-		//  defgh
-		//  ijklm
-		//  nopqr
-		//   stu
-		float4 TxA = PostprocessInput0.SampleLevel(PostprocessInput0Sampler, UVAndScreenPos.xy, 0, int2( -1, -2));
-		float4 TxB = PostprocessInput0.SampleLevel(PostprocessInput0Sampler, UVAndScreenPos.xy, 0, int2(  0, -2));
-		float4 TxC = PostprocessInput0.SampleLevel(PostprocessInput0Sampler, UVAndScreenPos.xy, 0, int2(  1, -2));
-
-		float4 TxD = PostprocessInput0.SampleLevel(PostprocessInput0Sampler, UVAndScreenPos.xy, 0, int2( -2, -1));
-		float4 TxE = PostprocessInput0.SampleLevel(PostprocessInput0Sampler, UVAndScreenPos.xy, 0, int2( -1, -1));
-		float4 TxF = PostprocessInput0.SampleLevel(PostprocessInput0Sampler, UVAndScreenPos.xy, 0, int2(  0, -1));
-		float4 TxG = PostprocessInput0.SampleLevel(PostprocessInput0Sampler, UVAndScreenPos.xy, 0, int2(  1, -1));
-		float4 TxH = PostprocessInput0.SampleLevel(PostprocessInput0Sampler, UVAndScreenPos.xy, 0, int2(  2, -1));
-
-		float4 TxI = PostprocessInput0.SampleLevel(PostprocessInput0Sampler, UVAndScreenPos.xy, 0, int2( -2, 0));
-		float4 TxJ = PostprocessInput0.SampleLevel(PostprocessInput0Sampler, UVAndScreenPos.xy, 0, int2( -1, 0));
-		float4 TxK = PostprocessInput0.SampleLevel(PostprocessInput0Sampler, UVAndScreenPos.xy, 0, int2(  0, 0));
-		float4 TxL = PostprocessInput0.SampleLevel(PostprocessInput0Sampler, UVAndScreenPos.xy, 0, int2(  1, 0));
-		float4 TxM = PostprocessInput0.SampleLevel(PostprocessInput0Sampler, UVAndScreenPos.xy, 0, int2(  2, 0));
-
-		float4 TxN = PostprocessInput0.SampleLevel(PostprocessInput0Sampler, UVAndScreenPos.xy, 0, int2( -2, 1));
-		float4 TxO = PostprocessInput0.SampleLevel(PostprocessInput0Sampler, UVAndScreenPos.xy, 0, int2( -1, 1));
-		float4 TxP = PostprocessInput0.SampleLevel(PostprocessInput0Sampler, UVAndScreenPos.xy, 0, int2(  0, 1));
-		float4 TxQ = PostprocessInput0.SampleLevel(PostprocessInput0Sampler, UVAndScreenPos.xy, 0, int2(  1, 1));
-		float4 TxR = PostprocessInput0.SampleLevel(PostprocessInput0Sampler, UVAndScreenPos.xy, 0, int2(  2, 1));
-
-		float4 TxS = PostprocessInput0.SampleLevel(PostprocessInput0Sampler, UVAndScreenPos.xy, 0, int2( -1, 2));
-		float4 TxT = PostprocessInput0.SampleLevel(PostprocessInput0Sampler, UVAndScreenPos.xy, 0, int2(  0, 2));
-		float4 TxU = PostprocessInput0.SampleLevel(PostprocessInput0Sampler, UVAndScreenPos.xy, 0, int2(  1, 2));
-
-		float4 TxAvg = (TxA+TxB+TxC+ TxD+TxE+TxF+TxG+TxH+ TxI+TxJ+TxK+TxL+TxM+ TxN+TxO+TxP+TxQ+TxR+ TxS+TxT+TxU) * (1.0/21.0);
-
-		TxAvg.rgb *= HdrWeightG(TxAvg.rgb, InExposureScale);
-		TxAvg.rgb = RGBToYCoCg(TxAvg.rgb);
-
-		// Get low resolution version of reprojected frame.
-		float4 TxA2 = PostprocessInput1.SampleLevel(PostprocessInput1Sampler, BackN.xy, 0, int2( -1, -2));
-		float4 TxB2 = PostprocessInput1.SampleLevel(PostprocessInput1Sampler, BackN.xy, 0, int2(  0, -2));
-		float4 TxC2 = PostprocessInput1.SampleLevel(PostprocessInput1Sampler, BackN.xy, 0, int2(  1, -2));
-
-		float4 TxD2 = PostprocessInput1.SampleLevel(PostprocessInput1Sampler, BackN.xy, 0, int2( -2, -1));
-		float4 TxE2 = PostprocessInput1.SampleLevel(PostprocessInput1Sampler, BackN.xy, 0, int2( -1, -1));
-		float4 TxF2 = PostprocessInput1.SampleLevel(PostprocessInput1Sampler, BackN.xy, 0, int2(  0, -1));
-		float4 TxG2 = PostprocessInput1.SampleLevel(PostprocessInput1Sampler, BackN.xy, 0, int2(  1, -1));
-		float4 TxH2 = PostprocessInput1.SampleLevel(PostprocessInput1Sampler, BackN.xy, 0, int2(  2, -1));
-
-		float4 TxI2 = PostprocessInput1.SampleLevel(PostprocessInput1Sampler, BackN.xy, 0, int2( -2, 0));
-		float4 TxJ2 = PostprocessInput1.SampleLevel(PostprocessInput1Sampler, BackN.xy, 0, int2( -1, 0));
-		float4 TxK2 = PostprocessInput1.SampleLevel(PostprocessInput1Sampler, BackN.xy, 0, int2(  0, 0));
-		float4 TxL2 = PostprocessInput1.SampleLevel(PostprocessInput1Sampler, BackN.xy, 0, int2(  1, 0));
-		float4 TxM2 = PostprocessInput1.SampleLevel(PostprocessInput1Sampler, BackN.xy, 0, int2(  2, 0));
-
-		float4 TxN2 = PostprocessInput1.SampleLevel(PostprocessInput1Sampler, BackN.xy, 0, int2( -2, 1));
-		float4 TxO2 = PostprocessInput1.SampleLevel(PostprocessInput1Sampler, BackN.xy, 0, int2( -1, 1));
-		float4 TxP2 = PostprocessInput1.SampleLevel(PostprocessInput1Sampler, BackN.xy, 0, int2(  0, 1));
-		float4 TxQ2 = PostprocessInput1.SampleLevel(PostprocessInput1Sampler, BackN.xy, 0, int2(  1, 1));
-		float4 TxR2 = PostprocessInput1.SampleLevel(PostprocessInput1Sampler, BackN.xy, 0, int2(  2, 1));
-
-		float4 TxS2 = PostprocessInput1.SampleLevel(PostprocessInput1Sampler, BackN.xy, 0, int2( -1, 2));
-		float4 TxT2 = PostprocessInput1.SampleLevel(PostprocessInput1Sampler, BackN.xy, 0, int2(  0, 2));
-		float4 TxU2 = PostprocessInput1.SampleLevel(PostprocessInput1Sampler, BackN.xy, 0, int2(  1, 2));
-
-		float4 TxAvg2 = (TxA2+TxB2+TxC2+ TxD2+TxE2+TxF2+TxG2+TxH2+ TxI2+TxJ2+TxK2+TxL2+TxM2+ TxN2+TxO2+TxP2+TxQ2+TxR2+ TxS2+TxT2+TxU2) * (1.0/21.0);
-
-		TxAvg2.rgb *= HdrWeightG(TxAvg2.rgb, InExposureScale);
-		TxAvg2.rgb = RGBToYCoCg(TxAvg2.rgb);
-
-
-
-
-		TxA.rgb = RGBToYCoCg(TxA.rgb * HdrWeightG(TxA.rgb, InExposureScale));
-		TxB.rgb = RGBToYCoCg(TxB.rgb * HdrWeightG(TxB.rgb, InExposureScale));
-		TxC.rgb = RGBToYCoCg(TxC.rgb * HdrWeightG(TxC.rgb, InExposureScale));
-
-		TxD.rgb = RGBToYCoCg(TxD.rgb * HdrWeightG(TxD.rgb, InExposureScale));
-		TxE.rgb = RGBToYCoCg(TxE.rgb * HdrWeightG(TxE.rgb, InExposureScale));
-		TxF.rgb = RGBToYCoCg(TxF.rgb * HdrWeightG(TxF.rgb, InExposureScale));
-		TxG.rgb = RGBToYCoCg(TxG.rgb * HdrWeightG(TxG.rgb, InExposureScale));
-		TxH.rgb = RGBToYCoCg(TxH.rgb * HdrWeightG(TxH.rgb, InExposureScale));
-
-		TxI.rgb = RGBToYCoCg(TxI.rgb * HdrWeightG(TxI.rgb, InExposureScale));
-		TxJ.rgb = RGBToYCoCg(TxJ.rgb * HdrWeightG(TxJ.rgb, InExposureScale));
-		TxK.rgb = RGBToYCoCg(TxK.rgb * HdrWeightG(TxK.rgb, InExposureScale));
-		TxL.rgb = RGBToYCoCg(TxL.rgb * HdrWeightG(TxL.rgb, InExposureScale));
-		TxM.rgb = RGBToYCoCg(TxM.rgb * HdrWeightG(TxM.rgb, InExposureScale));
-
-		TxN.rgb = RGBToYCoCg(TxN.rgb * HdrWeightG(TxN.rgb, InExposureScale));
-		TxO.rgb = RGBToYCoCg(TxO.rgb * HdrWeightG(TxO.rgb, InExposureScale));
-		TxP.rgb = RGBToYCoCg(TxP.rgb * HdrWeightG(TxP.rgb, InExposureScale));
-		TxQ.rgb = RGBToYCoCg(TxQ.rgb * HdrWeightG(TxQ.rgb, InExposureScale));
-		TxR.rgb = RGBToYCoCg(TxR.rgb * HdrWeightG(TxR.rgb, InExposureScale));
-
-		TxS.rgb = RGBToYCoCg(TxS.rgb * HdrWeightG(TxS.rgb, InExposureScale));
-		TxT.rgb = RGBToYCoCg(TxT.rgb * HdrWeightG(TxT.rgb, InExposureScale));
-		TxU.rgb = RGBToYCoCg(TxU.rgb * HdrWeightG(TxU.rgb, InExposureScale));
-
-
-		TxA2.rgb = RGBToYCoCg(TxA2.rgb * HdrWeightG(TxA2.rgb, InExposureScale));
-		TxB2.rgb = RGBToYCoCg(TxB2.rgb * HdrWeightG(TxB2.rgb, InExposureScale));
-		TxC2.rgb = RGBToYCoCg(TxC2.rgb * HdrWeightG(TxC2.rgb, InExposureScale));
-
-		TxD2.rgb = RGBToYCoCg(TxD2.rgb * HdrWeightG(TxD2.rgb, InExposureScale));
-		TxE2.rgb = RGBToYCoCg(TxE2.rgb * HdrWeightG(TxE2.rgb, InExposureScale));
-		TxF2.rgb = RGBToYCoCg(TxF2.rgb * HdrWeightG(TxF2.rgb, InExposureScale));
-		TxG2.rgb = RGBToYCoCg(TxG2.rgb * HdrWeightG(TxG2.rgb, InExposureScale));
-		TxH2.rgb = RGBToYCoCg(TxH2.rgb * HdrWeightG(TxH2.rgb, InExposureScale));
-
-		TxI2.rgb = RGBToYCoCg(TxI2.rgb * HdrWeightG(TxI2.rgb, InExposureScale));
-		TxJ2.rgb = RGBToYCoCg(TxJ2.rgb * HdrWeightG(TxJ2.rgb, InExposureScale));
-		TxK2.rgb = RGBToYCoCg(TxK2.rgb * HdrWeightG(TxK2.rgb, InExposureScale));
-		TxL2.rgb = RGBToYCoCg(TxL2.rgb * HdrWeightG(TxL2.rgb, InExposureScale));
-		TxM2.rgb = RGBToYCoCg(TxM2.rgb * HdrWeightG(TxM2.rgb, InExposureScale));
-
-		TxN2.rgb = RGBToYCoCg(TxN2.rgb * HdrWeightG(TxN2.rgb, InExposureScale));
-		TxO2.rgb = RGBToYCoCg(TxO2.rgb * HdrWeightG(TxO2.rgb, InExposureScale));
-		TxP2.rgb = RGBToYCoCg(TxP2.rgb * HdrWeightG(TxP2.rgb, InExposureScale));
-		TxQ2.rgb = RGBToYCoCg(TxQ2.rgb * HdrWeightG(TxQ2.rgb, InExposureScale));
-		TxR2.rgb = RGBToYCoCg(TxR2.rgb * HdrWeightG(TxR2.rgb, InExposureScale));
+#endif	// AA_CROSS
+// Camera motion for pixel or nearest pixel (in ScreenPos space).
+float ScaleM = 1.0 / (dot(PosN, CameraMotion[0].xyz) + CameraMotion[0].w);
+float2 BackN;
+BackN.x = -2.0 * ((PosN.x * ((CameraMotion[1].x * PosN.y) + (CameraMotion[1].y * PosN.z) + CameraMotion[1].z)) + (CameraMotion[1].w * PosN.y) + (CameraMotion[2].x * PosN.x * PosN.x) + (CameraMotion[2].y * PosN.z) + CameraMotion[2].z) * ScaleM;
+BackN.y = 2.0 * ((PosN.y * ((CameraMotion[3].x * PosN.x) + (CameraMotion[3].y * PosN.z) + CameraMotion[3].z)) + (CameraMotion[3].w * PosN.x) + (CameraMotion[4].x * PosN.y * PosN.y) + (CameraMotion[4].y * PosN.z) + CameraMotion[4].z) * ScaleM;
+
+//Fix for VR 'wobble' Imprecision in this method makes it not good enough to have on by default.  A proper fix is coming soon, but in the meantime, VR users can put this in locally.
+#if 0
+{
+	float4 ThisClip = float4((UVAndScreenPos.xy - View.ScreenPositionScaleBias.wz) / View.ScreenPositionScaleBias.xy, PosN.z, 1);
+	float4 ThisTranslatedWorld = mul(ThisClip, View.ClipToTranslatedWorld);
+	ThisTranslatedWorld /= ThisTranslatedWorld.w;
+
+	float3 PrevTranslatedWorld = ThisTranslatedWorld.xyz + (View.PrevPreViewTranslation - View.PreViewTranslation);
+	float4 PrevClip = mul(float4(PrevTranslatedWorld, 1), View.PrevTranslatedWorldToClip);
+	float2 PrevScreen = PrevClip.xy / PrevClip.w;
+	//float2 PrevUV = PrevScreen.xy * View.ScreenPositionScaleBias.xy + View.ScreenPositionScaleBias.wz;
+
+	BackN = UVAndScreenPos.zw - PrevScreen - (View.ViewToClip[2].xy - View.PrevProjection[2].xy);
+}
+#endif
 
-		TxS2.rgb = RGBToYCoCg(TxS2.rgb * HdrWeightG(TxS2.rgb, InExposureScale));
-		TxT2.rgb = RGBToYCoCg(TxT2.rgb * HdrWeightG(TxT2.rgb, InExposureScale));
-		TxU2.rgb = RGBToYCoCg(TxU2.rgb * HdrWeightG(TxU2.rgb, InExposureScale));
+float2 BackTemp = BackN * ViewportSize.xy;
+#if AA_DYNAMIC
+float2 VelocityN;
+#if AA_CROSS
+VelocityN = PostprocessInput3.SampleLevel(PostprocessInput3Sampler, UVAndScreenPos.xy + VelocityOffset, 0).xy;
+#else
+VelocityN = PostprocessInput3.SampleLevel(PostprocessInput3Sampler, UVAndScreenPos.xy, 0).xy;
+#endif
+bool DynamicN = VelocityN.x > 0.0;
+if (DynamicN)
+{
+	BackN = DecodeVelocityFromTexture(VelocityN);
+}
+BackTemp = BackN * ViewportSize.xy;
+#endif
+// Save the amount of pixel offset of just camera motion, used later as the amount of blur introduced by history.
+float HistoryBlurAmp = 2.0;
+float HistoryBlur = saturate(abs(BackTemp.x) * HistoryBlurAmp + abs(BackTemp.y) * HistoryBlurAmp);
+float Velocity = sqrt(dot(BackTemp, BackTemp));
+// Easier to do off screen check before conversion.
+// BackN is in units of 2pixels/viewportWidthInPixels
+// This converts back projection vector to [-1 to 1] offset in viewport.
+BackN = UVAndScreenPos.zw - BackN;
+bool OffScreen = max(abs(BackN.x), abs(BackN.y)) >= 1.0;
+// Also clamp to be on screen (fixes problem with DOF).
+// The .z and .w is the 1/width and 1/height.
+// This clamps to be a pixel inside the viewport.
+BackN.x = clamp(BackN.x, -1.0 + ViewportSize.z, 1.0 - ViewportSize.z);
+BackN.y = clamp(BackN.y, -1.0 + ViewportSize.w, 1.0 - ViewportSize.w);
+// Convert from [-1 to 1] to view rectangle which is somewhere in [0 to 1].
+// The extra +0.5 factor is because ScreenPosToPixel.zw is incorrectly computed
+// as the upper left of the pixel instead of the center of the pixel.
+BackN = (BackN * ScreenPosToPixel.xy + ScreenPosToPixel.zw + 0.5) * PostprocessInput0Size.zw;
+
+
+
+// FILTER PIXEL (RESAMPLE TO REMOVE JITTER OFFSET) AND GET NEIGHBORHOOD
+// --------------------------------------------------------------------
+// 012
+// 345
+// 678
+#if AA_YCOCG	
+// Special case, only using 5 taps.
+float4 Neighbor1 = PostprocessInput0.SampleLevel(PostprocessInput0Sampler, UVAndScreenPos.xy, 0, int2(0, -1));
+float4 Neighbor3 = PostprocessInput0.SampleLevel(PostprocessInput0Sampler, UVAndScreenPos.xy, 0, int2(-1, 0));
+float4 Neighbor4 = PostprocessInput0.SampleLevel(PostprocessInput0Sampler, UVAndScreenPos.xy, 0);
+float4 Neighbor5 = PostprocessInput0.SampleLevel(PostprocessInput0Sampler, UVAndScreenPos.xy, 0, int2(1, 0));
+float4 Neighbor7 = PostprocessInput0.SampleLevel(PostprocessInput0Sampler, UVAndScreenPos.xy, 0, int2(0, 1));
+#if 1
+Neighbor1.rgb *= HdrWeightG(Neighbor1.rgb, InExposureScale);
+Neighbor3.rgb *= HdrWeightG(Neighbor3.rgb, InExposureScale);
+Neighbor4.rgb *= HdrWeightG(Neighbor4.rgb, InExposureScale);
+Neighbor5.rgb *= HdrWeightG(Neighbor5.rgb, InExposureScale);
+Neighbor7.rgb *= HdrWeightG(Neighbor7.rgb, InExposureScale);
+#endif
+Neighbor1.rgb = RGBToYCoCg(Neighbor1.rgb);
+Neighbor3.rgb = RGBToYCoCg(Neighbor3.rgb);
+Neighbor4.rgb = RGBToYCoCg(Neighbor4.rgb);
+Neighbor5.rgb = RGBToYCoCg(Neighbor5.rgb);
+Neighbor7.rgb = RGBToYCoCg(Neighbor7.rgb);
+#if 0
+Neighbor1.xyz *= HdrWeightY(Neighbor1.x, InExposureScale);
+Neighbor3.xyz *= HdrWeightY(Neighbor3.x, InExposureScale);
+Neighbor4.xyz *= HdrWeightY(Neighbor4.x, InExposureScale);
+Neighbor5.xyz *= HdrWeightY(Neighbor5.x, InExposureScale);
+Neighbor7.xyz *= HdrWeightY(Neighbor7.x, InExposureScale);
+#endif
+#if AA_FILTERED
+float4 Filtered =
+Neighbor1 * PlusWeights[0] +
+Neighbor3 * PlusWeights[1] +
+Neighbor4 * PlusWeights[2] +
+Neighbor5 * PlusWeights[3] +
+Neighbor7 * PlusWeights[4];
+#if AA_BORDER
+// Use unfiltered for 1 pixel border.
+float2 TestPos = abs(UVAndScreenPos.zw);
+// Add 1 pixel and check if off screen.
+TestPos += ViewportSize.zw * 2.0;
+bool FilteredOffScreen = max(TestPos.x, TestPos.y) >= 1.0;
+if (FilteredOffScreen)
+{
+	Filtered = Neighbor4;
+}
+#endif
+#else
+// Unfiltered.
+float4 Filtered = Neighbor4;
+#endif
+float4 FilteredLow = Filtered;
+// Neighborhood seems to only need the "+" pattern.
+float4 NeighborMin = min(min(min(Neighbor1, Neighbor3), min(Neighbor4, Neighbor5)), Neighbor7);
+float4 NeighborMax = max(max(max(Neighbor1, Neighbor3), max(Neighbor4, Neighbor5)), Neighbor7);
+#else
+float4 Neighbor0 = PostprocessInput0.SampleLevel(PostprocessInput0Sampler, UVAndScreenPos.xy, 0, int2(-1, -1));
+float4 Neighbor1 = PostprocessInput0.SampleLevel(PostprocessInput0Sampler, UVAndScreenPos.xy, 0, int2(0, -1));
+float4 Neighbor2 = PostprocessInput0.SampleLevel(PostprocessInput0Sampler, UVAndScreenPos.xy, 0, int2(1, -1));
+float4 Neighbor3 = PostprocessInput0.SampleLevel(PostprocessInput0Sampler, UVAndScreenPos.xy, 0, int2(-1, 0));
+float4 Neighbor4 = PostprocessInput0.SampleLevel(PostprocessInput0Sampler, UVAndScreenPos.xy, 0);
+float4 Neighbor5 = PostprocessInput0.SampleLevel(PostprocessInput0Sampler, UVAndScreenPos.xy, 0, int2(1, 0));
+float4 Neighbor6 = PostprocessInput0.SampleLevel(PostprocessInput0Sampler, UVAndScreenPos.xy, 0, int2(-1, 1));
+float4 Neighbor7 = PostprocessInput0.SampleLevel(PostprocessInput0Sampler, UVAndScreenPos.xy, 0, int2(0, 1));
+float4 Neighbor8 = PostprocessInput0.SampleLevel(PostprocessInput0Sampler, UVAndScreenPos.xy, 0, int2(1, 1));
+#if AA_GREEN_AS_LUMA
+Neighbor0.rgb *= HdrWeightG(Neighbor0.rgb, InExposureScale);
+Neighbor1.rgb *= HdrWeightG(Neighbor1.rgb, InExposureScale);
+Neighbor2.rgb *= HdrWeightG(Neighbor2.rgb, InExposureScale);
+Neighbor3.rgb *= HdrWeightG(Neighbor3.rgb, InExposureScale);
+Neighbor4.rgb *= HdrWeightG(Neighbor4.rgb, InExposureScale);
+Neighbor5.rgb *= HdrWeightG(Neighbor5.rgb, InExposureScale);
+Neighbor6.rgb *= HdrWeightG(Neighbor6.rgb, InExposureScale);
+Neighbor7.rgb *= HdrWeightG(Neighbor7.rgb, InExposureScale);
+Neighbor8.rgb *= HdrWeightG(Neighbor8.rgb, InExposureScale);
+#else
+Neighbor0.rgb *= HdrWeight4(Neighbor0.rgb, InExposureScale);
+Neighbor1.rgb *= HdrWeight4(Neighbor1.rgb, InExposureScale);
+Neighbor2.rgb *= HdrWeight4(Neighbor2.rgb, InExposureScale);
+Neighbor3.rgb *= HdrWeight4(Neighbor3.rgb, InExposureScale);
+Neighbor4.rgb *= HdrWeight4(Neighbor4.rgb, InExposureScale);
+Neighbor5.rgb *= HdrWeight4(Neighbor5.rgb, InExposureScale);
+Neighbor6.rgb *= HdrWeight4(Neighbor6.rgb, InExposureScale);
+Neighbor7.rgb *= HdrWeight4(Neighbor7.rgb, InExposureScale);
+Neighbor8.rgb *= HdrWeight4(Neighbor8.rgb, InExposureScale);
+#endif
+#if AA_FILTERED
+float4 Filtered =
+Neighbor0 * SampleWeights[0] +
+Neighbor1 * SampleWeights[1] +
+Neighbor2 * SampleWeights[2] +
+Neighbor3 * SampleWeights[3] +
+Neighbor4 * SampleWeights[4] +
+Neighbor5 * SampleWeights[5] +
+Neighbor6 * SampleWeights[6] +
+Neighbor7 * SampleWeights[7] +
+Neighbor8 * SampleWeights[8];
+#if AA_LOWPASS
+float4 FilteredLow =
+Neighbor0 * LowpassWeights[0] +
+Neighbor1 * LowpassWeights[1] +
+Neighbor2 * LowpassWeights[2] +
+Neighbor3 * LowpassWeights[3] +
+Neighbor4 * LowpassWeights[4] +
+Neighbor5 * LowpassWeights[5] +
+Neighbor6 * LowpassWeights[6] +
+Neighbor7 * LowpassWeights[7] +
+Neighbor8 * LowpassWeights[8];
+#else
+float4 FilteredLow = Filtered;
+#endif
+#if AA_BORDER
+// Use unfiltered for 1 pixel border.
+float2 TestPos = abs(UVAndScreenPos.zw);
+// Add 1 pixel and check if off screen.
+TestPos += ViewportSize.zw * 2.0;
+bool FilteredOffScreen = max(TestPos.x, TestPos.y) >= 1.0;
+if (FilteredOffScreen)
+{
+	Filtered = Neighbor4;
+	FilteredLow = Neighbor4;
+}
+#endif
+#else
+// Unfiltered.
+float4 Filtered = Neighbor4;
+float4 FilteredLow = Neighbor4;
+#endif
+#if AA_ROUND
+float4 NeighborMin2 = min(min(Neighbor0, Neighbor2), min(Neighbor6, Neighbor8));
+float4 NeighborMax2 = max(max(Neighbor0, Neighbor2), max(Neighbor6, Neighbor8));
+float4 NeighborMin = min(min(min(Neighbor1, Neighbor3), min(Neighbor4, Neighbor5)), Neighbor7);
+float4 NeighborMax = max(max(max(Neighbor1, Neighbor3), max(Neighbor4, Neighbor5)), Neighbor7);
+NeighborMin2 = min(NeighborMin2, NeighborMin);
+NeighborMax2 = max(NeighborMax2, NeighborMax);
+NeighborMin = NeighborMin * 0.5 + NeighborMin2 * 0.5;
+NeighborMax = NeighborMax * 0.5 + NeighborMax2 * 0.5;
+#else
+float4 NeighborMin = min(min(
+	min(min(Neighbor0, Neighbor1), min(Neighbor2, Neighbor3)),
+	min(min(Neighbor4, Neighbor5), min(Neighbor6, Neighbor7))), Neighbor8);
+float4 NeighborMax = max(max(
+	max(max(Neighbor0, Neighbor1), max(Neighbor2, Neighbor3)),
+	max(max(Neighbor4, Neighbor5), max(Neighbor6, Neighbor7))), Neighbor8);
+#endif
+#endif
 
 
-		TxA = abs(TxA - TxA2);
-		TxB = abs(TxB - TxB2);
-		TxC = abs(TxC - TxC2);
 
-		TxD = abs(TxD - TxD2);
-		TxE = abs(TxE - TxE2);
-		TxF = abs(TxF - TxF2);
-		TxG = abs(TxG - TxG2);
-		TxH = abs(TxH - TxH2);
+// FETCH HISTORY
+// -------------
+OutColor = PostprocessInput1.SampleLevel(PostprocessInput1Sampler, BackN.xy, 0);
+#if AA_DEBUG
+Neighbor4.rg = float2(0.0, 0.0);
+NeighborMin.rg = float2(0.0, 0.0);
+NeighborMax.rg = float2(0.0, 0.0);
+Filtered.rg = float2(0.0, 0.0);
+FilteredLow.rg = float2(0.0, 0.0);
+float DebugDiffCurrent = Filtered.b;
+#endif
+#if AA_YCOCG
+#if 1
+OutColor.rgb *= HdrWeightG(OutColor.rgb, InExposureScale);
+#endif
+OutColor.rgb = RGBToYCoCg(OutColor.rgb);
+#if 0
+OutColor.xyz *= HdrWeightY(OutColor.x, InExposureScale);
+#endif
+#else
+#if AA_GREEN_AS_LUMA
+OutColor.rgb *= HdrWeightG(OutColor.rgb, InExposureScale);
+#else
+OutColor.rgb *= HdrWeight4(OutColor.rgb, InExposureScale);
+#endif
+#endif
+#if AA_DEBUG
+OutColor.rg = float2(0.0, 0.0);
+float DebugDiffPrior = OutColor.b;
+#endif
 
-		TxI = abs(TxI - TxI2);
-		TxJ = abs(TxJ - TxJ2);
-		TxK = abs(TxK - TxK2);
-		TxL = abs(TxL - TxL2);
-		TxM = abs(TxM - TxM2);
 
-		TxN = abs(TxN - TxN2);
-		TxO = abs(TxO - TxO2);
-		TxP = abs(TxP - TxP2);
-		TxQ = abs(TxQ - TxQ2);
-		TxR = abs(TxR - TxR2);
 
-		TxS = abs(TxS - TxS2);
-		TxT = abs(TxT - TxT2);
-		TxU = abs(TxU - TxU2);
 
-		float4 TxDifAvg = (TxA+TxB+TxC+ TxD+TxE+TxF+TxG+TxH+ TxI+TxJ+TxK+TxL+TxM+ TxN+TxO+TxP+TxQ+TxR+ TxS+TxT+TxU) * (1.0/21.0);
 
-#if 0
-		// Filtered does not work as good (probably because of lower footprint).
-		float Wx1 = 1.0;
-		float Wx2 = 2.0;
-		float Wx3 = 3.0;
-		float Wx4 = 4.0;
-		float Wx5 = 5.0;
-
-		//   abc
-		//  defgh
-		//  ijklm
-		//  nopqr
-		//   stu
-
-		float4 TxDifAvg = (TxA+TxC+TxD+TxN+TxH+TxR+TxS+TxU)*Wx1 + (TxB+TxI+TxM+TxT)*Wx2 + (TxE+TxG+TxO+TxQ)*Wx3 + (TxF+TxJ+TxL+TxP)*Wx4 + TxK*Wx5;
-
-		TxDifAvg /= (Wx1*8.0)+(Wx2*4.0)+(Wx3*4.0)+(Wx4*4.0)+(Wx5);
-#endif
-
-	#endif
-
-
-
-
-	// FETCH HISTORY
-	// -------------
-	OutColor = PostprocessInput1.SampleLevel(PostprocessInput1Sampler, BackN.xy, 0);
-	#if AA_DEBUG
-		Neighbor4.rg = float2(0.0, 0.0);
-		NeighborMin.rg = float2(0.0, 0.0);
-		NeighborMax.rg = float2(0.0, 0.0);
-		Filtered.rg = float2(0.0, 0.0);
-		FilteredLow.rg = float2(0.0, 0.0);
-		float DebugDiffCurrent = Filtered.b;
-	#endif
-	#if AA_YCOCG
-		#if 1
-			OutColor.rgb *= HdrWeightG(OutColor.rgb, InExposureScale);
-		#endif
-		OutColor.rgb = RGBToYCoCg(OutColor.rgb);
-		#if 0
-			OutColor.xyz *= HdrWeightY(OutColor.x, InExposureScale);
-		#endif
-	#else
-		#if AA_GREEN_AS_LUMA
-			OutColor.rgb *= HdrWeightG(OutColor.rgb, InExposureScale);
-		#else
-			OutColor.rgb *= HdrWeight4(OutColor.rgb, InExposureScale);
-		#endif
-	#endif
-	#if AA_DEBUG
-		OutColor.rg = float2(0.0, 0.0);
-		float DebugDiffPrior = OutColor.b;
-	#endif
-
-
-
-
-
-
-	// FIND LUMA OF CLAMPED HISTORY
-	// ----------------------------
-	// Save off luma of history before the clamp.
-	#if AA_YCOCG
-		float LumaMin = NeighborMin.x;
-		float LumaMax = NeighborMax.x;
-		float LumaHistory = OutColor.x;
-	#else
-		#if AA_GREEN_AS_LUMA
-			float LumaMin = NeighborMin.g;
-			float LumaMax = NeighborMax.g;
-			float LumaHistory = OutColor.g;
-		#else
-			float LumaMin = Luma4(NeighborMin.rgb);
-			float LumaMax = Luma4(NeighborMax.rgb);
-			float LumaHistory = Luma4(OutColor.rgb);
-		#endif
-	#endif
-	float LumaContrast = LumaMax - LumaMin;
-	#if AA_YCOCG
-		OutColor.rgb = clamp(OutColor.rgb, NeighborMin.rgb, NeighborMax.rgb);
-		#if (AA_ALPHA == 0)
-			OutColor.a = clamp(OutColor.a, NeighborMin.a, NeighborMax.a);
-		#endif
-	#else
-		#if AA_AABB
-			// Clamp history, this uses color AABB intersection for tighter fit.
-			// Clamping works with the low pass (if available) to reduce flicker.
-			float ClampBlend = HistoryClamp(OutColor.rgb, FilteredLow.rgb, NeighborMin.rgb, NeighborMax.rgb);
-			#if AA_ALPHA
-				OutColor.rgb = lerp(OutColor.rgb, FilteredLow.rgb, ClampBlend);
-			#else
-				OutColor.rgba = lerp(OutColor.rgba, FilteredLow.rgba, ClampBlend);
-			#endif
-		#else
-			OutColor = clamp(OutColor, NeighborMin, NeighborMax);
-		#endif
-	#endif
-	#if AA_DEBUG
-		OutColor.rg = float2(0.0, 0.0);
-	#endif
-
-	// ADD BACK IN ALIASING TO SHARPEN
-    // -------------------------------
-	#if AA_FILTERED
-		// Blend in non-filtered based on the amount of sub-pixel motion.
-		float AddAliasing = saturate(HistoryBlur) * 0.5;
-		float LumaContrastFactor = 32.0;
-		#if AA_GREEN_AS_LUMA || AA_YCOCG
-			// GREEN_AS_LUMA is 1/4 as bright.
-			LumaContrastFactor *= 4.0;
-		#endif
-		AddAliasing = saturate(AddAliasing + rcp(1.0 + LumaContrast * LumaContrastFactor));
-		Filtered.rgb = lerp(Filtered.rgb, Neighbor4.rgb, AddAliasing);
-	#endif
-	#if AA_YCOCG
-		float LumaFiltered = Filtered.x;
-	#else
-		#if AA_GREEN_AS_LUMA
-			float LumaFiltered = Filtered.g;
-		#else
-			float LumaFiltered = Luma4(Filtered.rgb);
-		#endif
-	#endif
-		
-	// COMPUTE BLEND AMOUNT 
-    // --------------------
-	// Replace history with minimum difference of history from min or max neighborhood.
-	LumaHistory = min(abs(LumaMin-LumaHistory), abs(LumaMax-LumaHistory));
-	float HistoryAmount = (1.0/8.0) + HistoryBlur * (1.0/8.0);
-	float HistoryFactor = LumaHistory * HistoryAmount * (1.0 + HistoryBlur * HistoryAmount * 8.0);
-	float BlendFinal = saturate(HistoryFactor * rcp(LumaHistory + LumaContrast));
-	#if RESPONSIVE
-		// Responsive forces 1/4 of new frame.
-		BlendFinal = 1.0/4.0;
-	#endif
-
-   	#if (AA_ALPHA == 1) && (AA_VELOCITY_WEIGHTING == 1)
-		// Velocity weighting.
-		// Decay prior value.
-		float PriorVelocity = OutColor.a;
-		float VelocityDecay = 0.5;
-		OutColor.a = max(OutColor.a * VelocityDecay, Velocity * rcp(VelocityDecay));
-		float VelocityDiff = abs(PriorVelocity - Velocity) / max(1.0, max(PriorVelocity, Velocity));
-		BlendFinal = max(BlendFinal, VelocityDiff * (1.0/8.0));
-		BlendFinal = min(1.0/2.0, BlendFinal);
-	#endif
-
-	// Offscreen feedback resets.
-	#if AA_LERP 
-		float FixedLerp = 1.0/float(AA_LERP);
-	#endif
-	if(OffScreen) 
-	{
-		OutColor = Filtered;
-		#if AA_ALPHA
-			OutColor.a = 0.0;
-		#endif
-		#if AA_LERP
-			FixedLerp = 1.0;
-		#endif
-	}
-	
-	// LIMIT BLEND AMOUNT BASED ON LOW FREQUENCY SIGNAL
-	// -------------------------------------------------
-	#if AA_TEST
 
+// FIND LUMA OF CLAMPED HISTORY
+// ----------------------------
+// Save off luma of history before the clamp.
+#if AA_YCOCG
+float LumaMin = NeighborMin.x;
+float LumaMax = NeighborMax.x;
+float LumaHistory = OutColor.x;
+#else
+#if AA_GREEN_AS_LUMA
+float LumaMin = NeighborMin.g;
+float LumaMax = NeighborMax.g;
+float LumaHistory = OutColor.g;
+#else
+float LumaMin = Luma4(NeighborMin.rgb);
+float LumaMax = Luma4(NeighborMax.rgb);
+float LumaHistory = Luma4(OutColor.rgb);
+#endif
+#endif
+float LumaContrast = LumaMax - LumaMin;
+#if AA_YCOCG
+OutColor.rgb = clamp(OutColor.rgb, NeighborMin.rgb, NeighborMax.rgb);
+#if (AA_ALPHA == 0)
+OutColor.a = clamp(OutColor.a, NeighborMin.a, NeighborMax.a);
+#endif
+#else
+#if AA_AABB
+// Clamp history, this uses color AABB intersection for tighter fit.
+// Clamping works with the low pass (if available) to reduce flicker.
+float ClampBlend = HistoryClamp(OutColor.rgb, FilteredLow.rgb, NeighborMin.rgb, NeighborMax.rgb);
+#if AA_ALPHA
+OutColor.rgb = lerp(OutColor.rgb, FilteredLow.rgb, ClampBlend);
+#else
+OutColor.rgba = lerp(OutColor.rgba, FilteredLow.rgba, ClampBlend);
+#endif
+#else
+OutColor = clamp(OutColor, NeighborMin, NeighborMax);
+#endif
+#endif
+#if AA_DEBUG
+OutColor.rg = float2(0.0, 0.0);
+#endif
+
+// ADD BACK IN ALIASING TO SHARPEN
+// -------------------------------
+#if AA_FILTERED
+// Blend in non-filtered based on the amount of sub-pixel motion.
+float AddAliasing = saturate(HistoryBlur) * 0.5;
+float LumaContrastFactor = 32.0;
+#if AA_GREEN_AS_LUMA || AA_YCOCG
+// GREEN_AS_LUMA is 1/4 as bright.
+LumaContrastFactor *= 4.0;
+#endif
+AddAliasing = saturate(AddAliasing + rcp(1.0 + LumaContrast * LumaContrastFactor));
+Filtered.rgb = lerp(Filtered.rgb, Neighbor4.rgb, AddAliasing);
+#endif
+#if AA_YCOCG
+float LumaFiltered = Filtered.x;
+#else
+#if AA_GREEN_AS_LUMA
+float LumaFiltered = Filtered.g;
+#else
+float LumaFiltered = Luma4(Filtered.rgb);
+#endif
+#endif
+
+// COMPUTE BLEND AMOUNT 
+// --------------------
+// Replace history with minimum difference of history from min or max neighborhood.
+LumaHistory = min(abs(LumaMin - LumaHistory), abs(LumaMax - LumaHistory));
+float HistoryAmount = (1.0 / 8.0) + HistoryBlur * (1.0 / 8.0);
+float HistoryFactor = LumaHistory * HistoryAmount * (1.0 + HistoryBlur * HistoryAmount * 8.0);
+float BlendFinal = saturate(HistoryFactor * rcp(LumaHistory + LumaContrast));
+#if RESPONSIVE
+// Responsive forces 1/4 of new frame.
+BlendFinal = 1.0 / 4.0;
+#endif
+
+#if (AA_ALPHA == 1) && (AA_VELOCITY_WEIGHTING == 1)
+// Velocity weighting.
+// Decay prior value.
+float PriorVelocity = OutColor.a;
+float VelocityDecay = 0.5;
+OutColor.a = max(OutColor.a * VelocityDecay, Velocity * rcp(VelocityDecay));
+float VelocityDiff = abs(PriorVelocity - Velocity) / max(1.0, max(PriorVelocity, Velocity));
+BlendFinal = max(BlendFinal, VelocityDiff * (1.0 / 8.0));
+BlendFinal = min(1.0 / 2.0, BlendFinal);
+#endif
+
+// Offscreen feedback resets.
+#if AA_LERP 
+float FixedLerp = 1.0 / float(AA_LERP);
+#endif
+if (OffScreen)
+{
+	OutColor = Filtered;
+#if AA_ALPHA
+	OutColor.a = 0.0;
+#endif
+#if AA_LERP
+	FixedLerp = 1.0;
+#endif
+}
+
+// DO FINAL BLEND BETWEEN HISTORY AND FILTERED COLOR
+// -------------------------------------------------
+#if (AA_LERP == 0)
+#if AA_ALPHA
+// Blend in linear to hit luma target.
+OutColor.rgb = lerp(OutColor.rgb, Filtered.rgb, BlendFinal);
+#if RESPONSIVE
+OutColor.a = max(OutColor.a, 1.0 / 2.0);
+#endif
+#else
+OutColor = lerp(OutColor, Filtered, BlendFinal);
+#if AA_FORCE_ALPHA_CLAMP
+OutColor.a = clamp(OutColor.a, NeighborMin.a, NeighborMax.a);
+#endif
+#endif
+#else
+OutColor = lerp(OutColor, Filtered, FixedLerp);
+#endif
+#if AA_YCOCG
 #if 0
-		// LOWPASS
-		float TxScale = 2.0 * max(1.0, min(2.0, Velocity));
-		float TxMaxX = abs(TxAvg.x - TxAvg2.x) * 2.0 * TxScale;
-		float TxMaxY = abs(TxAvg.y - TxAvg2.y) * 4.0 * TxScale;
-		float TxMaxZ = abs(TxAvg.z - TxAvg2.z) * 8.0 * TxScale;
+OutColor.xyz *= HdrWeightInvY(OutColor.x, InExposureScale);
+#endif
+OutColor.rgb = YCoCgToRGB(OutColor.rgb);
+#if 1
+OutColor.rgb *= HdrWeightInvG(OutColor.rgb, InExposureScale);
+#endif
+#else
+// Convert back into linear.
+#if AA_GREEN_AS_LUMA
+OutColor.rgb *= HdrWeightInvG(OutColor.rgb, InExposureScale);
 #else
-		// SUM ABS DIFF
-		float TxScale = 1.0 * max(1.0, min(2.0, Velocity));
-		float TxMaxX = TxDifAvg.x * 2.0 * TxScale;
-		float TxMaxY = TxDifAvg.y * 4.0 * TxScale;
-		float TxMaxZ = TxDifAvg.z * 8.0 * TxScale;
-
-#endif
-		
-
-		float TxMax = max(TxMaxX,max(TxMaxY,TxMaxZ));
-		TxMax = min(1.0/1.0, TxMax);
-		BlendFinal = max(BlendFinal, TxMax);
-
-	#endif
-
-
-
-	// DO FINAL BLEND BETWEEN HISTORY AND FILTERED COLOR
-	// -------------------------------------------------
-	#if (AA_LERP == 0)
-		#if AA_ALPHA
-			// Blend in linear to hit luma target.
-			OutColor.rgb = lerp(OutColor.rgb, Filtered.rgb, BlendFinal);
-			#if RESPONSIVE
-				OutColor.a = max(OutColor.a, 1.0/2.0);
-			#endif
-		#else
-			OutColor = lerp(OutColor, Filtered, BlendFinal);
-			#if AA_FORCE_ALPHA_CLAMP
-				OutColor.a = clamp(OutColor.a, NeighborMin.a, NeighborMax.a);
-			#endif
-		#endif
-	#else
-		OutColor = lerp(OutColor, Filtered, FixedLerp);
-	#endif
-	#if AA_YCOCG
-		#if 0
-			OutColor.xyz *= HdrWeightInvY(OutColor.x, InExposureScale);
-		#endif
-		OutColor.rgb = YCoCgToRGB(OutColor.rgb);
-		#if 1
-			OutColor.rgb *= HdrWeightInvG(OutColor.rgb, InExposureScale);
-		#endif
-	#else
-		// Convert back into linear.
-		#if AA_GREEN_AS_LUMA
-			OutColor.rgb *= HdrWeightInvG(OutColor.rgb, InExposureScale);
-		#else
-			OutColor.rgb *= HdrWeightInv4(OutColor.rgb, InExposureScale);
-		#endif
-	#endif
-	#if AA_NAN
-		// Transform NaNs to black, transform negative colors to black.
-		OutColor.rgb = -min(-OutColor.rgb, 0.0);
-	#endif
-	#if AA_DEBUG
-		OutColor.g = abs(DebugDiffPrior - DebugDiffCurrent);
-		OutColor.r = OutColor.a;
-	#endif
+OutColor.rgb *= HdrWeightInv4(OutColor.rgb, InExposureScale);
+#endif
+#endif
+#if AA_NAN
+// Transform NaNs to black, transform negative colors to black.
+OutColor.rgb = -min(-OutColor.rgb, 0.0);
+#endif
+#if AA_DEBUG
+OutColor.g = abs(DebugDiffPrior - DebugDiffCurrent);
+OutColor.r = OutColor.a;
+#endif
 
 
 
diff --git a/Engine/Source/Programs/UnrealBuildTool/Windows/UEBuildWindows.cs b/Engine/Source/Programs/UnrealBuildTool/Windows/UEBuildWindows.cs
--- a/Engine/Source/Programs/UnrealBuildTool/Windows/UEBuildWindows.cs
+++ b/Engine/Source/Programs/UnrealBuildTool/Windows/UEBuildWindows.cs
@@ -453,6 +453,7 @@ namespace UnrealBuildTool
             if (File.Exists(MorpheusShaderPath))
             {
                 InBuildTarget.GlobalCompileEnvironment.Config.Definitions.Add("HAS_MORPHEUS=1");
+                InBuildTarget.GlobalCompileEnvironment.Config.Definitions.Add("MORPHEUS_120HZ=0");
 
 				//on PS4 the SDK now handles distortion correction.  On PC we will still have to handle it manually,
 				//but we require SDK changes before we can get the required data.  For the moment, no platform does in-engine morpheus distortion
diff --git a/Engine/Source/Runtime/Engine/Private/Slate/SceneViewport.cpp b/Engine/Source/Runtime/Engine/Private/Slate/SceneViewport.cpp
--- a/Engine/Source/Runtime/Engine/Private/Slate/SceneViewport.cpp
+++ b/Engine/Source/Runtime/Engine/Private/Slate/SceneViewport.cpp
@@ -17,8 +17,7 @@ FSceneViewport::FSceneViewport( FViewportClient* InViewportClient, TSharedPtr<SV
 	, CachedMousePos(-1, -1)
 	, PreCaptureMousePos(-1, -1)
 	, SoftwareCursorPosition( 0, 0 )
-	, bIsSoftwareCursorVisible( false )
-	, SlateRenderTargetHandle( NULL )
+	, bIsSoftwareCursorVisible( false )	
 	, DebugCanvasDrawer( new FDebugCanvasDrawer )
 	, ViewportWidget( InViewportWidget )
 	, NumMouseSamplesX( 0 )
@@ -31,9 +30,12 @@ FSceneViewport::FSceneViewport( FViewportClient* InViewportClient, TSharedPtr<SV
 	, bPlayInEditorGetsMouseControl( true )
 	, bPlayInEditorIsSimulate( false )
 	, bCursorHiddenDueToCapture( false )
-	, MousePosBeforeHiddenDueToCapture( -1, -1 )
+	, MousePosBeforeHiddenDueToCapture( -1, -1 )	
+	, NumBufferedFrames(1)
+	, CurrentBufferedTargetIndex(0)
 {
 	bIsSlateViewport = true;
+	RenderThreadSlateTexture = new FSlateRenderTargetRHI(nullptr, 0, 0);
 }
 
 FSceneViewport::~FSceneViewport()
@@ -127,8 +129,11 @@ bool FSceneViewport::KeyState( FKey Key ) const
 void FSceneViewport::Destroy()
 {
 	ViewportClient = NULL;
-
+	
 	UpdateViewportRHI( true, 0, 0, EWindowMode::Windowed );
+
+	delete RenderThreadSlateTexture;
+	RenderThreadSlateTexture = nullptr;
 }
 
 int32 FSceneViewport::GetMouseX() const
@@ -874,7 +879,8 @@ void FSceneViewport::OnViewportClosed()
 
 FSlateShaderResource* FSceneViewport::GetViewportRenderTargetTexture() const
 { 
-	return SlateRenderTargetHandle; 
+	check(IsInGameThread());
+	return BufferedSlateHandles[CurrentBufferedTargetIndex];
 }
 
 void FSceneViewport::ResizeFrame(uint32 NewSizeX, uint32 NewSizeY, EWindowMode::Type NewWindowMode, int32 InPosX, int32 InPosY)
@@ -1081,6 +1087,38 @@ FCanvas* FSceneViewport::GetDebugCanvas()
 	return DebugCanvasDrawer->GetGameThreadDebugCanvas();
 }
 
+const FTexture2DRHIRef& FSceneViewport::GetRenderTargetTexture() const
+{
+	if (IsInRenderingThread())
+	{
+		return RenderTargetTextureRenderThreadRHI;
+	}
+	return 	RenderTargetTextureRHI;
+}
+
+FSlateShaderResource* FSceneViewport::GetViewportRenderTargetTexture()
+{
+	if (IsInRenderingThread())
+	{
+		return RenderThreadSlateTexture;
+	}
+	return BufferedSlateHandles[CurrentBufferedTargetIndex];
+}
+
+void FSceneViewport::SetRenderTargetTextureRenderThread(FTexture2DRHIRef& RT)
+{
+	check(IsInRenderingThread());
+	RenderTargetTextureRenderThreadRHI = RT;
+	if (RT.IsValid())
+	{
+		RenderThreadSlateTexture->SetRHIRef(RenderTargetTextureRenderThreadRHI, RT->GetSizeX(), RT->GetSizeY());
+	}
+	else
+	{
+		RenderThreadSlateTexture->SetRHIRef(nullptr, 0, 0);
+	}
+}
+
 void FSceneViewport::UpdateViewportRHI(bool bDestroyed, uint32 NewSizeX, uint32 NewSizeY, EWindowMode::Type NewWindowMode)
 {
 	// Make sure we're not in the middle of streaming textures.
@@ -1119,10 +1157,13 @@ void FSceneViewport::UpdateViewportRHI(bool bDestroyed, uint32 NewSizeX, uint32
 		else
 		{
 			// Enqueue a render command to delete the handle.  It must be deleted on the render thread after the resource is released
-			ENQUEUE_UNIQUE_RENDER_COMMAND_ONEPARAMETER(DeleteSlateRenderTarget, FSlateRenderTargetRHI*&, SlateRenderTargetHandle, SlateRenderTargetHandle,
+			ENQUEUE_UNIQUE_RENDER_COMMAND_ONEPARAMETER(DeleteSlateRenderTarget, TArray<FSlateRenderTargetRHI*>&, BufferedSlateHandles, BufferedSlateHandles,
 			{
-				delete SlateRenderTargetHandle;
-				SlateRenderTargetHandle = NULL;
+				for (int32 i = 0; i < BufferedSlateHandles.Num(); ++i)
+				{
+					delete BufferedSlateHandles[i];
+				}
+				BufferedSlateHandles.Empty();				
 			});
 
 		}
@@ -1133,6 +1174,9 @@ void FSceneViewport::EnqueueBeginRenderFrame()
 {
 	check( IsInGameThread() );
 
+	CurrentBufferedTargetIndex = (CurrentBufferedTargetIndex + 1) % BufferedSlateHandles.Num();
+	RenderTargetTextureRHI = BufferedRenderTargetsRHI[CurrentBufferedTargetIndex];
+
 	// check if we need to reallocate rendertarget for HMD and update HMD rendering viewport 
 	if (GEngine->StereoRenderingDevice.IsValid() && IsStereoRenderingAllowed())
 	{
@@ -1169,6 +1213,18 @@ void FSceneViewport::EnqueueBeginRenderFrame()
 		}
 	}
 
+
+	//set the rendertarget visible to the render thread
+	//must come before any render thread frame handling.
+	ENQUEUE_UNIQUE_RENDER_COMMAND_TWOPARAMETER(
+		SetRenderThreadViewportTarget,
+		FSceneViewport*, Viewport, this,
+		FTexture2DRHIRef, RT, RenderTargetTextureRHI,		
+		{
+
+		Viewport->SetRenderTargetTextureRenderThread(RT);			
+	});		
+
 	FViewport::EnqueueBeginRenderFrame();
 
 	if (GEngine->StereoRenderingDevice.IsValid())
@@ -1181,17 +1237,18 @@ void FSceneViewport::BeginRenderFrame(FRHICommandListImmediate& RHICmdList)
 {
 	check( IsInRenderingThread() );
 	if (bUseSeparateRenderTarget)
-	{
-		SetRenderTarget(RHICmdList,  RenderTargetTextureRHI,  FTexture2DRHIRef() );
+	{		
+		SetRenderTarget(RHICmdList,  RenderTargetTextureRenderThreadRHI,  FTexture2DRHIRef() );
 	}
 	else if( IsValidRef( ViewportRHI ) ) 
 	{
 		// Get the backbuffer render target to render directly to it
-		RenderTargetTextureRHI = RHICmdList.GetViewportBackBuffer(ViewportRHI);
+		RenderTargetTextureRenderThreadRHI = RHICmdList.GetViewportBackBuffer(ViewportRHI);
+		RenderThreadSlateTexture->SetRHIRef(RenderTargetTextureRenderThreadRHI, RenderTargetTextureRenderThreadRHI->GetSizeX(), RenderTargetTextureRenderThreadRHI->GetSizeY());
 		if (GRHIRequiresEarlyBackBufferRenderTarget)
 		{
 			// unused set render targets are bad on Metal
-			SetRenderTarget(RHICmdList, RenderTargetTextureRHI, FTexture2DRHIRef());
+			SetRenderTarget(RHICmdList, RenderTargetTextureRenderThreadRHI, FTexture2DRHIRef());
 		}
 	}
 }
@@ -1201,9 +1258,9 @@ void FSceneViewport::EndRenderFrame(FRHICommandListImmediate& RHICmdList, bool b
 	check( IsInRenderingThread() );
 	if (bUseSeparateRenderTarget)
 	{
-		if (SlateRenderTargetHandle)
-		{
-			RHICmdList.CopyToResolveTarget( RenderTargetTextureRHI, SlateRenderTargetHandle->GetRHIRef(), false, FResolveParams() );
+		if (BufferedSlateHandles[CurrentBufferedTargetIndex])
+		{			
+			RHICmdList.CopyToResolveTarget(RenderTargetTextureRenderThreadRHI, RenderTargetTextureRenderThreadRHI, false, FResolveParams());
 		}
 	}
 	else
@@ -1211,7 +1268,8 @@ void FSceneViewport::EndRenderFrame(FRHICommandListImmediate& RHICmdList, bool b
 		// Set the active render target(s) to nothing to release references in the case that the viewport is resized by slate before we draw again
 		SetRenderTarget(RHICmdList,  FTexture2DRHIRef(), FTexture2DRHIRef() );
 		// Note: this releases our reference but does not release the resource as it is owned by slate (this is intended)
-		RenderTargetTextureRHI.SafeRelease();
+		RenderTargetTextureRenderThreadRHI.SafeRelease();
+		RenderThreadSlateTexture->SetRHIRef(nullptr, 0, 0);
 	}
 }
 
@@ -1257,6 +1315,38 @@ void FSceneViewport::SwapStatCommands( const FSceneViewport& OtherViewport )
 	}
 }
 
+/** Queue an update to the Window's RT on the Renderthread */
+void FSceneViewport::WindowRenderTargetUpdate(FSlateRenderer* Renderer, SWindow* Window)
+{	
+	check(IsInGameThread());
+	if (Renderer)
+	{
+		if (bUseSeparateRenderTarget)
+		{
+			if (Window)
+			{
+				// We need to pass a texture to the renderer only for stereo rendering. Otherwise, Editor will be rendered incorrectly.
+				if (GEngine->IsStereoscopic3D(this))
+				{
+					//todo: mw Make this function take an FSlateTexture* rather than a void*
+					Renderer->SetWindowRenderTarget(*Window, static_cast<IViewportRenderTargetProvider*>(this));
+				}
+				else
+				{
+					Renderer->SetWindowRenderTarget(*Window, nullptr);
+				}
+			}
+		}
+		else
+		{
+			if (Window)
+			{
+				Renderer->SetWindowRenderTarget(*Window, nullptr);
+			}
+		}
+	}
+}
+
 void FSceneViewport::InitDynamicRHI()
 {
 	if(bRequiresHitProxyStorage)
@@ -1274,43 +1364,83 @@ void FSceneViewport::InitDynamicRHI()
 		{
 			GEngine->StereoRenderingDevice->CalculateRenderTargetSize(TexSizeX, TexSizeY);
 		}
-		FTexture2DRHIRef ShaderResourceTextureRHI;
+				
+		check(BufferedSlateHandles.Num() == BufferedRenderTargetsRHI.Num() && BufferedSlateHandles.Num() == BufferedShaderResourceTexturesRHI.Num());
 
-		FRHIResourceCreateInfo CreateInfo;
-		RHICreateTargetableShaderResource2D( TexSizeX, TexSizeY, PF_B8G8R8A8, 1, TexCreate_None, TexCreate_RenderTargetable, false, CreateInfo, RenderTargetTextureRHI, ShaderResourceTextureRHI );
+		//todo, get from HMD
+		NumBufferedFrames = 3;
 
-		if( !SlateRenderTargetHandle )
-		{
-			SlateRenderTargetHandle = new FSlateRenderTargetRHI( ShaderResourceTextureRHI, TexSizeX, TexSizeY );
-			//UE_LOG(LogSlate, Log, TEXT("SRTH: %p, %d x %d"), ShaderResourceTextureRHI.GetReference(), TexSizeX, TexSizeY);
-		}
-		else
-		{
-			//UE_LOG(LogSlate, Log, TEXT("SRTH: %p, %d x %d, prev %p"), ShaderResourceTextureRHI.GetReference(), TexSizeX, TexSizeY, SlateRenderTargetHandle->GetRHIRef().GetReference());
-			SlateRenderTargetHandle->SetRHIRef( ShaderResourceTextureRHI, TexSizeX, TexSizeY );
-		}
 
-		TSharedPtr<SWindow> Window = FSlateApplication::Get().FindWidgetWindow(ViewportWidget.Pin().ToSharedRef(), WidgetPath);
-		if (Window.IsValid())
+		if (BufferedSlateHandles.Num() != NumBufferedFrames)
 		{
-			// We need to pass a texture to the renderer only for stereo rendering. Otherwise, Editor will be rendered incorrectly.
-			if (GEngine->IsStereoscopic3D(this))
+			//clear existing entries
+			for (int32 i = 0; i < BufferedSlateHandles.Num(); ++i)
 			{
-				Renderer->SetWindowRenderTarget(*Window, RenderTargetTextureRHI);
+				if (!BufferedSlateHandles[i])
+				{
+					BufferedSlateHandles[i] = new FSlateRenderTargetRHI(nullptr, 0, 0);
+				}
+				BufferedRenderTargetsRHI[i] = nullptr;
+				BufferedShaderResourceTexturesRHI[i] = nullptr;
 			}
-			else
+
+			//add sufficient entires for buffering.
+			for (int32 i = BufferedSlateHandles.Num(); i < NumBufferedFrames; i++)
 			{
-				Renderer->SetWindowRenderTarget(*Window, nullptr);
+				BufferedSlateHandles.Add(new FSlateRenderTargetRHI(nullptr, 0, 0));
+				BufferedRenderTargetsRHI.Add(nullptr);
+				BufferedShaderResourceTexturesRHI.Add(nullptr);
 			}
 		}
+
+		FRHIResourceCreateInfo CreateInfo;
+		FTexture2DRHIRef BufferedRTRHI;
+		FTexture2DRHIRef BufferedSRVRHI;
+
+		
+		for (int32 i = 0; i < NumBufferedFrames; ++i)
+		{
+			RHICreateTargetableShaderResource2D(TexSizeX, TexSizeY, PF_B8G8R8A8, 1, TexCreate_None, TexCreate_RenderTargetable, false, CreateInfo, BufferedRTRHI, BufferedSRVRHI);
+			BufferedRenderTargetsRHI[i] = BufferedRTRHI;
+			BufferedShaderResourceTexturesRHI[i] = BufferedSRVRHI;
+
+			BufferedSlateHandles[i]->SetRHIRef(BufferedShaderResourceTexturesRHI[0], TexSizeX, TexSizeY);
+		}
+
+		// clear out any extra entries we have hanging around
+		for (int32 i = NumBufferedFrames; i < BufferedSlateHandles.Num(); ++i)
+		{
+			BufferedSlateHandles[i]->SetRHIRef(nullptr, 0, 0);
+			BufferedRenderTargetsRHI[i] = nullptr;
+			BufferedShaderResourceTexturesRHI[i] = nullptr;
+		}
+
+		CurrentBufferedTargetIndex = 0;
+		RenderTargetTextureRHI = BufferedShaderResourceTexturesRHI[CurrentBufferedTargetIndex];
+
+		
 	}
 	else
 	{
-		TSharedPtr<SWindow> Window = FSlateApplication::Get().FindWidgetWindow(ViewportWidget.Pin().ToSharedRef(), WidgetPath);
-		if (Window.IsValid())
-		{
-			Renderer->SetWindowRenderTarget(*Window, nullptr);
-		}
+		check(BufferedSlateHandles.Num() == 0);
+		BufferedSlateHandles.Add(nullptr);	
+		BufferedRenderTargetsRHI.Add(nullptr);
+		BufferedShaderResourceTexturesRHI.Add(nullptr);
+		NumBufferedFrames = 1;
+
+		RenderTargetTextureRHI = nullptr;		
+	}
+
+	//how is this useful at all?  Pinning a weakptr to get a non-threadsafe shared ptr?  Pinning a weakptr is supposed to be protecting me from my weakptr dying underneath me...
+	TSharedPtr<SWidget> PinnedViewport = ViewportWidget.Pin();
+	if (PinnedViewport.IsValid())
+	{
+
+		FWidgetPath WidgetPath;
+		TSharedPtr<SWindow> Window = FSlateApplication::Get().FindWidgetWindow(PinnedViewport.ToSharedRef(), WidgetPath);
+		TSharedPtr<FSlateRenderer> Renderer = FSlateApplication::Get().GetRenderer();
+		
+		WindowRenderTargetUpdate(Renderer.Get(), Window.Get());
 	}
 }
 
@@ -1320,10 +1450,13 @@ void FSceneViewport::ReleaseDynamicRHI()
 
 	ViewportRHI.SafeRelease();
 
-	if( SlateRenderTargetHandle )
+	for (int32 i = 0; i < BufferedSlateHandles.Num(); ++i)
 	{
-		SlateRenderTargetHandle->ReleaseDynamicRHI();
-	}
+		if (BufferedSlateHandles[i])
+		{
+			BufferedSlateHandles[i]->ReleaseDynamicRHI();
+		}
+	}	
 }
 
 void FSceneViewport::SetPreCaptureMousePosFromSlateCursor()
diff --git a/Engine/Source/Runtime/Engine/Public/SceneViewExtension.h b/Engine/Source/Runtime/Engine/Public/SceneViewExtension.h
--- a/Engine/Source/Runtime/Engine/Public/SceneViewExtension.h
+++ b/Engine/Source/Runtime/Engine/Public/SceneViewExtension.h
@@ -22,7 +22,7 @@ public:
 	/**
 	 * Called on render thread at the start of rendering.
 	 */
-	virtual void PreRenderViewFamily_RenderThread(FSceneViewFamily& InViewFamily) = 0;
+	virtual void PreRenderViewFamily_RenderThread(FRHICommandListImmediate& RHICmdList, FSceneViewFamily& InViewFamily) = 0;
 
 	/**
 	 * Called on render thread at the start of rendering, for each view, after PreRenderViewFamily_RenderThread call.
diff --git a/Engine/Source/Runtime/Engine/Public/Slate/SceneViewport.h b/Engine/Source/Runtime/Engine/Public/Slate/SceneViewport.h
--- a/Engine/Source/Runtime/Engine/Public/Slate/SceneViewport.h
+++ b/Engine/Source/Runtime/Engine/Public/Slate/SceneViewport.h
@@ -12,7 +12,7 @@ DECLARE_DELEGATE_OneParam( FOnSceneViewportResize, FVector2D );
 /**
  * A viewport for use with Slate SViewport widgets.
  */
-class ENGINE_API FSceneViewport : public FViewportFrame, public FViewport, public ISlateViewport
+class ENGINE_API FSceneViewport : public FViewportFrame, public FViewport, public ISlateViewport, public IViewportRenderTargetProvider
 {
 public:
 	FSceneViewport( FViewportClient* InViewportClient, TSharedPtr<SViewport> InViewportWidget );
@@ -36,6 +36,11 @@ public:
 	virtual FVector2D GetSoftwareCursorPosition() const override { return SoftwareCursorPosition; }
 	virtual FCanvas* GetDebugCanvas() override;
 
+	/** Gets the proper RenderTarget based on the current thread*/
+	virtual const FTexture2DRHIRef& GetRenderTargetTexture() const;
+
+	virtual void SetRenderTargetTextureRenderThread(FTexture2DRHIRef& RT);
+
 	/**
 	 * Captures or uncaptures the joystick
 	 *
@@ -216,6 +221,8 @@ public:
 	/** Should return true, if stereo rendering is allowed in this viewport */
 	virtual bool IsStereoRenderingAllowed() const override;
 
+	virtual FSlateShaderResource* GetViewportRenderTargetTexture() override;
+
 private:
 	/**
 	 * Called when this viewport is destroyed
@@ -279,6 +286,9 @@ private:
 	 */
 	void ApplyModifierKeys( const FModifierKeysState& InKeysState );
 
+
+	void WindowRenderTargetUpdate(FSlateRenderer* Renderer, SWindow* Window);
+
 private:
 	/** An intermediate reply state that is reset whenever an input event is generated */
 	FReply CurrentReplyState;
@@ -293,9 +303,7 @@ private:
 	/**	The current position of the software cursor */
 	FVector2D SoftwareCursorPosition;
 	/**	Whether the software cursor should be drawn in the viewport */
-	bool bIsSoftwareCursorVisible;
-	/** The render target used by Slate to draw the viewport.  Can be null if this viewport renders directly to the backbuffer */
-	class FSlateRenderTargetRHI* SlateRenderTargetHandle;
+	bool bIsSoftwareCursorVisible;	
 	/** Draws the debug canvas in Slate */
 	TSharedPtr<class FDebugCanvasDrawer, ESPMode::ThreadSafe> DebugCanvasDrawer;
 	/** The Slate viewport widget where this viewport is drawn */
@@ -324,6 +332,18 @@ private:
 	bool bCursorHiddenDueToCapture;
 	/** Position the cursor was at when we hid it due to capture, so we can put it back afterwards */
 	FIntPoint MousePosBeforeHiddenDueToCapture;
+
+	/** Reprojection on some HMD RHI's requires ViewportTargets to be buffered */
+	/** The render target used by Slate to draw the viewport.  Can be null if this viewport renders directly to the backbuffer */
+	TArray<class FSlateRenderTargetRHI*> BufferedSlateHandles;
+	TArray<FTexture2DRHIRef> BufferedRenderTargetsRHI;
+	TArray<FTexture2DRHIRef> BufferedShaderResourceTexturesRHI;
+
+	FTexture2DRHIRef RenderTargetTextureRenderThreadRHI;
+	class FSlateRenderTargetRHI* RenderThreadSlateTexture;
+
+	int32 NumBufferedFrames;
+	int32 CurrentBufferedTargetIndex;
 };
 
 
diff --git a/Engine/Source/Runtime/HeadMountedDisplay/Public/IHeadMountedDisplay.h b/Engine/Source/Runtime/HeadMountedDisplay/Public/IHeadMountedDisplay.h
--- a/Engine/Source/Runtime/HeadMountedDisplay/Public/IHeadMountedDisplay.h
+++ b/Engine/Source/Runtime/HeadMountedDisplay/Public/IHeadMountedDisplay.h
@@ -271,9 +271,7 @@ public:
 	 */
 	virtual FVector GetPositionOffset() const { return FVector::ZeroVector; }
 
-	virtual void DrawDistortionMesh_RenderThread(struct FRenderingCompositePassContext& Context, const FSceneView& View, const FIntPoint& TextureSize) {}
-
-	virtual void DrawSingleEyeDistortionSDK_RenderThread(FRHICommandListImmediate& RHICmdList, const FSceneView& View, const FTexture2DRHIRef& InputTexture) {}
+	virtual void DrawDistortionMesh_RenderThread(struct FRenderingCompositePassContext& Context, const FSceneView& View, const FIntPoint& TextureSize) {}	
 
 	/**
 	 * This method is able to change screen settings right before any drawing occurs. 
@@ -321,6 +319,8 @@ public:
 	virtual FVector2D GetTextureScaleLeft() const {return FVector2D::ZeroVector;}
 	virtual FVector2D GetTextureScaleRight() const {return FVector2D::ZeroVector;}
 
+	virtual bool NeedsUpscalePostProcessPass()  { return false; }
+
 	/**
 	 * Record analytics
 	 */
diff --git a/Engine/Source/Runtime/Renderer/Private/PostProcess/PostProcessMorpheus.cpp b/Engine/Source/Runtime/Renderer/Private/PostProcess/PostProcessMorpheus.cpp
--- a/Engine/Source/Runtime/Renderer/Private/PostProcess/PostProcessMorpheus.cpp
+++ b/Engine/Source/Runtime/Renderer/Private/PostProcess/PostProcessMorpheus.cpp
@@ -241,24 +241,7 @@ void FRCPassPostProcessMorpheus::Process(FRenderingCompositePassContext& Context
 		SrcSize
 		);
 #elif PLATFORM_PS4
-
-		FRenderingCompositeOutputRef* OutputRef = Context.Pass->GetInput(ePId_Input0);
-		check(OutputRef);
-		
-		FRenderingCompositeOutput* Input = OutputRef->GetOutput();
-		TRefCountPtr<IPooledRenderTarget> InputPooledElement;
-
-		if (Input)
-		{
-			InputPooledElement = Input->RequestInput();
-		}
-
-		check(InputPooledElement);
-		
-		check(!InputPooledElement->IsFree());
-
-		const FTextureRHIRef& SrcTexture = InputPooledElement->GetRenderTargetItem().ShaderResourceTexture;		
-		GEngine->HMDDevice->DrawSingleEyeDistortionSDK_RenderThread(Context.RHICmdList, View, SrcTexture->GetTexture2D());		
+	checkf(false, TEXT("PS4 uses SDK distortion."));
 #else
 	checkf(false, TEXT("Unsupported path.  Morpheus should be disabled."));
 #endif
diff --git a/Engine/Source/Runtime/Renderer/Private/PostProcess/PostProcessing.cpp b/Engine/Source/Runtime/Renderer/Private/PostProcess/PostProcessing.cpp
--- a/Engine/Source/Runtime/Renderer/Private/PostProcess/PostProcessing.cpp
+++ b/Engine/Source/Runtime/Renderer/Private/PostProcess/PostProcessing.cpp
@@ -1518,14 +1518,17 @@ void FPostProcessing::Process(FRHICommandListImmediate& RHICmdList, FViewInfo& V
 			{
 				Node = Context.Graph.RegisterPass(new FRCPassPostProcessHMD());
 			}
-#if HAS_MORPHEUS
 			else if(DeviceType == EHMDDeviceType::DT_Morpheus)
 			{
+				
+#if MORPHEUS_ENGINE_DISTORTION
 				FRCPassPostProcessMorpheus* MorpheusPass = new FRCPassPostProcessMorpheus();
 				MorpheusPass->SetInput(ePId_Input0, FRenderingCompositeOutputRef(Context.FinalOutput));
 				Node = Context.Graph.RegisterPass(MorpheusPass);
-			}
 #endif
+			}
+
+			bHMDWantsUpscale = GEngine->HMDDevice->NeedsUpscalePostProcessPass();
 			
 			if(Node)
 			{
diff --git a/Engine/Source/Runtime/Renderer/Private/SceneRendering.cpp b/Engine/Source/Runtime/Renderer/Private/SceneRendering.cpp
--- a/Engine/Source/Runtime/Renderer/Private/SceneRendering.cpp
+++ b/Engine/Source/Runtime/Renderer/Private/SceneRendering.cpp
@@ -1324,7 +1324,7 @@ static void RenderViewFamily_RenderThread(FRHICommandListImmediate& RHICmdList,
 
 	for( int ViewExt = 0; ViewExt < SceneRenderer->ViewFamily.ViewExtensions.Num(); ViewExt++ )
 	{
-		SceneRenderer->ViewFamily.ViewExtensions[ViewExt]->PreRenderViewFamily_RenderThread(SceneRenderer->ViewFamily);
+		SceneRenderer->ViewFamily.ViewExtensions[ViewExt]->PreRenderViewFamily_RenderThread(RHICmdList, SceneRenderer->ViewFamily);
 		for( int ViewIndex = 0; ViewIndex < SceneRenderer->ViewFamily.Views.Num(); ViewIndex++ )
 		{
 			SceneRenderer->ViewFamily.ViewExtensions[ViewExt]->PreRenderView_RenderThread(SceneRenderer->Views[ViewIndex]);
diff --git a/Engine/Source/Runtime/SlateCore/Public/Rendering/SlateRenderer.h b/Engine/Source/Runtime/SlateCore/Public/Rendering/SlateRenderer.h
--- a/Engine/Source/Runtime/SlateCore/Public/Rendering/SlateRenderer.h
+++ b/Engine/Source/Runtime/SlateCore/Public/Rendering/SlateRenderer.h
@@ -206,7 +206,7 @@ public:
 	/**
 	 * Pushes the rendering of the specified window to the specified render target
 	 */
-	virtual void SetWindowRenderTarget(const SWindow& Window, FTexture2DRHIParamRef RT) {}
+	virtual void SetWindowRenderTarget(const SWindow& Window, class IViewportRenderTargetProvider* Provider) {}
 
 	/**
 	 * Create an updatable texture that can receive new data dynamically
diff --git a/Engine/Source/Runtime/SlateCore/Public/Textures/SlateShaderResource.h b/Engine/Source/Runtime/SlateCore/Public/Textures/SlateShaderResource.h
--- a/Engine/Source/Runtime/SlateCore/Public/Textures/SlateShaderResource.h
+++ b/Engine/Source/Runtime/SlateCore/Public/Textures/SlateShaderResource.h
@@ -138,3 +138,10 @@ protected:
 	// Holds the resource.
 	ResourceType ShaderResource;
 };
+
+class IViewportRenderTargetProvider
+{
+public:
+	virtual FSlateShaderResource* GetViewportRenderTargetTexture() = 0;
+};
+
diff --git a/Engine/Source/Runtime/SlateRHIRenderer/Private/SlateRHIRenderer.cpp b/Engine/Source/Runtime/SlateRHIRenderer/Private/SlateRHIRenderer.cpp
--- a/Engine/Source/Runtime/SlateRHIRenderer/Private/SlateRHIRenderer.cpp
+++ b/Engine/Source/Runtime/SlateRHIRenderer/Private/SlateRHIRenderer.cpp
@@ -72,8 +72,7 @@ void FSlateRHIRenderer::FViewportInfo::InitRHI()
 void FSlateRHIRenderer::FViewportInfo::ReleaseRHI()
 {
 	DepthStencil.SafeRelease();
-	ViewportRHI.SafeRelease();
-	RenderTargetTexture.SafeRelease();
+	ViewportRHI.SafeRelease();	
 }
 
 void FSlateRHIRenderer::FViewportInfo::ConditionallyUpdateDepthBuffer(bool bInRequiresStencilTest)
@@ -406,11 +405,12 @@ void FSlateRHIRenderer::DrawWindow_RenderThread(FRHICommandListImmediate& RHICmd
 		// should have been created by the game thread
 		check( IsValidRef(ViewportInfo.ViewportRHI) );
 
-		FTexture2DRHIRef BackBuffer = (ViewportInfo.RenderTargetTexture) ? 
-			ViewportInfo.RenderTargetTexture : RHICmdList.GetViewportBackBuffer(ViewportInfo.ViewportRHI);
+		FTexture2DRHIRef ViewportRT = ViewportInfo.GetRenderTargetTexture();
+		FTexture2DRHIRef BackBuffer = (ViewportRT) ?
+		ViewportRT : RHICmdList.GetViewportBackBuffer(ViewportInfo.ViewportRHI);
 		
-		const uint32 ViewportWidth  = (ViewportInfo.RenderTargetTexture) ? ViewportInfo.RenderTargetTexture->GetSizeX() : ViewportInfo.Width;
-		const uint32 ViewportHeight = (ViewportInfo.RenderTargetTexture) ? ViewportInfo.RenderTargetTexture->GetSizeY() : ViewportInfo.Height;
+		const uint32 ViewportWidth = (ViewportRT) ? ViewportRT->GetSizeX() : ViewportInfo.Width;
+		const uint32 ViewportHeight = (ViewportRT) ? ViewportRT->GetSizeY() : ViewportInfo.Height;
 
 		RHICmdList.BeginDrawingViewport( ViewportInfo.ViewportRHI, FTextureRHIRef() );
 		RHICmdList.SetViewport( 0,0,0,ViewportWidth, ViewportHeight, 0.0f ); 
@@ -452,9 +452,9 @@ void FSlateRHIRenderer::DrawWindow_RenderThread(FRHICommandListImmediate& RHICmd
 	}
 
 	bool bNeedCallFinishFrameForStereo = false;
-	if (GEngine && IsValidRef(ViewportInfo.RenderTargetTexture) && GEngine->StereoRenderingDevice.IsValid())
+	if (GEngine && IsValidRef(ViewportInfo.GetRenderTargetTexture()) && GEngine->StereoRenderingDevice.IsValid())
 	{
-		GEngine->StereoRenderingDevice->RenderTexture_RenderThread(RHICmdList, RHICmdList.GetViewportBackBuffer(ViewportInfo.ViewportRHI), ViewportInfo.RenderTargetTexture);
+		GEngine->StereoRenderingDevice->RenderTexture_RenderThread(RHICmdList, RHICmdList.GetViewportBackBuffer(ViewportInfo.ViewportRHI), ViewportInfo.GetRenderTargetTexture());
 		bNeedCallFinishFrameForStereo = true;
 	}
 
@@ -1216,11 +1216,11 @@ void FSlateRHIRenderer::RequestResize( const TSharedPtr<SWindow>& Window, uint32
 	}
 }
 
-void FSlateRHIRenderer::SetWindowRenderTarget(const SWindow& Window, FTexture2DRHIParamRef RT)
-{
+void FSlateRHIRenderer::SetWindowRenderTarget(const SWindow& Window, IViewportRenderTargetProvider* Provider)
+{	
 	FViewportInfo* ViewInfo = WindowToViewportInfo.FindRef(&Window);
 	if (ViewInfo)
 	{
-		ViewInfo->RenderTargetTexture = RT;
+		ViewInfo->RTProvider = Provider;
 	}
 }
diff --git a/Engine/Source/Runtime/SlateRHIRenderer/Private/SlateRHIRenderer.h b/Engine/Source/Runtime/SlateRHIRenderer/Private/SlateRHIRenderer.h
--- a/Engine/Source/Runtime/SlateRHIRenderer/Private/SlateRHIRenderer.h
+++ b/Engine/Source/Runtime/SlateRHIRenderer/Private/SlateRHIRenderer.h
@@ -94,8 +94,8 @@ private:
 		FViewportRHIRef ViewportRHI;
 		/** The depth buffer texture if any */
 		FTexture2DRHIRef DepthStencil;
-		/** The render target texture (if rendering into a texture */
-		FTexture2DRHIRef RenderTargetTexture;
+		
+		//FTexture2DRHIRef RenderTargetTexture;
 		/** The OS Window handle (for recreating the viewport) */
 		void* OSWindow;
 		/** The actual width of the viewport */
@@ -112,6 +112,7 @@ private:
 		bool bFullscreen;
 		/** The desired pixel format for this viewport */
 		EPixelFormat PixelFormat;
+		IViewportRenderTargetProvider* RTProvider;
 	
 		/** FRenderResource interface */
 		virtual void InitRHI() override;
@@ -125,7 +126,8 @@ private:
 				DesiredHeight(0),
 				bRequiresStencilTest(false),
 				bFullscreen(false),
-				PixelFormat(EPixelFormat::PF_Unknown)
+				PixelFormat(EPixelFormat::PF_Unknown),
+				RTProvider(nullptr)
 		{
 
 		}
@@ -137,6 +139,18 @@ private:
 
 		void ConditionallyUpdateDepthBuffer(bool bInRequiresStencilTest);
 		void RecreateDepthBuffer_RenderThread();
+
+		FTexture2DRHIRef GetRenderTargetTexture() const
+		{
+			if (RTProvider)
+			{
+				FSlateRenderTargetRHI* SlateTarget = (FSlateRenderTargetRHI*)(RTProvider->GetViewportRenderTargetTexture());
+				return SlateTarget->GetTypedResource();
+			}
+			return nullptr;
+		}
+
+	private:		
 	};
 public:
 	FSlateRHIRenderer( TSharedPtr<FSlateRHIResourceManager> InResourceManager, TSharedPtr<FSlateFontCache> InFontCache, TSharedPtr<FSlateFontMeasure> InFontMeasure );
@@ -225,7 +239,7 @@ public:
 
 	virtual void PrepareToTakeScreenshot(const FIntRect& Rect, TArray<FColor>* OutColorData) override;
 
-	virtual void SetWindowRenderTarget(const SWindow& Window, FTexture2DRHIParamRef RT) override;
+	virtual void SetWindowRenderTarget(const SWindow& Window, class IViewportRenderTargetProvider* Provider) override;
 
 private:
 	/** Loads all known textures from Slate styles */
