Fixes a shader compilation error caused by the DecalMipmapLevel node's code with d-buffer decals and no textures

diff --git a/Engine/Shaders/MaterialTemplate.usf b/Engine/Shaders/MaterialTemplate.usf
--- a/Engine/Shaders/MaterialTemplate.usf
+++ b/Engine/Shaders/MaterialTemplate.usf
@@ -261,7 +261,7 @@ struct FMaterialVertexParameters
 };
 
 
-#ifdef DECAL_RENDERSTAGE
+#if DECAL_RENDERSTAGE && NUM_MATERIAL_TEXCOORDS
 	/*
 	 * Material node DecalMipmapLevel's code designed to avoid the 2x2 pixels artefacts on the edges around where the decal
 	 * is projected to. The technique is fetched from (http://www.humus.name/index.php?page=3D&ID=84).
