fixed typo

diff --git a/Engine/Shaders/BasePassPixelShader.usf b/Engine/Shaders/BasePassPixelShader.usf
--- a/Engine/Shaders/BasePassPixelShader.usf
+++ b/Engine/Shaders/BasePassPixelShader.usf
@@ -692,7 +692,7 @@ void Main(
 
 		FDBufferData DBufferData = GetDBufferData(ScreenUV);
 
-		// the material can disable the DBuffer effects for beter performance or control
+		// the material can disable the DBuffer effects for better performance or control
 		if((MATERIALDECALRESPONSEMASK & 0x1) == 0) { DBufferData.PreMulColor = 0; DBufferData.ColorOpacity = 1; }
 		if((MATERIALDECALRESPONSEMASK & 0x2) == 0) { DBufferData.PreMulWorldNormal = 0; DBufferData.NormalOpacity = 1; }
 		if((MATERIALDECALRESPONSEMASK & 0x4) == 0) { DBufferData.PreMulRoughness = 0; DBufferData.RoughnessOpacity = 1; }
