added missing file for recent checkin

diff --git a/Engine/Shaders/PostProcessMotionBlur.usf b/Engine/Shaders/PostProcessMotionBlur.usf
--- a/Engine/Shaders/PostProcessMotionBlur.usf
+++ b/Engine/Shaders/PostProcessMotionBlur.usf
@@ -90,9 +90,9 @@ float4 MotionBlurParameters;
 // BoneMatrices and PreviousBoneMatrices for visualization
 
 // The bone matrix buffer stored as 4x3 (3 float4 texels behind each other), all chunks of a skeletal mesh in one
-Buffer<float4> BoneMatrices;
-// buffer with all old bone matrices stored as 4x3 (3 float4 texels behind each other), all chunks of a skeletal mesh in one
-Buffer<float4> PreviousBoneMatrices;
+Buffer<float4> BoneMatrices0;
+// The bone matrix buffer stored as 4x3 (3 float4 texels behind each other), all chunks of a skeletal mesh in one
+Buffer<float4> BoneMatrices1;
 
 // ------------------------------------------
 
@@ -635,7 +635,7 @@ void MainPS(float4 UVAndScreenPos : TEXCOORD0, out float4 OutColor : SV_Target0)
 		{
 			uint Index = PixelPosInt.x + LineInBuffer * ElementsPerLine;
 
-			float4 Value = WhichBuffer ? BoneMatrices[Index] : PreviousBoneMatrices[Index];
+			float4 Value = WhichBuffer ? BoneMatrices1[Index] : BoneMatrices0[Index];
 
 			OutColor.rgb = lerp(OutColor.rgb, Value.rgb, Alpha);
 		}
