UE4 - Fix editor grid for ES3.1 emulation

diff --git a/Engine/Shaders/BasePassForForwardShadingPixelShader.usf b/Engine/Shaders/BasePassForForwardShadingPixelShader.usf
--- a/Engine/Shaders/BasePassForForwardShadingPixelShader.usf
+++ b/Engine/Shaders/BasePassForForwardShadingPixelShader.usf
@@ -64,7 +64,7 @@ SamplerState	FilteredSceneDepthTextureSampler;
 #endif
 
 
-#if USE_EDITOR_COMPOSITING && (ES2_EMULATION)
+#if USE_EDITOR_COMPOSITING && (MOBILE_EMULATION)
 // @return 0:translucent..1:opaque
 float ClipForEditorPrimitives(FMaterialPixelParameters MaterialParameters)
 {
@@ -269,7 +269,7 @@ void Main(
 #endif
 #endif
 
-#if !MATERIAL_SHADINGMODEL_UNLIT && ES2_EMULATION
+#if !MATERIAL_SHADINGMODEL_UNLIT && MOBILE_EMULATION
 	Color = lerp(Color, DiffuseColor + SpecularColor, View.UnlitViewmodeMask);
 #endif
 
@@ -313,7 +313,7 @@ void Main(
 	#endif
 
 	// Editor primitive depth testing
-	#if USE_EDITOR_COMPOSITING && (ES2_EMULATION)
+	#if USE_EDITOR_COMPOSITING && (MOBILE_EMULATION)
 		#if MATERIALBLENDING_MASKED
 			// some material might have a opacity value
 			OutColor.a = GetMaterialMaskInputRaw(MaterialParameters);
diff --git a/Engine/Shaders/BasePassPixelShader.usf b/Engine/Shaders/BasePassPixelShader.usf
--- a/Engine/Shaders/BasePassPixelShader.usf
+++ b/Engine/Shaders/BasePassPixelShader.usf
@@ -529,7 +529,7 @@ float ClipForEditorPrimitives(FMaterialPixelParameters MaterialParameters)
 {
 	float Ret = 1;
 
-#if USE_EDITOR_COMPOSITING && (FEATURE_LEVEL >= FEATURE_LEVEL_SM4 || ES2_EMULATION)
+#if USE_EDITOR_COMPOSITING && (FEATURE_LEVEL >= FEATURE_LEVEL_SM4 || MOBILE_EMULATION)
 	// Depth test manually if compositing editor primitives since the depth buffer is different (MSAA only)
 	BRANCH
 	if (bEnableEditorPrimitiveDepthTest)
@@ -550,7 +550,7 @@ float ClipForEditorPrimitives(FMaterialPixelParameters MaterialParameters)
 		// 0.5f is to bias around the actual value, 1 or 0 are another option
 		Ret = saturate(0.5f - (DeviceZ - PixelDeviceZ) / DeviceDepthFade);
 	}
-#endif	// USE_EDITOR_COMPOSITING && (FEATURE_LEVEL >= FEATURE_LEVEL_SM4 || ES2_EMULATION)
+#endif	// USE_EDITOR_COMPOSITING && (FEATURE_LEVEL >= FEATURE_LEVEL_SM4 || MOBILE_EMULATION)
 
 	// Note: multiple returns cause strange HLSL compiler error for CV_Coverage in later code
 	return Ret;
@@ -638,7 +638,7 @@ void Main(
 #endif
 	);
 
-#if USE_EDITOR_COMPOSITING && (FEATURE_LEVEL >= FEATURE_LEVEL_SM4 || ES2_EMULATION)
+#if USE_EDITOR_COMPOSITING && (FEATURE_LEVEL >= FEATURE_LEVEL_SM4 || MOBILE_EMULATION)
 	const bool bEditorWeightedZBuffering = true;
 #else
 	const bool bEditorWeightedZBuffering = false;
diff --git a/Engine/Shaders/Common.usf b/Engine/Shaders/Common.usf
--- a/Engine/Shaders/Common.usf
+++ b/Engine/Shaders/Common.usf
@@ -160,8 +160,8 @@
 #define PLATFORM_SUPPORTS_DEVELOPMENT_SHADERS (!PS4_PROFILE && !XBOXONE_PROFILE && !ES31_AEP_PROFILE)
 #define USE_DEVELOPMENT_SHADERS (COMPILE_SHADERS_FOR_DEVELOPMENT && PLATFORM_SUPPORTS_DEVELOPMENT_SHADERS)
 
-#ifndef ES2_EMULATION
-#define ES2_EMULATION ((FEATURE_LEVEL == FEATURE_LEVEL_ES2) && (!COMPILER_GLSL_ES2 && USE_DEVELOPMENT_SHADERS && !METAL_PROFILE))
+#ifndef MOBILE_EMULATION
+#define MOBILE_EMULATION ((FEATURE_LEVEL == FEATURE_LEVEL_ES2 || FEATURE_LEVEL == FEATURE_LEVEL_ES3_1) && (!COMPILER_GLSL_ES2 && USE_DEVELOPMENT_SHADERS && !METAL_PROFILE))
 #endif
 
 // If we didn't request color component swizzling, just make it empty
@@ -738,7 +738,7 @@ float CalcSceneDepth( float2 ScreenUV )
 #endif
 }
 
-#if FEATURE_LEVEL >= FEATURE_LEVEL_SM4 || ES2_EMULATION
+#if FEATURE_LEVEL >= FEATURE_LEVEL_SM4 || MOBILE_EMULATION
 
 	// depth in in DeviceZ
 	Texture2D<float> SceneDepthTextureNonMS;
diff --git a/Engine/Shaders/MaterialTemplate.usf b/Engine/Shaders/MaterialTemplate.usf
--- a/Engine/Shaders/MaterialTemplate.usf
+++ b/Engine/Shaders/MaterialTemplate.usf
@@ -610,15 +610,15 @@ MaterialFloat2 GetParticleMacroUV(FMaterialPixelParameters Parameters)
 	return (Parameters.ScreenPosition.xy / Parameters.ScreenPosition.w - Parameters.Particle.MacroUV.xy) * Parameters.Particle.MacroUV.zw + MaterialFloat2(.5, .5);
 }
 
-#ifndef ES2_EMULATION
-#define ES2_EMULATION ((FEATURE_LEVEL == FEATURE_LEVEL_ES2) && (!COMPILER_GLSL_ES2 && USE_DEVELOPMENT_SHADERS && !METAL_PROFILE))
+#ifndef MOBILE_EMULATION
+#define MOBILE_EMULATION ((FEATURE_LEVEL == FEATURE_LEVEL_ES2 || FEATURE_LEVEL == FEATURE_LEVEL_ES3_1) && (!COMPILER_GLSL_ES2 && USE_DEVELOPMENT_SHADERS && !METAL_PROFILE))
 #endif
 
 MaterialFloat4 ProcessMaterialColorTextureLookup(MaterialFloat4 TextureValue)
 {
 #if (ES2_PROFILE || ES3_1_PROFILE)
-	#if ES2_EMULATION
-	if( View.ES2PreviewMode > 0.5f)
+	#if MOBILE_EMULATION
+	if( View.MobilePreviewMode > 0.5f)
 	{
 		// undo HW srgb->lin
 		TextureValue = pow(TextureValue, 1.0f/2.2f); // TODO: replace with a more accurate lin -> sRGB conversion.
@@ -639,8 +639,8 @@ MaterialFloat4 ProcessMaterialLinearColorTextureLookup(MaterialFloat4 TextureVal
 MaterialFloat ProcessMaterialGreyscaleTextureLookup(MaterialFloat TextureValue)
 {
 #if (ES2_PROFILE || ES3_1_PROFILE)
-	#if ES2_EMULATION
-	if( View.ES2PreviewMode > 0.5f )
+	#if MOBILE_EMULATION
+	if( View.MobilePreviewMode > 0.5f )
 	{
 		// undo HW srgb->lin
 		TextureValue = pow(TextureValue, 1.0f/2.2f); // TODO: replace with a more accurate lin -> sRGB conversion.
diff --git a/Engine/Shaders/SimpleElementPixelShader.usf b/Engine/Shaders/SimpleElementPixelShader.usf
--- a/Engine/Shaders/SimpleElementPixelShader.usf
+++ b/Engine/Shaders/SimpleElementPixelShader.usf
@@ -10,7 +10,7 @@
 #define USE_IES_PROFILE 1
 #include "IESLightProfilesCommon.usf"
 
-#define ENABLE_EDITOR_COMPOSITING (FEATURE_LEVEL >= FEATURE_LEVEL_SM4 || ES2_EMULATION)
+#define ENABLE_EDITOR_COMPOSITING (FEATURE_LEVEL >= FEATURE_LEVEL_SM4 || MOBILE_EMULATION)
 
 Texture2D InTexture;
 SamplerState InTextureSampler;
diff --git a/Engine/Source/Runtime/Engine/Public/SceneView.h b/Engine/Source/Runtime/Engine/Public/SceneView.h
--- a/Engine/Source/Runtime/Engine/Public/SceneView.h
+++ b/Engine/Source/Runtime/Engine/Public/SceneView.h
@@ -361,7 +361,7 @@ BEGIN_UNIFORM_BUFFER_STRUCT_WITH_CONSTRUCTOR(FViewUniformShaderParameters,ENGINE
 	DECLARE_UNIFORM_BUFFER_STRUCT_MEMBER(FVector4,SceneTextureMinMax)
 	DECLARE_UNIFORM_BUFFER_STRUCT_MEMBER(FLinearColor,SkyLightColor)
 	DECLARE_UNIFORM_BUFFER_STRUCT_MEMBER_ARRAY(FVector4,SkyIrradianceEnvironmentMap,[7])
-	DECLARE_UNIFORM_BUFFER_STRUCT_MEMBER(float, ES2PreviewMode)
+	DECLARE_UNIFORM_BUFFER_STRUCT_MEMBER(float, MobilePreviewMode)
 	DECLARE_UNIFORM_BUFFER_STRUCT_MEMBER_TEXTURE(Texture2D, DirectionalLightShadowTexture)	
 	DECLARE_UNIFORM_BUFFER_STRUCT_MEMBER_SAMPLER(SamplerState, DirectionalLightShadowSampler)
 END_UNIFORM_BUFFER_STRUCT(FViewUniformShaderParameters)
diff --git a/Engine/Source/Runtime/Renderer/Private/SceneRendering.cpp b/Engine/Source/Runtime/Renderer/Private/SceneRendering.cpp
--- a/Engine/Source/Runtime/Renderer/Private/SceneRendering.cpp
+++ b/Engine/Source/Runtime/Renderer/Private/SceneRendering.cpp
@@ -691,10 +691,10 @@ TUniformBufferRef<FViewUniformShaderParameters> FViewInfo::CreateUniformBuffer(
 	checkSlow(sizeof(ViewUniformShaderParameters.SkyIrradianceEnvironmentMap) == sizeof(FVector4) * 7);
 	SetupSkyIrradianceEnvironmentMapConstants((FVector4*)&ViewUniformShaderParameters.SkyIrradianceEnvironmentMap);
 
-	ViewUniformShaderParameters.ES2PreviewMode =
+	ViewUniformShaderParameters.MobilePreviewMode =
 		(GIsEditor &&
-		FeatureLevel == ERHIFeatureLevel::ES2 &&
-		GMaxRHIFeatureLevel > ERHIFeatureLevel::ES2) ? 1.0f : 0.0f;
+		(FeatureLevel == ERHIFeatureLevel::ES2 || FeatureLevel == ERHIFeatureLevel::ES3_1) &&
+		GMaxRHIFeatureLevel > ERHIFeatureLevel::ES3_1) ? 1.0f : 0.0f;
 
 	return TUniformBufferRef<FViewUniformShaderParameters>::CreateUniformBufferImmediate(ViewUniformShaderParameters, UniformBuffer_SingleFrame);
 }
