UE4 - Fix for motion blur on cloth (UE-14233)

diff --git a/Engine/Shaders/GpuSkinVertexFactory.usf b/Engine/Shaders/GpuSkinVertexFactory.usf
--- a/Engine/Shaders/GpuSkinVertexFactory.usf
+++ b/Engine/Shaders/GpuSkinVertexFactory.usf
@@ -393,7 +393,7 @@ float3x3 SkinTangents( FVertexFactoryInput Input, FVertexFactoryIntermediates In
 	
 	#if GPUSKIN_APEX_CLOTH
 
-	float3x3 ClothTangent;
+	float3x3 ClothTangent = float3x3(1,0,0, 0,1,0, 0,0,1);
 	float TempClothBlendWeight = 0.0f;
 	if( IsSimulatedVertex(Input) )
 	{ 
@@ -504,7 +504,7 @@ float4 CalcWorldPosition(FVertexFactoryInput Input, FVertexFactoryIntermediates
 		float4 TransformedSkinPos = mul(SkinnedPosition, Primitive.LocalToWorld);
 
 		// Intermediates.SimulatedPosition is a clothing position in world coord
-		float3 BlendedPos = Intermediates.SimulatedPosition * ClothBlendWeight + TransformedSkinPos.xyz * (1.0f - ClothBlendWeight);
+		float3 BlendedPos = lerp(TransformedSkinPos.xyz, Intermediates.SimulatedPosition, ClothBlendWeight);
 		return float4(BlendedPos + View.PreViewTranslation, 1);
 	}
 	else
@@ -582,16 +582,16 @@ float4 VertexFactoryGetPreviousWorldPosition(FVertexFactoryInput Input, FVertexF
 							1),
 					PreviousLocalToWorld);
 #else
-	float4 PrevSkinPosInWolrd = mul(float4(SkinPreviousPosition(Input, Intermediates),1), PreviousLocalToWorld);
+	float4 PrevSkinPosInWorld = mul(float4(SkinPreviousPosition(Input, Intermediates),1), PreviousLocalToWorld);
 #if GPUSKIN_APEX_CLOTH
 	if( IsSimulatedVertex(Input) )
 	{
-		float3 BlendedPrevPos = (Intermediates.SimulatedPosition + View.PrevPreViewTranslation)*ClothBlendWeight + PrevSkinPosInWolrd.xyz*(1.0f - ClothBlendWeight);		 
-		return float4(BlendedPrevPos, 1);
+		float3 BlendedPrevPos = lerp(PrevSkinPosInWorld.xyz, Intermediates.SimulatedPosition, ClothBlendWeight);		 
+		return float4(BlendedPrevPos + View.PrevPreViewTranslation, 1);
 	}
 #endif // GPUSKIN_APEX_CLOTH
 
-	return PrevSkinPosInWolrd;
+	return PrevSkinPosInWorld;
 #endif
 	}
 
