Remove statics for ps4.

diff --git a/Engine/Shaders/ACES.usf b/Engine/Shaders/ACES.usf
--- a/Engine/Shaders/ACES.usf
+++ b/Engine/Shaders/ACES.usf
@@ -641,8 +641,8 @@ float3 ODT( float3 oces )
 	rgbPost[2] = segmented_spline_c9_fwd( rgbPre[2] );
 
 	// Target white and black points for cinema system tonescale
-	static const float CINEMA_WHITE = 48.0;
-	static const float CINEMA_BLACK = 0.02; // CINEMA_WHITE / 2400.
+	const float CINEMA_WHITE = 48.0;
+	const float CINEMA_BLACK = 0.02; // CINEMA_WHITE / 2400.
 
 	// Scale luminance to linear code value
 	float3 linearCV;
