UE4 - Fix shader issues for GL

diff --git a/Engine/Shaders/ACES.usf b/Engine/Shaders/ACES.usf
--- a/Engine/Shaders/ACES.usf
+++ b/Engine/Shaders/ACES.usf
@@ -264,9 +264,9 @@ float segmented_spline_c5_fwd( float x )
 		{ -4.0000000000, -4.0000000000, -3.1573765773, -0.4852499958, 1.8477324706, 1.8477324706 },
 		// coefsHigh[6]
 		{ -0.7185482425, 2.0810307172, 3.6681241237, 4.0000000000, 4.0000000000, 4.0000000000 },
-		{ 0.18*exp2(-15), 0.0001},    // minPoint
+		{ 0.18*exp2(-15.0), 0.0001},    // minPoint
 		{ 0.18,              4.8},    // midPoint
-		{ 0.18*exp2(18),  10000.},    // maxPoint
+		{ 0.18*exp2(18.0),  10000.},    // maxPoint
 		0.0,  // slopeLow
 		0.0   // slopeHigh
 	};
@@ -276,7 +276,7 @@ float segmented_spline_c5_fwd( float x )
 
 	// Check for negatives or zero before taking the log. If negative or zero,
 	// set to ACESMIN.1
-	float xCheck = x <= 0 ? exp2(-14) : x;
+	float xCheck = x <= 0 ? exp2(-14.0) : x;
 
 	float logx = log10( xCheck);
 	float logy;
diff --git a/Engine/Shaders/TonemapCommon.usf b/Engine/Shaders/TonemapCommon.usf
--- a/Engine/Shaders/TonemapCommon.usf
+++ b/Engine/Shaders/TonemapCommon.usf
@@ -4,6 +4,11 @@
 	TonemapCommon.usf: PostProcessing tone mapping common
 =============================================================================*/
 
+// #todo-rco: Fix in the cross-compiler
+#if COMPILER_GLSL || COMPILER_GLSL_ES2
+#define log10(x)  log(x) / log(10)
+#endif
+
 #include "ACES.usf"
 
 // usually 1/2.2, the .y is used for inverse gamma when "gamma only" mode is not used
