UE4 - Fix again

diff --git a/Engine/Shaders/TonemapCommon.usf b/Engine/Shaders/TonemapCommon.usf
--- a/Engine/Shaders/TonemapCommon.usf
+++ b/Engine/Shaders/TonemapCommon.usf
@@ -6,7 +6,7 @@
 
 // #todo-rco: Fix in the cross-compiler
 #if COMPILER_GLSL || COMPILER_GLSL_ES2
-#define log10(x)  log(x) / log(10)
+#define log10(x)  log(x) / log(10.0)
 #endif
 
 #include "ACES.usf"
