Fixed formatting.

diff --git a/Engine/Shaders/PostProcessTemporalCommon.usf b/Engine/Shaders/PostProcessTemporalCommon.usf
--- a/Engine/Shaders/PostProcessTemporalCommon.usf
+++ b/Engine/Shaders/PostProcessTemporalCommon.usf
@@ -1,7 +1,7 @@
 // Copyright 1998-2015 Epic Games, Inc. All Rights Reserved.
 
 /*=============================================================================
-PostProcessTemporalCommon.usf: Common Temporal AA Functionality
+	   PostProcessTemporalCommon.usf: Common Temporal AA Functionality
 -------------------------------------------------------------------------------
 This is the common bits of the temporal AA shader
 which can be configured via defines for various temporal AA passes.
@@ -10,27 +10,27 @@ which can be configured via defines for various temporal AA passes.
 // 1 = Use tighter AABB clamp for history.
 // 0 = Use simple min/max clamp.
 #ifndef AA_AABB
-#define AA_AABB 1
+	#define AA_AABB 1
 #endif
 
 // 0 = Anti-alias the alpha channel also (not getting used currently). 
 // 1 = Use alpha channel to improve quality (required for primary AA).
 //     Leverages dead code removal to work in RGB instead of RGBA.
 #ifndef AA_ALPHA
-#define AA_ALPHA 1
+	#define AA_ALPHA 1
 #endif
 
 // Cross distance in pixels used in depth search X pattern.
 // 0 = Turn this feature off.
 // 2 = Is required for standard temporal AA pass.
 #ifndef AA_CROSS
-#define AA_CROSS 2
+	#define AA_CROSS 2
 #endif
-
+	
 // 1 = Render in blue, with green = diff between frames, red = alpha channel.
 // 0 = Non-debug.
 #ifndef AA_DEBUG
-#define AA_DEBUG 0
+	#define AA_DEBUG 0
 #endif
 
 // 2 = Dilate in cross pattern by 2 pixels in distance (this can be larger than 2 if required).
@@ -39,543 +39,543 @@ which can be configured via defines for various temporal AA passes.
 //     This is only valid for AA_ALPHA == 1
 // 0 = Turn off.
 #ifndef AA_DILATE
-#define AA_DILATE AA_ALPHA
+	#define AA_DILATE AA_ALPHA
 #endif
 
 // 1 = Use dynamic motion.
 // 0 = Skip dynamic motion, currently required for half resolution passes.
 #ifndef AA_DYNAMIC
-#define AA_DYNAMIC 1
+	#define AA_DYNAMIC 1
 #endif
 
 // 1 = Use filtered sample.
 // 0 = Use center sample.
 #ifndef AA_FILTERED
-#define AA_FILTERED 1
+	#define AA_FILTERED 1
 #endif
 
 // 1 = Improve quality on converged edges.
 // 0 = Default.
 #ifndef AA_GRAIN
-#define AA_GRAIN 0
+	#define AA_GRAIN 0
 #endif
 
 // 1 = Use higher quality HDR weighted filtering.
 // 0 = Don't use.
 #ifndef AA_HDR
-#define AA_HDR 1
+	#define AA_HDR 1
 #endif
 
 // 1 = Use manual HDR weighted bilinear filtering for history.
 #if AA_HDR_HISTORY
-#define AA_HDR_HISTORY 1
+	#define AA_HDR_HISTORY 1
 #endif
 
 // 0 = Dynamic motion based lerp value (default).
 // non-zero = Use 1/LERP fixed lerp value (used for reflections).
 #ifndef AA_LERP
-#define AA_LERP 0
+	#define AA_LERP 0
 #endif
 
 // 1 = Use extra lowpass filter for quality bump.
 // 0 = Don't use.
 #ifndef AA_LOWPASS
-#define AA_LOWPASS 1
+	#define AA_LOWPASS 1
 #endif
 
 // 1 = Use higher quality round clamp.
 // 0 = Use lower quality but faster box clamp.
 #ifndef AA_ROUND
-#define AA_ROUND 1
+	#define AA_ROUND 1
 #endif
 
 // 1 = Use extra clamp to avoid NANs
 // 0 = Don't use.
 #ifndef AA_NAN
-#define AA_NAN 1
+	#define AA_NAN 1
 #endif
 
 // Fix for lack of borders during current frame filter.
 #ifndef AA_BORDER
-#define AA_BORDER 0
+	#define AA_BORDER 0
 #endif
 
 // Force clamp on alpha.
 #ifndef AA_FORCE_ALPHA_CLAMP
-#define AA_FORCE_ALPHA_CLAMP 0
+	#define AA_FORCE_ALPHA_CLAMP 0
 #endif
 
 // Use velocity weighting.
 #ifndef AA_VELOCITY_WEIGHTING
-#define AA_VELOCITY_WEIGHTING 0
+	#define AA_VELOCITY_WEIGHTING 0
 #endif
 
 // Use YCoCg path.
 #ifndef AA_YCOCG
-#define AA_YCOCG 0
+	#define AA_YCOCG 0
 #endif
 
 // Use green as luma.
 #ifndef AA_GREEN_AS_LUMA
-#define AA_GREEN_AS_LUMA AA_YCOCG
+	#define AA_GREEN_AS_LUMA AA_YCOCG
 #endif
 
 // Special adjustments for DOF.
 #ifndef AA_DOF
-#define AA_DOF 0
+	#define AA_DOF 0
 #endif
 
 // Completely non-optimized prototype for solving noise problem.
 #ifndef AA_TEST
-#define AA_TEST 0
-#endif
-
-
-// FIND MOTION OF PIXEL AND NEAREST IN NEIGHBORHOOD
-// ------------------------------------------------
-float3 PosN; // Position of this pixel, possibly later nearest pixel in neighborhood.
-PosN.xy = UVAndScreenPos.zw * float2(0.5, -0.5) + 0.5; // View position [0 to 1] flipped in Y.
-PosN.z = SceneDepthTexture.SampleLevel(SceneDepthTextureSampler, UVAndScreenPos.xy, 0).r;
-// Screen position of minimum depth.
-float2 VelocityOffset = float2(0.0, 0.0);
-#if AA_CROSS
-// For motion vector, use camera/dynamic motion from min depth pixel in pattern around pixel.
-// This enables better quality outline on foreground against different motion background.
-// Larger 2 pixel distance "x" works best (because AA dilates surface).
-float4 Depths;
-#if FEATURE_LEVEL >= FEATURE_LEVEL_SM5
-Depths = SceneDepthTexture.GatherRed(SceneDepthTextureSampler, UVAndScreenPos.xy, int2(-AA_CROSS, -AA_CROSS), int2(AA_CROSS, -AA_CROSS), int2(-AA_CROSS, AA_CROSS), int2(AA_CROSS, AA_CROSS));
-#else
-Depths.x = SceneDepthTexture.SampleLevel(SceneDepthTextureSampler, UVAndScreenPos.xy, 0, int2(-AA_CROSS, -AA_CROSS)).r;
-Depths.y = SceneDepthTexture.SampleLevel(SceneDepthTextureSampler, UVAndScreenPos.xy, 0, int2(AA_CROSS, -AA_CROSS)).r;
-Depths.z = SceneDepthTexture.SampleLevel(SceneDepthTextureSampler, UVAndScreenPos.xy, 0, int2(-AA_CROSS, AA_CROSS)).r;
-Depths.w = SceneDepthTexture.SampleLevel(SceneDepthTextureSampler, UVAndScreenPos.xy, 0, int2(AA_CROSS, AA_CROSS)).r;
-#endif
-float2 DepthOffset = float2(AA_CROSS, AA_CROSS);
-float DepthOffsetXx = float(AA_CROSS);
+	#define AA_TEST 0
+#endif
+
+
+	// FIND MOTION OF PIXEL AND NEAREST IN NEIGHBORHOOD
+    // ------------------------------------------------
+	float3 PosN; // Position of this pixel, possibly later nearest pixel in neighborhood.
+	PosN.xy = UVAndScreenPos.zw * float2(0.5, -0.5) + 0.5; // View position [0 to 1] flipped in Y.
+	PosN.z = SceneDepthTexture.SampleLevel(SceneDepthTextureSampler, UVAndScreenPos.xy, 0).r;
+	// Screen position of minimum depth.
+	float2 VelocityOffset = float2(0.0, 0.0);
+	#if AA_CROSS
+		// For motion vector, use camera/dynamic motion from min depth pixel in pattern around pixel.
+		// This enables better quality outline on foreground against different motion background.
+		// Larger 2 pixel distance "x" works best (because AA dilates surface).
+		float4 Depths;
+		#if FEATURE_LEVEL >= FEATURE_LEVEL_SM5
+			Depths = SceneDepthTexture.GatherRed(SceneDepthTextureSampler, UVAndScreenPos.xy, int2(-AA_CROSS, -AA_CROSS), int2(AA_CROSS, -AA_CROSS), int2(-AA_CROSS, AA_CROSS), int2(AA_CROSS, AA_CROSS));
+		#else
+			Depths.x = SceneDepthTexture.SampleLevel(SceneDepthTextureSampler, UVAndScreenPos.xy, 0, int2(-AA_CROSS, -AA_CROSS)).r;
+			Depths.y = SceneDepthTexture.SampleLevel(SceneDepthTextureSampler, UVAndScreenPos.xy, 0, int2( AA_CROSS, -AA_CROSS)).r;
+			Depths.z = SceneDepthTexture.SampleLevel(SceneDepthTextureSampler, UVAndScreenPos.xy, 0, int2(-AA_CROSS,  AA_CROSS)).r;
+			Depths.w = SceneDepthTexture.SampleLevel(SceneDepthTextureSampler, UVAndScreenPos.xy, 0, int2( AA_CROSS,  AA_CROSS)).r;
+		#endif
+		float2 DepthOffset = float2(AA_CROSS, AA_CROSS);
+		float DepthOffsetXx = float(AA_CROSS);
 #if HAS_INVERTED_Z_BUFFER
-//@todo-briank
-// Nearest depth is the largest depth (depth surface 0=far, 1=near).
-if (Depths.x > Depths.y)
-{
-	DepthOffsetXx = -AA_CROSS;
-}
-if (Depths.z > Depths.w)
-{
-	DepthOffset.x = -AA_CROSS;
-}
-float DepthsXY = max(Depths.x, Depths.y);
-float DepthsZW = max(Depths.z, Depths.w);
-if (DepthsXY > DepthsZW)
-{
-	DepthOffset.y = -AA_CROSS;
-	DepthOffset.x = DepthOffsetXx;
-}
-float DepthsXYZW = max(DepthsXY, DepthsZW);
-if (DepthsXYZW > PosN.z)
-{
-	// This is offset for reading from velocity texture.
-	// This supports half or fractional resolution velocity textures.
-	// With the assumption that UV position scales between velocity and color.
-	VelocityOffset = DepthOffset * PostprocessInput0Size.zw;
-	// This is [0 to 1] flipped in Y.
-	PosN.xy = (UVAndScreenPos.zw + DepthOffset * ViewportSize.zw * 2.0) * float2(0.5, -0.5) + 0.5;
-	PosN.z = DepthsXYZW;
-}
+		//@todo-briank
+		// Nearest depth is the largest depth (depth surface 0=far, 1=near).
+		if(Depths.x > Depths.y) 
+		{
+			DepthOffsetXx = -AA_CROSS;
+		}
+		if(Depths.z > Depths.w) 
+		{
+			DepthOffset.x = -AA_CROSS;
+		}
+		float DepthsXY = max(Depths.x, Depths.y);
+		float DepthsZW = max(Depths.z, Depths.w);
+		if(DepthsXY > DepthsZW) 
+		{
+			DepthOffset.y = -AA_CROSS;
+			DepthOffset.x = DepthOffsetXx; 
+		}
+		float DepthsXYZW = max(DepthsXY, DepthsZW);
+		if(DepthsXYZW > PosN.z) 
+		{
+			// This is offset for reading from velocity texture.
+			// This supports half or fractional resolution velocity textures.
+			// With the assumption that UV position scales between velocity and color.
+			VelocityOffset = DepthOffset * PostprocessInput0Size.zw;
+			// This is [0 to 1] flipped in Y.
+			PosN.xy = (UVAndScreenPos.zw + DepthOffset * ViewportSize.zw * 2.0) * float2(0.5, -0.5) + 0.5;
+			PosN.z = DepthsXYZW;
+		}
 #else
 #error Fix me!
 #endif // HAS_INVERTED_Z_BUFFER
-#endif	// AA_CROSS
-// Camera motion for pixel or nearest pixel (in ScreenPos space).
-float ScaleM = 1.0 / (dot(PosN, CameraMotion[0].xyz) + CameraMotion[0].w);
-float2 BackN;
-BackN.x = -2.0 * ((PosN.x * ((CameraMotion[1].x * PosN.y) + (CameraMotion[1].y * PosN.z) + CameraMotion[1].z)) + (CameraMotion[1].w * PosN.y) + (CameraMotion[2].x * PosN.x * PosN.x) + (CameraMotion[2].y * PosN.z) + CameraMotion[2].z) * ScaleM;
-BackN.y = 2.0 * ((PosN.y * ((CameraMotion[3].x * PosN.x) + (CameraMotion[3].y * PosN.z) + CameraMotion[3].z)) + (CameraMotion[3].w * PosN.x) + (CameraMotion[4].x * PosN.y * PosN.y) + (CameraMotion[4].y * PosN.z) + CameraMotion[4].z) * ScaleM;
-
-//Fix for VR 'wobble' Imprecision in this method makes it not good enough to have on by default.  A proper fix is coming soon, but in the meantime, VR users can put this in locally.
-#if 0
-{
-	float4 ThisClip = float4((UVAndScreenPos.xy - View.ScreenPositionScaleBias.wz) / View.ScreenPositionScaleBias.xy, PosN.z, 1);
-	float4 ThisTranslatedWorld = mul(ThisClip, View.ClipToTranslatedWorld);
-	ThisTranslatedWorld /= ThisTranslatedWorld.w;
-
-	float3 PrevTranslatedWorld = ThisTranslatedWorld.xyz + (View.PrevPreViewTranslation - View.PreViewTranslation);
-	float4 PrevClip = mul(float4(PrevTranslatedWorld, 1), View.PrevTranslatedWorldToClip);
-	float2 PrevScreen = PrevClip.xy / PrevClip.w;
-	//float2 PrevUV = PrevScreen.xy * View.ScreenPositionScaleBias.xy + View.ScreenPositionScaleBias.wz;
-
-	BackN = UVAndScreenPos.zw - PrevScreen - (View.ViewToClip[2].xy - View.PrevProjection[2].xy);
-}
-#endif
-
-float2 BackTemp = BackN * ViewportSize.xy;
-#if AA_DYNAMIC
-float2 VelocityN;
-#if AA_CROSS
-VelocityN = PostprocessInput3.SampleLevel(PostprocessInput3Sampler, UVAndScreenPos.xy + VelocityOffset, 0).xy;
-#else
-VelocityN = PostprocessInput3.SampleLevel(PostprocessInput3Sampler, UVAndScreenPos.xy, 0).xy;
-#endif
-bool DynamicN = VelocityN.x > 0.0;
-if (DynamicN)
-{
-	BackN = DecodeVelocityFromTexture(VelocityN);
-}
-BackTemp = BackN * ViewportSize.xy;
-#endif
-// Save the amount of pixel offset of just camera motion, used later as the amount of blur introduced by history.
-float HistoryBlurAmp = 2.0;
-float HistoryBlur = saturate(abs(BackTemp.x) * HistoryBlurAmp + abs(BackTemp.y) * HistoryBlurAmp);
-float Velocity = sqrt(dot(BackTemp, BackTemp));
-// Easier to do off screen check before conversion.
-// BackN is in units of 2pixels/viewportWidthInPixels
-// This converts back projection vector to [-1 to 1] offset in viewport.
-BackN = UVAndScreenPos.zw - BackN;
-bool OffScreen = max(abs(BackN.x), abs(BackN.y)) >= 1.0;
-// Also clamp to be on screen (fixes problem with DOF).
-// The .z and .w is the 1/width and 1/height.
-// This clamps to be a pixel inside the viewport.
-BackN.x = clamp(BackN.x, -1.0 + ViewportSize.z, 1.0 - ViewportSize.z);
-BackN.y = clamp(BackN.y, -1.0 + ViewportSize.w, 1.0 - ViewportSize.w);
-// Convert from [-1 to 1] to view rectangle which is somewhere in [0 to 1].
-// The extra +0.5 factor is because ScreenPosToPixel.zw is incorrectly computed
-// as the upper left of the pixel instead of the center of the pixel.
-BackN = (BackN * ScreenPosToPixel.xy + ScreenPosToPixel.zw + 0.5) * PostprocessInput0Size.zw;
-
-
-
-// FILTER PIXEL (RESAMPLE TO REMOVE JITTER OFFSET) AND GET NEIGHBORHOOD
-// --------------------------------------------------------------------
-// 012
-// 345
-// 678
-#if AA_YCOCG	
-// Special case, only using 5 taps.
-float4 Neighbor1 = PostprocessInput0.SampleLevel(PostprocessInput0Sampler, UVAndScreenPos.xy, 0, int2(0, -1));
-float4 Neighbor3 = PostprocessInput0.SampleLevel(PostprocessInput0Sampler, UVAndScreenPos.xy, 0, int2(-1, 0));
-float4 Neighbor4 = PostprocessInput0.SampleLevel(PostprocessInput0Sampler, UVAndScreenPos.xy, 0);
-float4 Neighbor5 = PostprocessInput0.SampleLevel(PostprocessInput0Sampler, UVAndScreenPos.xy, 0, int2(1, 0));
-float4 Neighbor7 = PostprocessInput0.SampleLevel(PostprocessInput0Sampler, UVAndScreenPos.xy, 0, int2(0, 1));
-#if 1
-Neighbor1.rgb *= HdrWeightG(Neighbor1.rgb, InExposureScale);
-Neighbor3.rgb *= HdrWeightG(Neighbor3.rgb, InExposureScale);
-Neighbor4.rgb *= HdrWeightG(Neighbor4.rgb, InExposureScale);
-Neighbor5.rgb *= HdrWeightG(Neighbor5.rgb, InExposureScale);
-Neighbor7.rgb *= HdrWeightG(Neighbor7.rgb, InExposureScale);
-#endif
-Neighbor1.rgb = RGBToYCoCg(Neighbor1.rgb);
-Neighbor3.rgb = RGBToYCoCg(Neighbor3.rgb);
-Neighbor4.rgb = RGBToYCoCg(Neighbor4.rgb);
-Neighbor5.rgb = RGBToYCoCg(Neighbor5.rgb);
-Neighbor7.rgb = RGBToYCoCg(Neighbor7.rgb);
-#if 0
-Neighbor1.xyz *= HdrWeightY(Neighbor1.x, InExposureScale);
-Neighbor3.xyz *= HdrWeightY(Neighbor3.x, InExposureScale);
-Neighbor4.xyz *= HdrWeightY(Neighbor4.x, InExposureScale);
-Neighbor5.xyz *= HdrWeightY(Neighbor5.x, InExposureScale);
-Neighbor7.xyz *= HdrWeightY(Neighbor7.x, InExposureScale);
-#endif
-#if AA_FILTERED
-float4 Filtered =
-Neighbor1 * PlusWeights[0] +
-Neighbor3 * PlusWeights[1] +
-Neighbor4 * PlusWeights[2] +
-Neighbor5 * PlusWeights[3] +
-Neighbor7 * PlusWeights[4];
-#if AA_BORDER
-// Use unfiltered for 1 pixel border.
-float2 TestPos = abs(UVAndScreenPos.zw);
-// Add 1 pixel and check if off screen.
-TestPos += ViewportSize.zw * 2.0;
-bool FilteredOffScreen = max(TestPos.x, TestPos.y) >= 1.0;
-if (FilteredOffScreen)
-{
-	Filtered = Neighbor4;
-}
-#endif
-#else
-// Unfiltered.
-float4 Filtered = Neighbor4;
-#endif
-float4 FilteredLow = Filtered;
-// Neighborhood seems to only need the "+" pattern.
-float4 NeighborMin = min(min(min(Neighbor1, Neighbor3), min(Neighbor4, Neighbor5)), Neighbor7);
-float4 NeighborMax = max(max(max(Neighbor1, Neighbor3), max(Neighbor4, Neighbor5)), Neighbor7);
-#else
-float4 Neighbor0 = PostprocessInput0.SampleLevel(PostprocessInput0Sampler, UVAndScreenPos.xy, 0, int2(-1, -1));
-float4 Neighbor1 = PostprocessInput0.SampleLevel(PostprocessInput0Sampler, UVAndScreenPos.xy, 0, int2(0, -1));
-float4 Neighbor2 = PostprocessInput0.SampleLevel(PostprocessInput0Sampler, UVAndScreenPos.xy, 0, int2(1, -1));
-float4 Neighbor3 = PostprocessInput0.SampleLevel(PostprocessInput0Sampler, UVAndScreenPos.xy, 0, int2(-1, 0));
-float4 Neighbor4 = PostprocessInput0.SampleLevel(PostprocessInput0Sampler, UVAndScreenPos.xy, 0);
-float4 Neighbor5 = PostprocessInput0.SampleLevel(PostprocessInput0Sampler, UVAndScreenPos.xy, 0, int2(1, 0));
-float4 Neighbor6 = PostprocessInput0.SampleLevel(PostprocessInput0Sampler, UVAndScreenPos.xy, 0, int2(-1, 1));
-float4 Neighbor7 = PostprocessInput0.SampleLevel(PostprocessInput0Sampler, UVAndScreenPos.xy, 0, int2(0, 1));
-float4 Neighbor8 = PostprocessInput0.SampleLevel(PostprocessInput0Sampler, UVAndScreenPos.xy, 0, int2(1, 1));
-#if AA_GREEN_AS_LUMA
-Neighbor0.rgb *= HdrWeightG(Neighbor0.rgb, InExposureScale);
-Neighbor1.rgb *= HdrWeightG(Neighbor1.rgb, InExposureScale);
-Neighbor2.rgb *= HdrWeightG(Neighbor2.rgb, InExposureScale);
-Neighbor3.rgb *= HdrWeightG(Neighbor3.rgb, InExposureScale);
-Neighbor4.rgb *= HdrWeightG(Neighbor4.rgb, InExposureScale);
-Neighbor5.rgb *= HdrWeightG(Neighbor5.rgb, InExposureScale);
-Neighbor6.rgb *= HdrWeightG(Neighbor6.rgb, InExposureScale);
-Neighbor7.rgb *= HdrWeightG(Neighbor7.rgb, InExposureScale);
-Neighbor8.rgb *= HdrWeightG(Neighbor8.rgb, InExposureScale);
-#else
-Neighbor0.rgb *= HdrWeight4(Neighbor0.rgb, InExposureScale);
-Neighbor1.rgb *= HdrWeight4(Neighbor1.rgb, InExposureScale);
-Neighbor2.rgb *= HdrWeight4(Neighbor2.rgb, InExposureScale);
-Neighbor3.rgb *= HdrWeight4(Neighbor3.rgb, InExposureScale);
-Neighbor4.rgb *= HdrWeight4(Neighbor4.rgb, InExposureScale);
-Neighbor5.rgb *= HdrWeight4(Neighbor5.rgb, InExposureScale);
-Neighbor6.rgb *= HdrWeight4(Neighbor6.rgb, InExposureScale);
-Neighbor7.rgb *= HdrWeight4(Neighbor7.rgb, InExposureScale);
-Neighbor8.rgb *= HdrWeight4(Neighbor8.rgb, InExposureScale);
-#endif
-#if AA_FILTERED
-float4 Filtered =
-Neighbor0 * SampleWeights[0] +
-Neighbor1 * SampleWeights[1] +
-Neighbor2 * SampleWeights[2] +
-Neighbor3 * SampleWeights[3] +
-Neighbor4 * SampleWeights[4] +
-Neighbor5 * SampleWeights[5] +
-Neighbor6 * SampleWeights[6] +
-Neighbor7 * SampleWeights[7] +
-Neighbor8 * SampleWeights[8];
-#if AA_LOWPASS
-float4 FilteredLow =
-Neighbor0 * LowpassWeights[0] +
-Neighbor1 * LowpassWeights[1] +
-Neighbor2 * LowpassWeights[2] +
-Neighbor3 * LowpassWeights[3] +
-Neighbor4 * LowpassWeights[4] +
-Neighbor5 * LowpassWeights[5] +
-Neighbor6 * LowpassWeights[6] +
-Neighbor7 * LowpassWeights[7] +
-Neighbor8 * LowpassWeights[8];
-#else
-float4 FilteredLow = Filtered;
-#endif
-#if AA_BORDER
-// Use unfiltered for 1 pixel border.
-float2 TestPos = abs(UVAndScreenPos.zw);
-// Add 1 pixel and check if off screen.
-TestPos += ViewportSize.zw * 2.0;
-bool FilteredOffScreen = max(TestPos.x, TestPos.y) >= 1.0;
-if (FilteredOffScreen)
-{
-	Filtered = Neighbor4;
-	FilteredLow = Neighbor4;
-}
-#endif
-#else
-// Unfiltered.
-float4 Filtered = Neighbor4;
-float4 FilteredLow = Neighbor4;
-#endif
-#if AA_ROUND
-float4 NeighborMin2 = min(min(Neighbor0, Neighbor2), min(Neighbor6, Neighbor8));
-float4 NeighborMax2 = max(max(Neighbor0, Neighbor2), max(Neighbor6, Neighbor8));
-float4 NeighborMin = min(min(min(Neighbor1, Neighbor3), min(Neighbor4, Neighbor5)), Neighbor7);
-float4 NeighborMax = max(max(max(Neighbor1, Neighbor3), max(Neighbor4, Neighbor5)), Neighbor7);
-NeighborMin2 = min(NeighborMin2, NeighborMin);
-NeighborMax2 = max(NeighborMax2, NeighborMax);
-NeighborMin = NeighborMin * 0.5 + NeighborMin2 * 0.5;
-NeighborMax = NeighborMax * 0.5 + NeighborMax2 * 0.5;
-#else
-float4 NeighborMin = min(min(
-	min(min(Neighbor0, Neighbor1), min(Neighbor2, Neighbor3)),
-	min(min(Neighbor4, Neighbor5), min(Neighbor6, Neighbor7))), Neighbor8);
-float4 NeighborMax = max(max(
-	max(max(Neighbor0, Neighbor1), max(Neighbor2, Neighbor3)),
-	max(max(Neighbor4, Neighbor5), max(Neighbor6, Neighbor7))), Neighbor8);
-#endif
-#endif
-
-
+	#endif	// AA_CROSS
+	// Camera motion for pixel or nearest pixel (in ScreenPos space).
+	float ScaleM = 1.0 / (dot(PosN, CameraMotion[0].xyz) + CameraMotion[0].w);
+	float2 BackN;
+	BackN.x = -2.0 * ((PosN.x * ((CameraMotion[1].x * PosN.y) + (CameraMotion[1].y * PosN.z) + CameraMotion[1].z)) + (CameraMotion[1].w * PosN.y) + (CameraMotion[2].x * PosN.x * PosN.x) + (CameraMotion[2].y * PosN.z) + CameraMotion[2].z) * ScaleM;
+	BackN.y =  2.0 * ((PosN.y * ((CameraMotion[3].x * PosN.x) + (CameraMotion[3].y * PosN.z) + CameraMotion[3].z)) + (CameraMotion[3].w * PosN.x) + (CameraMotion[4].x * PosN.y * PosN.y) + (CameraMotion[4].y * PosN.z) + CameraMotion[4].z) * ScaleM;
 
-// FETCH HISTORY
-// -------------
-OutColor = PostprocessInput1.SampleLevel(PostprocessInput1Sampler, BackN.xy, 0);
-#if AA_DEBUG
-Neighbor4.rg = float2(0.0, 0.0);
-NeighborMin.rg = float2(0.0, 0.0);
-NeighborMax.rg = float2(0.0, 0.0);
-Filtered.rg = float2(0.0, 0.0);
-FilteredLow.rg = float2(0.0, 0.0);
-float DebugDiffCurrent = Filtered.b;
-#endif
-#if AA_YCOCG
-#if 1
-OutColor.rgb *= HdrWeightG(OutColor.rgb, InExposureScale);
-#endif
-OutColor.rgb = RGBToYCoCg(OutColor.rgb);
 #if 0
-OutColor.xyz *= HdrWeightY(OutColor.x, InExposureScale);
-#endif
-#else
-#if AA_GREEN_AS_LUMA
-OutColor.rgb *= HdrWeightG(OutColor.rgb, InExposureScale);
-#else
-OutColor.rgb *= HdrWeight4(OutColor.rgb, InExposureScale);
-#endif
-#endif
-#if AA_DEBUG
-OutColor.rg = float2(0.0, 0.0);
-float DebugDiffPrior = OutColor.b;
-#endif
-
-
-
-
-
-
-// FIND LUMA OF CLAMPED HISTORY
-// ----------------------------
-// Save off luma of history before the clamp.
-#if AA_YCOCG
-float LumaMin = NeighborMin.x;
-float LumaMax = NeighborMax.x;
-float LumaHistory = OutColor.x;
-#else
-#if AA_GREEN_AS_LUMA
-float LumaMin = NeighborMin.g;
-float LumaMax = NeighborMax.g;
-float LumaHistory = OutColor.g;
-#else
-float LumaMin = Luma4(NeighborMin.rgb);
-float LumaMax = Luma4(NeighborMax.rgb);
-float LumaHistory = Luma4(OutColor.rgb);
-#endif
-#endif
-float LumaContrast = LumaMax - LumaMin;
-#if AA_YCOCG
-OutColor.rgb = clamp(OutColor.rgb, NeighborMin.rgb, NeighborMax.rgb);
-#if (AA_ALPHA == 0)
-OutColor.a = clamp(OutColor.a, NeighborMin.a, NeighborMax.a);
-#endif
-#else
-#if AA_AABB
-// Clamp history, this uses color AABB intersection for tighter fit.
-// Clamping works with the low pass (if available) to reduce flicker.
-float ClampBlend = HistoryClamp(OutColor.rgb, FilteredLow.rgb, NeighborMin.rgb, NeighborMax.rgb);
-#if AA_ALPHA
-OutColor.rgb = lerp(OutColor.rgb, FilteredLow.rgb, ClampBlend);
-#else
-OutColor.rgba = lerp(OutColor.rgba, FilteredLow.rgba, ClampBlend);
-#endif
-#else
-OutColor = clamp(OutColor, NeighborMin, NeighborMax);
-#endif
-#endif
-#if AA_DEBUG
-OutColor.rg = float2(0.0, 0.0);
-#endif
-
-// ADD BACK IN ALIASING TO SHARPEN
-// -------------------------------
-#if AA_FILTERED
-// Blend in non-filtered based on the amount of sub-pixel motion.
-float AddAliasing = saturate(HistoryBlur) * 0.5;
-float LumaContrastFactor = 32.0;
-#if AA_GREEN_AS_LUMA || AA_YCOCG
-// GREEN_AS_LUMA is 1/4 as bright.
-LumaContrastFactor *= 4.0;
-#endif
-AddAliasing = saturate(AddAliasing + rcp(1.0 + LumaContrast * LumaContrastFactor));
-Filtered.rgb = lerp(Filtered.rgb, Neighbor4.rgb, AddAliasing);
-#endif
-#if AA_YCOCG
-float LumaFiltered = Filtered.x;
-#else
-#if AA_GREEN_AS_LUMA
-float LumaFiltered = Filtered.g;
-#else
-float LumaFiltered = Luma4(Filtered.rgb);
-#endif
-#endif
-
-// COMPUTE BLEND AMOUNT 
-// --------------------
-// Replace history with minimum difference of history from min or max neighborhood.
-LumaHistory = min(abs(LumaMin - LumaHistory), abs(LumaMax - LumaHistory));
-float HistoryAmount = (1.0 / 8.0) + HistoryBlur * (1.0 / 8.0);
-float HistoryFactor = LumaHistory * HistoryAmount * (1.0 + HistoryBlur * HistoryAmount * 8.0);
-float BlendFinal = saturate(HistoryFactor * rcp(LumaHistory + LumaContrast));
-#if RESPONSIVE
-// Responsive forces 1/4 of new frame.
-BlendFinal = 1.0 / 4.0;
-#endif
-
-#if (AA_ALPHA == 1) && (AA_VELOCITY_WEIGHTING == 1)
-// Velocity weighting.
-// Decay prior value.
-float PriorVelocity = OutColor.a;
-float VelocityDecay = 0.5;
-OutColor.a = max(OutColor.a * VelocityDecay, Velocity * rcp(VelocityDecay));
-float VelocityDiff = abs(PriorVelocity - Velocity) / max(1.0, max(PriorVelocity, Velocity));
-BlendFinal = max(BlendFinal, VelocityDiff * (1.0 / 8.0));
-BlendFinal = min(1.0 / 2.0, BlendFinal);
-#endif
-
-// Offscreen feedback resets.
-#if AA_LERP 
-float FixedLerp = 1.0 / float(AA_LERP);
-#endif
-if (OffScreen)
-{
-	OutColor = Filtered;
-#if AA_ALPHA
-	OutColor.a = 0.0;
-#endif
-#if AA_LERP
-	FixedLerp = 1.0;
-#endif
-}
-
-// DO FINAL BLEND BETWEEN HISTORY AND FILTERED COLOR
-// -------------------------------------------------
-#if (AA_LERP == 0)
-#if AA_ALPHA
-// Blend in linear to hit luma target.
-OutColor.rgb = lerp(OutColor.rgb, Filtered.rgb, BlendFinal);
-#if RESPONSIVE
-OutColor.a = max(OutColor.a, 1.0 / 2.0);
-#endif
-#else
-OutColor = lerp(OutColor, Filtered, BlendFinal);
-#if AA_FORCE_ALPHA_CLAMP
-OutColor.a = clamp(OutColor.a, NeighborMin.a, NeighborMax.a);
-#endif
-#endif
-#else
-OutColor = lerp(OutColor, Filtered, FixedLerp);
-#endif
-#if AA_YCOCG
-#if 0
-OutColor.xyz *= HdrWeightInvY(OutColor.x, InExposureScale);
-#endif
-OutColor.rgb = YCoCgToRGB(OutColor.rgb);
-#if 1
-OutColor.rgb *= HdrWeightInvG(OutColor.rgb, InExposureScale);
-#endif
-#else
-// Convert back into linear.
-#if AA_GREEN_AS_LUMA
-OutColor.rgb *= HdrWeightInvG(OutColor.rgb, InExposureScale);
-#else
-OutColor.rgb *= HdrWeightInv4(OutColor.rgb, InExposureScale);
-#endif
-#endif
-#if AA_NAN
-// Transform NaNs to black, transform negative colors to black.
-OutColor.rgb = -min(-OutColor.rgb, 0.0);
-#endif
-#if AA_DEBUG
-OutColor.g = abs(DebugDiffPrior - DebugDiffCurrent);
-OutColor.r = OutColor.a;
-#endif
+	//Fix for VR 'wobble' Imprecision in this method makes it not good enough to have on by default.  A proper fix is coming soon, but in the meantime, VR users can put this in locally.
+	{
+		float4 ThisClip = float4( (UVAndScreenPos.xy - View.ScreenPositionScaleBias.wz) / View.ScreenPositionScaleBias.xy, PosN.z, 1 );
+		float4 ThisTranslatedWorld = mul( ThisClip, View.ClipToTranslatedWorld );
+		ThisTranslatedWorld /= ThisTranslatedWorld.w;
+		
+		float3 PrevTranslatedWorld = ThisTranslatedWorld.xyz + (View.PrevPreViewTranslation - View.PreViewTranslation);
+		float4 PrevClip = mul( float4( PrevTranslatedWorld, 1 ), View.PrevTranslatedWorldToClip );
+		float2 PrevScreen = PrevClip.xy / PrevClip.w;
+		//float2 PrevUV = PrevScreen.xy * View.ScreenPositionScaleBias.xy + View.ScreenPositionScaleBias.wz;
+
+		BackN = UVAndScreenPos.zw - PrevScreen - ( View.ViewToClip[2].xy - View.PrevProjection[2].xy );
+	}
+#endif
+
+	float2 BackTemp = BackN * ViewportSize.xy;
+	#if AA_DYNAMIC
+		float2 VelocityN;
+		#if AA_CROSS
+			VelocityN = PostprocessInput3.SampleLevel(PostprocessInput3Sampler, UVAndScreenPos.xy + VelocityOffset, 0).xy;
+		#else
+			VelocityN = PostprocessInput3.SampleLevel(PostprocessInput3Sampler, UVAndScreenPos.xy, 0).xy;
+		#endif
+		bool DynamicN = VelocityN.x > 0.0;
+		if(DynamicN)
+		{
+			BackN = DecodeVelocityFromTexture(VelocityN);
+		}
+		BackTemp = BackN * ViewportSize.xy;
+	#endif
+	// Save the amount of pixel offset of just camera motion, used later as the amount of blur introduced by history.
+	float HistoryBlurAmp = 2.0;
+	float HistoryBlur = saturate(abs(BackTemp.x) * HistoryBlurAmp + abs(BackTemp.y) * HistoryBlurAmp);
+	float Velocity = sqrt(dot(BackTemp, BackTemp));
+	// Easier to do off screen check before conversion.
+	// BackN is in units of 2pixels/viewportWidthInPixels
+	// This converts back projection vector to [-1 to 1] offset in viewport.
+	BackN = UVAndScreenPos.zw - BackN;
+	bool OffScreen = max(abs(BackN.x), abs(BackN.y)) >= 1.0;
+	// Also clamp to be on screen (fixes problem with DOF).
+	// The .z and .w is the 1/width and 1/height.
+	// This clamps to be a pixel inside the viewport.
+	BackN.x = clamp(BackN.x, -1.0 + ViewportSize.z, 1.0 - ViewportSize.z);
+	BackN.y = clamp(BackN.y, -1.0 + ViewportSize.w, 1.0 - ViewportSize.w);
+	// Convert from [-1 to 1] to view rectangle which is somewhere in [0 to 1].
+	// The extra +0.5 factor is because ScreenPosToPixel.zw is incorrectly computed
+	// as the upper left of the pixel instead of the center of the pixel.
+	BackN = (BackN * ScreenPosToPixel.xy + ScreenPosToPixel.zw + 0.5) * PostprocessInput0Size.zw;
+
+
+
+	// FILTER PIXEL (RESAMPLE TO REMOVE JITTER OFFSET) AND GET NEIGHBORHOOD
+    // --------------------------------------------------------------------
+	// 012
+	// 345
+	// 678
+	#if AA_YCOCG	
+		// Special case, only using 5 taps.
+		float4 Neighbor1 = PostprocessInput0.SampleLevel(PostprocessInput0Sampler, UVAndScreenPos.xy, 0, int2( 0, -1));
+		float4 Neighbor3 = PostprocessInput0.SampleLevel(PostprocessInput0Sampler, UVAndScreenPos.xy, 0, int2(-1,  0));
+		float4 Neighbor4 = PostprocessInput0.SampleLevel(PostprocessInput0Sampler, UVAndScreenPos.xy, 0);
+		float4 Neighbor5 = PostprocessInput0.SampleLevel(PostprocessInput0Sampler, UVAndScreenPos.xy, 0, int2( 1,  0));
+		float4 Neighbor7 = PostprocessInput0.SampleLevel(PostprocessInput0Sampler, UVAndScreenPos.xy, 0, int2( 0,  1));
+		#if 1
+			Neighbor1.rgb *= HdrWeightG(Neighbor1.rgb, InExposureScale);
+			Neighbor3.rgb *= HdrWeightG(Neighbor3.rgb, InExposureScale);
+			Neighbor4.rgb *= HdrWeightG(Neighbor4.rgb, InExposureScale);
+			Neighbor5.rgb *= HdrWeightG(Neighbor5.rgb, InExposureScale);
+			Neighbor7.rgb *= HdrWeightG(Neighbor7.rgb, InExposureScale);
+		#endif
+		Neighbor1.rgb = RGBToYCoCg(Neighbor1.rgb);
+		Neighbor3.rgb = RGBToYCoCg(Neighbor3.rgb);
+		Neighbor4.rgb = RGBToYCoCg(Neighbor4.rgb);
+		Neighbor5.rgb = RGBToYCoCg(Neighbor5.rgb);
+		Neighbor7.rgb = RGBToYCoCg(Neighbor7.rgb);
+		#if 0
+			Neighbor1.xyz *= HdrWeightY(Neighbor1.x, InExposureScale);
+			Neighbor3.xyz *= HdrWeightY(Neighbor3.x, InExposureScale);
+			Neighbor4.xyz *= HdrWeightY(Neighbor4.x, InExposureScale);
+			Neighbor5.xyz *= HdrWeightY(Neighbor5.x, InExposureScale);
+			Neighbor7.xyz *= HdrWeightY(Neighbor7.x, InExposureScale);
+		#endif
+		#if AA_FILTERED
+			float4 Filtered = 
+				Neighbor1 * PlusWeights[0] +
+				Neighbor3 * PlusWeights[1] +
+				Neighbor4 * PlusWeights[2] +
+				Neighbor5 * PlusWeights[3] +
+				Neighbor7 * PlusWeights[4];
+			#if AA_BORDER
+				// Use unfiltered for 1 pixel border.
+				float2 TestPos = abs(UVAndScreenPos.zw);
+				// Add 1 pixel and check if off screen.
+				TestPos += ViewportSize.zw * 2.0;
+				bool FilteredOffScreen = max(TestPos.x, TestPos.y) >= 1.0;
+				if(FilteredOffScreen) 
+				{
+					Filtered = Neighbor4;
+				}
+			#endif
+		#else
+			// Unfiltered.
+			float4 Filtered = Neighbor4;
+		#endif
+		float4 FilteredLow = Filtered;
+		// Neighborhood seems to only need the "+" pattern.
+		float4 NeighborMin = min(min(min(Neighbor1, Neighbor3), min(Neighbor4, Neighbor5)), Neighbor7);		
+		float4 NeighborMax = max(max(max(Neighbor1, Neighbor3), max(Neighbor4, Neighbor5)), Neighbor7);		
+	#else
+		float4 Neighbor0 = PostprocessInput0.SampleLevel(PostprocessInput0Sampler, UVAndScreenPos.xy, 0, int2(-1, -1));
+		float4 Neighbor1 = PostprocessInput0.SampleLevel(PostprocessInput0Sampler, UVAndScreenPos.xy, 0, int2( 0, -1));
+		float4 Neighbor2 = PostprocessInput0.SampleLevel(PostprocessInput0Sampler, UVAndScreenPos.xy, 0, int2( 1, -1));
+		float4 Neighbor3 = PostprocessInput0.SampleLevel(PostprocessInput0Sampler, UVAndScreenPos.xy, 0, int2(-1,  0));
+		float4 Neighbor4 = PostprocessInput0.SampleLevel(PostprocessInput0Sampler, UVAndScreenPos.xy, 0);
+		float4 Neighbor5 = PostprocessInput0.SampleLevel(PostprocessInput0Sampler, UVAndScreenPos.xy, 0, int2( 1,  0));
+		float4 Neighbor6 = PostprocessInput0.SampleLevel(PostprocessInput0Sampler, UVAndScreenPos.xy, 0, int2(-1,  1));
+		float4 Neighbor7 = PostprocessInput0.SampleLevel(PostprocessInput0Sampler, UVAndScreenPos.xy, 0, int2( 0,  1));
+		float4 Neighbor8 = PostprocessInput0.SampleLevel(PostprocessInput0Sampler, UVAndScreenPos.xy, 0, int2( 1,  1));
+		#if AA_GREEN_AS_LUMA
+			Neighbor0.rgb *= HdrWeightG(Neighbor0.rgb, InExposureScale);
+			Neighbor1.rgb *= HdrWeightG(Neighbor1.rgb, InExposureScale);
+			Neighbor2.rgb *= HdrWeightG(Neighbor2.rgb, InExposureScale);
+			Neighbor3.rgb *= HdrWeightG(Neighbor3.rgb, InExposureScale);
+			Neighbor4.rgb *= HdrWeightG(Neighbor4.rgb, InExposureScale);
+			Neighbor5.rgb *= HdrWeightG(Neighbor5.rgb, InExposureScale);
+			Neighbor6.rgb *= HdrWeightG(Neighbor6.rgb, InExposureScale);
+			Neighbor7.rgb *= HdrWeightG(Neighbor7.rgb, InExposureScale);
+			Neighbor8.rgb *= HdrWeightG(Neighbor8.rgb, InExposureScale);
+		#else
+			Neighbor0.rgb *= HdrWeight4(Neighbor0.rgb, InExposureScale);
+			Neighbor1.rgb *= HdrWeight4(Neighbor1.rgb, InExposureScale);
+			Neighbor2.rgb *= HdrWeight4(Neighbor2.rgb, InExposureScale);
+			Neighbor3.rgb *= HdrWeight4(Neighbor3.rgb, InExposureScale);
+			Neighbor4.rgb *= HdrWeight4(Neighbor4.rgb, InExposureScale);
+			Neighbor5.rgb *= HdrWeight4(Neighbor5.rgb, InExposureScale);
+			Neighbor6.rgb *= HdrWeight4(Neighbor6.rgb, InExposureScale);
+			Neighbor7.rgb *= HdrWeight4(Neighbor7.rgb, InExposureScale);
+			Neighbor8.rgb *= HdrWeight4(Neighbor8.rgb, InExposureScale);
+		#endif
+		#if AA_FILTERED
+			float4 Filtered = 
+				Neighbor0 * SampleWeights[0] +
+				Neighbor1 * SampleWeights[1] +
+				Neighbor2 * SampleWeights[2] +
+				Neighbor3 * SampleWeights[3] +
+				Neighbor4 * SampleWeights[4] +
+				Neighbor5 * SampleWeights[5] +
+				Neighbor6 * SampleWeights[6] +
+				Neighbor7 * SampleWeights[7] +
+				Neighbor8 * SampleWeights[8];
+			#if AA_LOWPASS
+				float4 FilteredLow = 
+					Neighbor0 * LowpassWeights[0] +
+					Neighbor1 * LowpassWeights[1] +
+					Neighbor2 * LowpassWeights[2] +
+					Neighbor3 * LowpassWeights[3] +
+					Neighbor4 * LowpassWeights[4] +
+					Neighbor5 * LowpassWeights[5] +
+					Neighbor6 * LowpassWeights[6] +
+					Neighbor7 * LowpassWeights[7] +
+					Neighbor8 * LowpassWeights[8];
+			#else
+				float4 FilteredLow = Filtered;
+			#endif
+			#if AA_BORDER
+				// Use unfiltered for 1 pixel border.
+				float2 TestPos = abs(UVAndScreenPos.zw);
+				// Add 1 pixel and check if off screen.
+				TestPos += ViewportSize.zw * 2.0;
+				bool FilteredOffScreen = max(TestPos.x, TestPos.y) >= 1.0;
+				if(FilteredOffScreen) 
+				{
+					Filtered = Neighbor4;
+					FilteredLow = Neighbor4;
+				}
+			#endif
+		#else
+			// Unfiltered.
+			float4 Filtered = Neighbor4;
+			float4 FilteredLow = Neighbor4;
+		#endif
+		#if AA_ROUND
+			float4 NeighborMin2 = min(min(Neighbor0, Neighbor2), min(Neighbor6, Neighbor8));		
+			float4 NeighborMax2 = max(max(Neighbor0, Neighbor2), max(Neighbor6, Neighbor8));		
+			float4 NeighborMin = min(min(min(Neighbor1, Neighbor3), min(Neighbor4, Neighbor5)), Neighbor7);		
+			float4 NeighborMax = max(max(max(Neighbor1, Neighbor3), max(Neighbor4, Neighbor5)), Neighbor7);		
+			NeighborMin2 = min(NeighborMin2, NeighborMin);
+			NeighborMax2 = max(NeighborMax2, NeighborMax);
+			NeighborMin = NeighborMin * 0.5 + NeighborMin2 * 0.5;
+			NeighborMax = NeighborMax * 0.5 + NeighborMax2 * 0.5;
+		#else
+			float4 NeighborMin = min(min(
+				min(min(Neighbor0, Neighbor1), min(Neighbor2, Neighbor3)),
+				min(min(Neighbor4, Neighbor5), min(Neighbor6, Neighbor7))), Neighbor8);
+			float4 NeighborMax = max(max(
+				max(max(Neighbor0, Neighbor1), max(Neighbor2, Neighbor3)),
+				max(max(Neighbor4, Neighbor5), max(Neighbor6, Neighbor7))), Neighbor8);
+		#endif
+	#endif
+
+
+
+	// FETCH HISTORY
+	// -------------
+	OutColor = PostprocessInput1.SampleLevel(PostprocessInput1Sampler, BackN.xy, 0);
+	#if AA_DEBUG
+		Neighbor4.rg = float2(0.0, 0.0);
+		NeighborMin.rg = float2(0.0, 0.0);
+		NeighborMax.rg = float2(0.0, 0.0);
+		Filtered.rg = float2(0.0, 0.0);
+		FilteredLow.rg = float2(0.0, 0.0);
+		float DebugDiffCurrent = Filtered.b;
+	#endif
+	#if AA_YCOCG
+		#if 1
+			OutColor.rgb *= HdrWeightG(OutColor.rgb, InExposureScale);
+		#endif
+		OutColor.rgb = RGBToYCoCg(OutColor.rgb);
+		#if 0
+			OutColor.xyz *= HdrWeightY(OutColor.x, InExposureScale);
+		#endif
+	#else
+		#if AA_GREEN_AS_LUMA
+			OutColor.rgb *= HdrWeightG(OutColor.rgb, InExposureScale);
+		#else
+			OutColor.rgb *= HdrWeight4(OutColor.rgb, InExposureScale);
+		#endif
+	#endif
+	#if AA_DEBUG
+		OutColor.rg = float2(0.0, 0.0);
+		float DebugDiffPrior = OutColor.b;
+	#endif
+
+
+
+
+
+
+	// FIND LUMA OF CLAMPED HISTORY
+	// ----------------------------
+	// Save off luma of history before the clamp.
+	#if AA_YCOCG
+		float LumaMin = NeighborMin.x;
+		float LumaMax = NeighborMax.x;
+		float LumaHistory = OutColor.x;
+	#else
+		#if AA_GREEN_AS_LUMA
+			float LumaMin = NeighborMin.g;
+			float LumaMax = NeighborMax.g;
+			float LumaHistory = OutColor.g;
+		#else
+			float LumaMin = Luma4(NeighborMin.rgb);
+			float LumaMax = Luma4(NeighborMax.rgb);
+			float LumaHistory = Luma4(OutColor.rgb);
+		#endif
+	#endif
+	float LumaContrast = LumaMax - LumaMin;
+	#if AA_YCOCG
+		OutColor.rgb = clamp(OutColor.rgb, NeighborMin.rgb, NeighborMax.rgb);
+		#if (AA_ALPHA == 0)
+			OutColor.a = clamp(OutColor.a, NeighborMin.a, NeighborMax.a);
+		#endif
+	#else
+		#if AA_AABB
+			// Clamp history, this uses color AABB intersection for tighter fit.
+			// Clamping works with the low pass (if available) to reduce flicker.
+			float ClampBlend = HistoryClamp(OutColor.rgb, FilteredLow.rgb, NeighborMin.rgb, NeighborMax.rgb);
+			#if AA_ALPHA
+				OutColor.rgb = lerp(OutColor.rgb, FilteredLow.rgb, ClampBlend);
+			#else
+				OutColor.rgba = lerp(OutColor.rgba, FilteredLow.rgba, ClampBlend);
+			#endif
+		#else
+			OutColor = clamp(OutColor, NeighborMin, NeighborMax);
+		#endif
+	#endif
+	#if AA_DEBUG
+		OutColor.rg = float2(0.0, 0.0);
+	#endif
+
+	// ADD BACK IN ALIASING TO SHARPEN
+    // -------------------------------
+	#if AA_FILTERED
+		// Blend in non-filtered based on the amount of sub-pixel motion.
+		float AddAliasing = saturate(HistoryBlur) * 0.5;
+		float LumaContrastFactor = 32.0;
+		#if AA_GREEN_AS_LUMA || AA_YCOCG
+			// GREEN_AS_LUMA is 1/4 as bright.
+			LumaContrastFactor *= 4.0;
+		#endif
+		AddAliasing = saturate(AddAliasing + rcp(1.0 + LumaContrast * LumaContrastFactor));
+		Filtered.rgb = lerp(Filtered.rgb, Neighbor4.rgb, AddAliasing);
+	#endif
+	#if AA_YCOCG
+		float LumaFiltered = Filtered.x;
+	#else
+		#if AA_GREEN_AS_LUMA
+			float LumaFiltered = Filtered.g;
+		#else
+			float LumaFiltered = Luma4(Filtered.rgb);
+		#endif
+	#endif
+		
+	// COMPUTE BLEND AMOUNT 
+    // --------------------
+	// Replace history with minimum difference of history from min or max neighborhood.
+	LumaHistory = min(abs(LumaMin-LumaHistory), abs(LumaMax-LumaHistory));
+	float HistoryAmount = (1.0/8.0) + HistoryBlur * (1.0/8.0);
+	float HistoryFactor = LumaHistory * HistoryAmount * (1.0 + HistoryBlur * HistoryAmount * 8.0);
+	float BlendFinal = saturate(HistoryFactor * rcp(LumaHistory + LumaContrast));
+	#if RESPONSIVE
+		// Responsive forces 1/4 of new frame.
+		BlendFinal = 1.0/4.0;
+	#endif
+
+   	#if (AA_ALPHA == 1) && (AA_VELOCITY_WEIGHTING == 1)
+		// Velocity weighting.
+		// Decay prior value.
+		float PriorVelocity = OutColor.a;
+		float VelocityDecay = 0.5;
+		OutColor.a = max(OutColor.a * VelocityDecay, Velocity * rcp(VelocityDecay));
+		float VelocityDiff = abs(PriorVelocity - Velocity) / max(1.0, max(PriorVelocity, Velocity));
+		BlendFinal = max(BlendFinal, VelocityDiff * (1.0/8.0));
+		BlendFinal = min(1.0/2.0, BlendFinal);
+	#endif
+
+	// Offscreen feedback resets.
+	#if AA_LERP 
+		float FixedLerp = 1.0/float(AA_LERP);
+	#endif
+	if(OffScreen) 
+	{
+		OutColor = Filtered;
+		#if AA_ALPHA
+			OutColor.a = 0.0;
+		#endif
+		#if AA_LERP
+			FixedLerp = 1.0;
+		#endif
+	}
+	
+	// DO FINAL BLEND BETWEEN HISTORY AND FILTERED COLOR
+	// -------------------------------------------------
+	#if (AA_LERP == 0)
+		#if AA_ALPHA
+			// Blend in linear to hit luma target.
+			OutColor.rgb = lerp(OutColor.rgb, Filtered.rgb, BlendFinal);
+			#if RESPONSIVE
+				OutColor.a = max(OutColor.a, 1.0/2.0);
+			#endif
+		#else
+			OutColor = lerp(OutColor, Filtered, BlendFinal);
+			#if AA_FORCE_ALPHA_CLAMP
+				OutColor.a = clamp(OutColor.a, NeighborMin.a, NeighborMax.a);
+			#endif
+		#endif
+	#else
+		OutColor = lerp(OutColor, Filtered, FixedLerp);
+	#endif
+	#if AA_YCOCG
+		#if 0
+			OutColor.xyz *= HdrWeightInvY(OutColor.x, InExposureScale);
+		#endif
+		OutColor.rgb = YCoCgToRGB(OutColor.rgb);
+		#if 1
+			OutColor.rgb *= HdrWeightInvG(OutColor.rgb, InExposureScale);
+		#endif
+	#else
+		// Convert back into linear.
+		#if AA_GREEN_AS_LUMA
+			OutColor.rgb *= HdrWeightInvG(OutColor.rgb, InExposureScale);
+		#else
+			OutColor.rgb *= HdrWeightInv4(OutColor.rgb, InExposureScale);
+		#endif
+	#endif
+	#if AA_NAN
+		// Transform NaNs to black, transform negative colors to black.
+		OutColor.rgb = -min(-OutColor.rgb, 0.0);
+	#endif
+	#if AA_DEBUG
+		OutColor.g = abs(DebugDiffPrior - DebugDiffCurrent);
+		OutColor.r = OutColor.a;
+	#endif
 
 
 
