fix shader compile failure when compiling for Metal
#ios

diff --git a/Engine/Shaders/TonemapCommon.usf b/Engine/Shaders/TonemapCommon.usf
--- a/Engine/Shaders/TonemapCommon.usf
+++ b/Engine/Shaders/TonemapCommon.usf
@@ -5,7 +5,7 @@
 =============================================================================*/
 
 // #todo-rco: Fix in the cross-compiler
-#if COMPILER_GLSL || COMPILER_GLSL_ES2
+#if COMPILER_GLSL || COMPILER_GLSL_ES2 || METAL_PROFILE
 #define log10(x)  log(x) / log(10.0)
 #endif
 
