Reduce magnitude of PixelPosition.xyz in order to work around normalize bug with large values with some ES2 hardware.

Fixes UE-3417 CRITICAL: Regression: COMPAT: SunTemple: The sun is blown out in the room that overlooks the ocean on the NVidia Shield_4606

diff --git a/Engine/Shaders/MaterialTemplate.usf b/Engine/Shaders/MaterialTemplate.usf
--- a/Engine/Shaders/MaterialTemplate.usf
+++ b/Engine/Shaders/MaterialTemplate.usf
@@ -1540,7 +1540,14 @@ void CalcMaterialParameters(
 	}
 
 	// PixelPosition is the world position translated to the camera position, which is just -CameraVector
+#if COMPILER_GLSL_ES2
+	// ES2 normalize isn't done accurately. This seems to fix it.
+	// Originally this was normalize(normalize(PixelPosition.xyz)) but tegra4 appears to optimize that out.
+	Parameters.CameraVector = normalize(-0.01 * PixelPosition.xyz);
+#else
+	// PixelPosition is the world position translated to the camera position, which is just -CameraVector
 	Parameters.CameraVector = normalize(-PixelPosition.xyz);
+#endif
 
 	Parameters.LightVector = 0;
 
