Fix for Mac. Missing parenthesis.

diff --git a/Engine/Shaders/PostProcessTonemap.usf b/Engine/Shaders/PostProcessTonemap.usf
--- a/Engine/Shaders/PostProcessTonemap.usf
+++ b/Engine/Shaders/PostProcessTonemap.usf
@@ -50,7 +50,7 @@ half4 UnwrappedTexture3DSample( Texture2D Texture, SamplerState Sampler, float3
 	float IntW = floor( UVW.z * Size - 0.5 );
 	half FracW = UVW.z * Size - 0.5 - IntW;
 
-	float U = UVW.x + IntW / Size;
+	float U = ( UVW.x + IntW ) / Size;
 	float V = UVW.y;
 
 	half4 RG0 = Texture.Sample( Sampler, float2(U, V) );
