fix for black screens on iOS
UE-14460. UE-14009
#ios

diff --git a/Engine/Shaders/ShaderVersion.usf b/Engine/Shaders/ShaderVersion.usf
--- a/Engine/Shaders/ShaderVersion.usf
+++ b/Engine/Shaders/ShaderVersion.usf
@@ -2,4 +2,4 @@
 // Each time the console command is executed it generates a new GUID. As this file is included
 // in common.usf (which should be included in any shader) it allows to invalidate the shader DDC.
 // 
-// GUID = 7D700E574A85DE267C385C9E997BCC23  ds
\ No newline at end of file
+// GUID = 7D700E574A85DE267C385C9E997BCC26
\ No newline at end of file
