Fix incorrect interpolant declaration in landscape grass weight shader causing issues on OSX 10.10. (UE-15422)

diff --git a/Engine/Shaders/LandscapeGrassWeight.usf b/Engine/Shaders/LandscapeGrassWeight.usf
--- a/Engine/Shaders/LandscapeGrassWeight.usf
+++ b/Engine/Shaders/LandscapeGrassWeight.usf
@@ -54,16 +54,15 @@ int OutputPass;
 
 // Pixel Shader
 void PSMain( 
-	FVertexFactoryInterpolantsVSToPS Interpolants,
-	FLandscapeGrassWeightInterpolantsVSToPS GrassWeightInterpolants,
+	FLandscapeGrassWeightInterpolantsVSToPS Interpolants,
 	in float4 SvPosition : SV_Position,
 	OPTIONAL_IsFrontFace,
 	out half4 OutColor	: SV_Target0
 	)
 {
-	FMaterialPixelParameters MaterialParameters = GetMaterialPixelParameters(Interpolants, GrassWeightInterpolants.PixelPosition);
+	FMaterialPixelParameters MaterialParameters = GetMaterialPixelParameters(Interpolants.FactoryInterpolants, Interpolants.PixelPosition);
 
-	CalcMaterialParameters(MaterialParameters, SvPosition, bIsFrontFace, GrassWeightInterpolants.PixelPosition); 
+	CalcMaterialParameters(MaterialParameters, SvPosition, bIsFrontFace, Interpolants.PixelPosition); 
 
 #if NUM_MATERIAL_OUTPUTS_GETGRASSWEIGHT > 2
 	switch(OutputPass)
@@ -71,7 +70,7 @@ void PSMain(
 	case 0:
 #endif
 		// Height, outputs 1 & 2
-		OutColor.xy = GrassWeightInterpolants.Height;
+		OutColor.xy = Interpolants.Height;
 		#if NUM_MATERIAL_OUTPUTS_GETGRASSWEIGHT > 0
 			OutColor.z = GetGrassWeight0(MaterialParameters);
 			#if NUM_MATERIAL_OUTPUTS_GETGRASSWEIGHT > 1
