Fixes grading LUT color tint bug

diff --git a/Engine/Shaders/PostProcessCombineLUTs.usf b/Engine/Shaders/PostProcessCombineLUTs.usf
--- a/Engine/Shaders/PostProcessCombineLUTs.usf
+++ b/Engine/Shaders/PostProcessCombineLUTs.usf
@@ -212,31 +212,25 @@ void MainPS(float4 InUV : TEXCOORD0, out float4 OutColor : SV_Target0)
 		// FIXME LUTs are in sRGB space
 		half3 GammaColor = LinearToSrgbBranching( saturate( FilmColor ) );
 
-		// Transform color to 256x16 UV
-		float Scale = 15.0f / 16.0f;
-
-		float IntB = floor(GammaColor.b * 15.9999f) / 16.0f;
-
-		float2 UV = float2(GammaColor.r * Scale / 16.0f + IntB, GammaColor.g * Scale);
-		UV += float2(0.5f / 256.0f, 0.5f / 16.0f);
+		float3 UVW = GammaColor * (15.0 / 16.0) + (0.5f / 16.0);
 
 		GammaColor = LUTWeights[0] * GammaColor;
 
 		// BLENDCOUNT is the number of LUT that are blended together including the neutral one
 	#if BLENDCOUNT >= 2
-		GammaColor += LUTWeights[1] * Texture2DSample(Texture1, Texture1Sampler, UV).rgb;
+		GammaColor += LUTWeights[1] * UnwrappedTexture3DSample( Texture1, Texture1Sampler, UVW, 16 ).rgb;
 	#endif
 
 	#if BLENDCOUNT >= 3
-		GammaColor += LUTWeights[2] * Texture2DSample(Texture2, Texture2Sampler, UV).rgb;
+		GammaColor += LUTWeights[2] * UnwrappedTexture3DSample( Texture2, Texture2Sampler, UVW, 16 ).rgb;
 	#endif
 
 	#if BLENDCOUNT >= 4
-		GammaColor += LUTWeights[3] * Texture2DSample(Texture3, Texture3Sampler, UV).rgb;
+		GammaColor += LUTWeights[3] * UnwrappedTexture3DSample( Texture3, Texture3Sampler, UVW, 16 ).rgb;
 	#endif
 
 	#if BLENDCOUNT >= 5
-		GammaColor += LUTWeights[4] * Texture2DSample(Texture4, Texture4Sampler, UV).rgb;
+		GammaColor += LUTWeights[4] * UnwrappedTexture3DSample( Texture4, Texture4Sampler, UVW, 16 ).rgb;
 	#endif
 
 		// Back to linear space
diff --git a/Engine/Shaders/PostProcessCommon.usf b/Engine/Shaders/PostProcessCommon.usf
--- a/Engine/Shaders/PostProcessCommon.usf
+++ b/Engine/Shaders/PostProcessCommon.usf
@@ -175,4 +175,20 @@ float4 GatherSceneDepth(float2 UV, float2 InvBufferSize)
 #endif // FEATURE_LEVEL
 }
 
+half4 UnwrappedTexture3DSample( Texture2D Texture, SamplerState Sampler, float3 UVW, float Size )
+{
+	// a volume texture 16x16x16 would be unwrapped to a 2d texture 256x16
+
+	float IntW = floor( UVW.z * Size - 0.5 );
+	half FracW = UVW.z * Size - 0.5 - IntW;
+
+	float U = ( UVW.x + IntW ) / Size;
+	float V = UVW.y;
+
+	half4 RG0 = Texture.Sample( Sampler, float2(U, V) );
+	half4 RG1 = Texture.Sample( Sampler, float2(U + 1.0f / Size, V) );
+
+	return lerp(RG0, RG1, FracW);
+}
+
 #endif // __POST_PROCESS_COMMON__
diff --git a/Engine/Shaders/PostProcessTonemap.usf b/Engine/Shaders/PostProcessTonemap.usf
--- a/Engine/Shaders/PostProcessTonemap.usf
+++ b/Engine/Shaders/PostProcessTonemap.usf
@@ -43,22 +43,6 @@ half GrainFromUV(float2 GrainUV)
 
 static const float LUTSize = 32;
 
-half4 UnwrappedTexture3DSample( Texture2D Texture, SamplerState Sampler, float3 UVW, float Size )
-{
-	// a volume texture 16x16x16 would be unwrapped to a 2d texture 256x16
-
-	float IntW = floor( UVW.z * Size - 0.5 );
-	half FracW = UVW.z * Size - 0.5 - IntW;
-
-	float U = ( UVW.x + IntW ) / Size;
-	float V = UVW.y;
-
-	half4 RG0 = Texture.Sample( Sampler, float2(U, V) );
-	half4 RG1 = Texture.Sample( Sampler, float2(U + 1.0f / Size, V) );
-
-	return lerp(RG0, RG1, FracW);
-}
-
 half3 ColorLookupTable( half3 LinearColor )
 {
 	float LinearRange = 14;
