fixed D3D compiler error happening when using "VisualizeTexture"

diff --git a/Engine/Shaders/VisualizeTexture.usf b/Engine/Shaders/VisualizeTexture.usf
--- a/Engine/Shaders/VisualizeTexture.usf
+++ b/Engine/Shaders/VisualizeTexture.usf
@@ -66,16 +66,18 @@ void VisualizeTexturePS(
 	float MipLevel = VisualizeParam[1].a;
 	float ArrayIndex = VisualizeParam[1].b;
 
+
+
 #if TEXTURE_TYPE == 0
 	float4 TexLookup = TextureCubeSampleLevel(VisualizeTextureCube, VisualizeTextureCubeSampler, SpherePos, MipLevel);
-#elif TEXTURE_TYPE == 2
+#elif TEXTURE_TYPE == 2 || TEXTURE_TYPE == 6
 	float4 TexLookup = Texture2DSampleLevel(VisualizeTexture2D, VisualizeTexture2DSampler, UV, MipLevel);
 #elif TEXTURE_TYPE == 3
 	#if FEATURE_LEVEL < FEATURE_LEVEL_SM4
 		float4 TexLookup = float4(1.0f, 0.0f, 1.0f, 1.0f);
 	#else
 		float4 TexLookup = Texture3DSampleLevel(VisualizeTexture3D, VisualizeTexture3DSampler, VolumePos, MipLevel);
-#endif
+	#endif
 #elif TEXTURE_TYPE == 4
 	#if FEATURE_LEVEL < FEATURE_LEVEL_SM5
 		float4 TexLookup = float4(1.0f, 0.0f, 1.0f, 1.0f);
@@ -89,6 +91,7 @@ void VisualizeTexturePS(
 		TexLookup = VisualizeTexture2DMS.Load(int2(UV * TextureExtent.xy), 0);
 	#endif
 #endif
+	
 
 	int InputValueMapping = (int)VisualizeParam[2].r;
 	if(InputValueMapping == 1)
@@ -109,14 +112,18 @@ void VisualizeTexturePS(
 		float Value = Depth * Depth;
 		TexLookup = float4(pow(Value, 0.25f), pow(Value, 1), pow(Value, 4), 0);		
 
+#if TEXTURE_TYPE == 6
+		// inside #ifdef to avoid D3DDebug error:
+		//    The resource return type for component 0 declared in the shader code (UINT) is not compatible with the Shader Resource View format bound to slot 1 of the Pixel Shader unit (UNORM).  This mismatch is invalid if the shader actually uses the view (e.g. it is not skipped due to shader code branching). [ EXECUTION ERROR #361: DEVICE_DRAW_RESOURCE_RETURN_TYPE_MISMATCH]
 		float fStencil = VisualizeDepthStencilTexture.Load(int3(UV * TextureExtent.xy, 0)).g;
 		fStencil /= 255.0f;
-
 		TexLookup.a = fStencil;
+#endif
+
 #else
 		TexLookup.r = Depth;
 #endif
-	}	
+	}
 
 	float4 RGBAMask = 0;
 	int SingleChannel = (int)VisualizeParam[2].z;
diff --git a/Engine/Source/Runtime/Renderer/Private/PostProcess/VisualizeTexture.cpp b/Engine/Source/Runtime/Renderer/Private/PostProcess/VisualizeTexture.cpp
--- a/Engine/Source/Runtime/Renderer/Private/PostProcess/VisualizeTexture.cpp
+++ b/Engine/Source/Runtime/Renderer/Private/PostProcess/VisualizeTexture.cpp
@@ -15,7 +15,7 @@
 
 
 /** A pixel shader which filters a texture. */
-// @param TextureType 0:Cube, 1:1D(not yet supported), 2:2D, 3:3D, 4:Cube[], 5:2D MSAA
+// @param TextureType 0:Cube, 1:1D(not yet supported), 2:2D no MSAA, 3:3D, 4:Cube[], 5:2D MSAA, 6:2D DepthStencil no MSAA (needed to avoid D3DDebug error)
 template<uint32 TextureType>
 class VisualizeTexturePS : public FGlobalShader
 {
@@ -141,7 +141,7 @@ protected:
 #define VARIATION1(A) typedef VisualizeTexturePS<A> VisualizeTexturePS##A; \
 	IMPLEMENT_SHADER_TYPE2(VisualizeTexturePS##A, SF_Pixel);
 
-VARIATION1(0)			VARIATION1(2)			VARIATION1(3)			VARIATION1(4)			VARIATION1(5)
+VARIATION1(0)			VARIATION1(2)			VARIATION1(3)			VARIATION1(4)			VARIATION1(5)			VARIATION1(6)
 #undef VARIATION1
 
 
@@ -234,8 +234,16 @@ void RenderVisualizeTexture(FRHICommandListImmediate& RHICmdList, ERHIFeatureLev
 		}
 		else
 		{
-			// non MSAA
-			VisualizeTextureForTextureType<2>(RHICmdList, FeatureLevel, Data);
+			if(Data.Desc.Format == PF_DepthStencil)
+			{
+				// DepthStencil non MSAA (needed to avoid D3DDebug error)
+				VisualizeTextureForTextureType<6>(RHICmdList, FeatureLevel, Data);
+			}
+			else
+			{
+				// non MSAA
+				VisualizeTextureForTextureType<2>(RHICmdList, FeatureLevel, Data);
+			}
 		}
 	}
 	else if(Data.Desc.Is3DTexture())
