PS4 - Stop registers spilling to scratch buffer.

diff --git a/Engine/Shaders/LPVPropagate.usf b/Engine/Shaders/LPVPropagate.usf
--- a/Engine/Shaders/LPVPropagate.usf
+++ b/Engine/Shaders/LPVPropagate.usf
@@ -210,6 +210,9 @@ void GetGeometryVolumeCellWeights( int3 gridPos,
 	{
 		const int destIndex[8] = { 0, 17, 6, 23, 2, 19, 8, 25 };
 
+#if PS4_PROFILE
+		[unroll]
+#endif
 		for ( int i=0; i<8; i++ )
 		{
 #if NORMALS_LOOKUP
@@ -243,6 +246,9 @@ void GetGeometryVolumeCellWeights( int3 gridPos,
 		const int destIndex[12]      = {  1,	 3,	    5,	   7,	  9,	 11,	14,    16,    18,    20,    22,    24 };
 		const int sourceIndex[12][2] = { {0,4}, {0,2}, {4,6}, {2,6}, {0,1}, {4,5}, {2,3}, {6,7}, {1,5}, {1,3}, {5,7}, {3,7} };
 
+#if PS4_PROFILE
+		[unroll]
+#endif
 		for ( int i=0; i<12; i++ )
 		{
 #if NORMALS_LOOKUP
@@ -277,6 +283,9 @@ void GetGeometryVolumeCellWeights( int3 gridPos,
 		const int destIndex[6]      = {  4,         10,        12,        13,        15,        21 };
 		const int sourceIndex[6][4] = { {0,2,4,6}, {0,1,4,5}, {0,1,2,3}, {4,5,6,7}, {2,3,6,7}, {1,3,5,7} };
 
+#if PS4_PROFILE
+		[unroll]
+#endif
 		for ( int i=0; i<6; i++ )
 		{
 #if NORMALS_LOOKUP
