Removed 360 only IFALL / IFANY

diff --git a/Engine/Shaders/Definitions.usf b/Engine/Shaders/Definitions.usf
--- a/Engine/Shaders/Definitions.usf
+++ b/Engine/Shaders/Definitions.usf
@@ -163,10 +163,6 @@
 #define BRANCH											[branch]
 /** Performs branching by using the cnd instructions. */
 #define FLATTEN											[flatten]
-/** Executes the conditional part of an if statement when the condition is true for all threads on which the current shader is running. */
-#define IFALL											[ifAll]
-/** Executes the conditional part of an if statement when the condition is true for any thread on which the current shader is running. */
-#define IFANY											[ifAny]
 /** Allows a compute shader loop termination condition to be based off of a UAV read. The loop must not contain synchronization intrinsics. */
 #define ALLOW_UAV_CONDITION								[allow_uav_condition]
 #else
@@ -182,12 +178,6 @@
 #ifndef FLATTEN
 #define FLATTEN
 #endif
-#ifndef IFALL
-#define IFALL
-#endif
-#ifndef IFANY
-#define IFANY
-#endif
 #endif
 
 #define ISOLATE
