Global distance field functions are SM5

diff --git a/Engine/Shaders/GlobalDistanceFieldShared.usf b/Engine/Shaders/GlobalDistanceFieldShared.usf
--- a/Engine/Shaders/GlobalDistanceFieldShared.usf
+++ b/Engine/Shaders/GlobalDistanceFieldShared.usf
@@ -4,6 +4,8 @@
 	GlobalDistanceFieldShared.usf
 =============================================================================*/
 
+#if FEATURE_LEVEL >= FEATURE_LEVEL_SM5
+
 // Must match C++
 #define MAX_GLOBAL_DF_CLIPMAPS 4
 
@@ -123,4 +125,6 @@ float3 GetDistanceFieldGradientGlobal(float3 WorldPosition)
 	}
 
 	return Gradient;
-}
\ No newline at end of file
+}
+
+#endif
\ No newline at end of file
