Support compute shader object tile cull for DFAO with 'r.AOScatterTileCulling 0'
* Performance is currently on par with the rasterizer version on 980

diff --git a/Engine/Shaders/DistanceFieldAOShared.usf b/Engine/Shaders/DistanceFieldAOShared.usf
--- a/Engine/Shaders/DistanceFieldAOShared.usf
+++ b/Engine/Shaders/DistanceFieldAOShared.usf
@@ -38,10 +38,16 @@ float GetAOMaxDistance()
 	return max(AOObjectMaxDistance, AOGlobalMaxOcclusionDistance);
 }
 
-#define NUM_CULLED_OBJECT_LISTS 3
-
 #define USE_DEPTH_RANGE_LISTS 1
 
+#if USE_DEPTH_RANGE_LISTS
+	// One for near, one for far
+	#define NUM_CULLED_OBJECT_LISTS 2
+#else
+	// Objects are culled into one of 3 lists based on depth from the tile
+	#define NUM_CULLED_OBJECT_LISTS 3
+#endif
+
 void GetPhaseParameters(uint Phase, out uint StartIndex, out uint EndIndex)
 {
 	if (Phase == 0)
@@ -67,24 +73,18 @@ float GetStepOffset(float StepIndex)
 }
 
 Buffer<uint> TileHeadDataUnpacked;
-Buffer<uint4> TileHeadData;
 Buffer<uint> TileArrayData;
 uint2 TileListGroupSize;
-
-#define UNPACKED_TILE_HEAD 1
  
 uint4 GetTileHead(uint2 TileCoordinate) 
 { 
 	uint TileIndex = TileCoordinate.y * TileListGroupSize.x + TileCoordinate.x;
-#if UNPACKED_TILE_HEAD
+
 	return uint4(
 		TileHeadDataUnpacked[TileIndex * 4 + 0], 
 		min(TileHeadDataUnpacked[TileIndex * 4 + 1], (uint)MAX_OBJECTS_PER_TILE),
 		min(TileHeadDataUnpacked[TileIndex * 4 + 2], (uint)MAX_OBJECTS_PER_TILE),
 		min(TileHeadDataUnpacked[TileIndex * 4 + 3], (uint)MAX_OBJECTS_PER_TILE));
-#else
-	return TileHeadData[TileIndex];
-#endif
 }
 
 Texture2D DistanceFieldNormalTexture;
@@ -222,4 +222,81 @@ float2 GetScreenUVFromScreenGrid(uint2 OutputCoordinate)
 {
 	float2 ScreenUV = ((OutputCoordinate * TRACE_DOWNSAMPLE_FACTOR + JitterOffset) * AO_DOWNSAMPLE_FACTOR + View.ViewRectMin.xy + float2(.5f, .5f)) * View.ViewSizeAndSceneTexelSize.zw;
 	return ScreenUV;
+}
+
+RWBuffer<uint> RWTileHeadDataUnpacked;
+RWBuffer<uint> RWTileArrayData;
+
+/** Used for object <-> tile culling */
+void IntersectObjectWithConeDepthRange(
+	float3 TileConeVertex, 
+	float3 TileConeAxis, 
+	float TileConeAngleCos, 
+	float TileConeAngleSin, 
+	float2 ConeDepthRange, 
+	float2 ConeAxisDistanceMinMax, 
+	uint TileIndex, 
+	uint ListIndex, 
+	uint ObjectIndex,
+	uint TotalNumGroups)
+{
+	BRANCH
+	if (ConeAxisDistanceMinMax.x > ConeDepthRange.x && ConeAxisDistanceMinMax.y < ConeDepthRange.y)
+	{
+#define USE_DISTANCE_FIELD_FOR_OBJECT_CULLING 1
+#if USE_DISTANCE_FIELD_FOR_OBJECT_CULLING
+		// Use the position halfway between the depth ranges as the center for the bounding sphere of this tile depth range
+		float3 ViewTileBoundingSphereCenter = TileConeVertex + TileConeAxis * (.5f * (ConeDepthRange.x + ConeDepthRange.y));
+		float3 WorldTileBoundingSphereCenter = mul(float4(ViewTileBoundingSphereCenter.xyz, 1), View.ViewToTranslatedWorld).xyz - View.PreViewTranslation.xyz;
+		float DistanceAlongAxis = .5f * (ConeDepthRange.x - ConeDepthRange.y);
+		float FarDepthDistanceToEdgeOfCone = ConeDepthRange.y * TileConeAngleSin / TileConeAngleCos;
+		float TileBoundingSphereRadius = sqrt(DistanceAlongAxis * DistanceAlongAxis + FarDepthDistanceToEdgeOfCone * FarDepthDistanceToEdgeOfCone);
+
+		float3 LocalPositionExtent = LoadObjectLocalPositionExtent(ObjectIndex);
+		float4x4 WorldToVolume = LoadObjectWorldToVolume(ObjectIndex);
+		bool bGeneratedAsTwoSided;
+		float4 UVScaleAndVolumeScale = LoadObjectUVScale(ObjectIndex, bGeneratedAsTwoSided);
+		float3 VolumeTileBoundingSphereCenter = mul(float4(WorldTileBoundingSphereCenter, 1), WorldToVolume).xyz;
+		float BoxDistance = ComputeDistanceFromBoxToPoint(-LocalPositionExtent, LocalPositionExtent, VolumeTileBoundingSphereCenter) * UVScaleAndVolumeScale.w;
+
+		BRANCH
+		if (BoxDistance < TileBoundingSphereRadius + AOObjectMaxDistance)
+		{
+			float3 UVAdd = LoadObjectUVAdd(ObjectIndex);
+
+			float3 ClampedSamplePosition = clamp(VolumeTileBoundingSphereCenter, -LocalPositionExtent, LocalPositionExtent);
+			float DistanceToClamped = length(VolumeTileBoundingSphereCenter - ClampedSamplePosition);
+
+			float3 StepVolumeUV = DistanceFieldVolumePositionToUV(ClampedSamplePosition, UVScaleAndVolumeScale.xyz, UVAdd);
+			float DistanceToOccluder = (Texture3DSampleLevel(DistanceFieldTexture, DistanceFieldSampler, StepVolumeUV, 0).x + DistanceToClamped) * UVScaleAndVolumeScale.w;
+
+			BRANCH
+			if (DistanceToOccluder < TileBoundingSphereRadius + AOObjectMaxDistance)
+			{
+				uint ArrayIndex;
+				InterlockedAdd(RWTileHeadDataUnpacked[TileIndex * 4 + 1 + ListIndex], 1U, ArrayIndex);
+
+				if (ArrayIndex < MAX_OBJECTS_PER_TILE)
+				{
+					// Note: indexing so that threads are writing to RWTileArrayData coherently, has a huge impact on speed, even though the array data for one record is no longer dense
+					uint DataIndex = (ArrayIndex * TotalNumGroups + TileIndex) * NUM_CULLED_OBJECT_LISTS + ListIndex;
+
+					RWTileArrayData[DataIndex] = ObjectIndex;
+				}
+			}
+		}
+#else
+		uint ListIndex = 0;
+		uint ArrayIndex;
+		InterlockedAdd(RWTileHeadDataUnpacked[TileIndex * 4 + 1 + ListIndex], 1U, ArrayIndex);
+
+		if (ArrayIndex < MAX_OBJECTS_PER_TILE)
+		{
+			// Note: indexing so that threads are writing to RWTileArrayData coherently, has a huge impact on speed, even though the array data for one record is no longer dense
+			uint DataIndex = (ArrayIndex * TotalNumGroups + TileIndex) * NUM_CULLED_OBJECT_LISTS + ListIndex;
+
+			RWTileArrayData[DataIndex] = ObjectIndex;
+		}
+#endif
+	}
 }
\ No newline at end of file
diff --git a/Engine/Shaders/DistanceFieldSurfaceCacheLighting.usf b/Engine/Shaders/DistanceFieldSurfaceCacheLighting.usf
--- a/Engine/Shaders/DistanceFieldSurfaceCacheLighting.usf
+++ b/Engine/Shaders/DistanceFieldSurfaceCacheLighting.usf
@@ -36,71 +36,11 @@ void ObjectCullVS(
 	Output.ObjectIndex = ObjectIndex;
 } 
 
-RWBuffer<uint> RWTileHeadDataUnpacked;
-RWBuffer<uint> RWTileArrayData;
-
 Buffer<float4> TileConeAxisAndCos;
 Buffer<float4> TileConeDepthRanges;
 
 float2 NumGroups;
 
-void IntersectObjectWithConeDepthRange(
-	float3 TileConeVertex, 
-	float3 TileConeAxis, 
-	float TileConeAngleCos, 
-	float TileConeAngleSin, 
-	float2 ConeDepthRange, 
-	float2 ConeAxisDistanceMinMax, 
-	uint TileIndex, 
-	uint ListIndex, 
-	uint ObjectIndex)
-{
-	BRANCH
-	if (ConeAxisDistanceMinMax.x > ConeDepthRange.x && ConeAxisDistanceMinMax.y < ConeDepthRange.y)
-	{
-		// Use the position halfway between the depth ranges as the center for the bounding sphere of this tile depth range
-		float3 ViewTileBoundingSphereCenter = TileConeVertex + TileConeAxis * (.5f * (ConeDepthRange.x + ConeDepthRange.y));
-		float3 WorldTileBoundingSphereCenter = mul(float4(ViewTileBoundingSphereCenter.xyz, 1), View.ViewToTranslatedWorld).xyz - View.PreViewTranslation.xyz;
-		float DistanceAlongAxis = .5f * (ConeDepthRange.x - ConeDepthRange.y);
-		float FarDepthDistanceToEdgeOfCone = ConeDepthRange.y * TileConeAngleSin / TileConeAngleCos;
-		float TileBoundingSphereRadius = sqrt(DistanceAlongAxis * DistanceAlongAxis + FarDepthDistanceToEdgeOfCone * FarDepthDistanceToEdgeOfCone);
-
-		float3 LocalPositionExtent = LoadObjectLocalPositionExtent(ObjectIndex);
-		float4x4 WorldToVolume = LoadObjectWorldToVolume(ObjectIndex);
-		bool bGeneratedAsTwoSided;
-		float4 UVScaleAndVolumeScale = LoadObjectUVScale(ObjectIndex, bGeneratedAsTwoSided);
-		float3 VolumeTileBoundingSphereCenter = mul(float4(WorldTileBoundingSphereCenter, 1), WorldToVolume).xyz;
-		float BoxDistance = ComputeDistanceFromBoxToPoint(-LocalPositionExtent, LocalPositionExtent, VolumeTileBoundingSphereCenter) * UVScaleAndVolumeScale.w;
-
-		BRANCH
-		if (BoxDistance < TileBoundingSphereRadius + AOObjectMaxDistance)
-		{
-			float3 UVAdd = LoadObjectUVAdd(ObjectIndex);
-
-			float3 ClampedSamplePosition = clamp(VolumeTileBoundingSphereCenter, -LocalPositionExtent, LocalPositionExtent);
-			float DistanceToClamped = length(VolumeTileBoundingSphereCenter - ClampedSamplePosition);
-
-			float3 StepVolumeUV = DistanceFieldVolumePositionToUV(ClampedSamplePosition, UVScaleAndVolumeScale.xyz, UVAdd);
-			float DistanceToOccluder = (Texture3DSampleLevel(DistanceFieldTexture, DistanceFieldSampler, StepVolumeUV, 0).x + DistanceToClamped) * UVScaleAndVolumeScale.w;
-
-			BRANCH
-			if (DistanceToOccluder < TileBoundingSphereRadius + AOObjectMaxDistance)
-			{
-				uint ArrayIndex;
-				InterlockedAdd(RWTileHeadDataUnpacked[TileIndex * 4 + 1 + ListIndex], 1U, ArrayIndex);
-
-				if (ArrayIndex < MAX_OBJECTS_PER_TILE)
-				{
-					// Note: indexing so that threads are writing to RWTileArrayData coherently, has a huge impact on speed, even though the array data for one record is no longer dense
-					uint DataIndex = (ArrayIndex * (uint)(NumGroups.x * NumGroups.y + .5f) + TileIndex) * NUM_CULLED_OBJECT_LISTS + ListIndex;
-
-					RWTileArrayData[DataIndex] = ObjectIndex;
-				}
-			}
-		}
-	}		
-}
-
 /** Intersects a single object with the tile and adds to the intersection list if needed. */
 void ObjectCullPS(
 	FObjectCullVertexOutput Input, 
@@ -133,38 +73,9 @@ void ObjectCullPS(
 		float ConeAxisDistance = dot(SphereCenterAndRadius.xyz - TileConeVertex, TileConeAxis);
 		float2 ConeAxisDistanceMinMax = float2(ConeAxisDistance + SphereCenterAndRadius.w, ConeAxisDistance - SphereCenterAndRadius.w);
 
-		IntersectObjectWithConeDepthRange(TileConeVertex, TileConeAxis, TileConeAngleCos, TileConeAngleSin, ConeAxisDepthRanges.xy, ConeAxisDistanceMinMax, TileIndex, 0, Input.ObjectIndex);
-		IntersectObjectWithConeDepthRange(TileConeVertex, TileConeAxis, TileConeAngleCos, TileConeAngleSin, ConeAxisDepthRanges.zw, ConeAxisDistanceMinMax, TileIndex, 1, Input.ObjectIndex);
-		/*
-		if (ConeAxisDistanceMinMax.x > ConeAxisDepthRanges.x && ConeAxisDistanceMinMax.y < ConeAxisDepthRanges.y)
-		{
-			uint ListIndex = 0;
-			uint ArrayIndex;
-			InterlockedAdd(RWTileHeadDataUnpacked[TileIndex * 4 + 1 + ListIndex], 1U, ArrayIndex);
-
-			if (ArrayIndex < MAX_OBJECTS_PER_TILE)
-			{
-				// Note: indexing so that threads are writing to RWTileArrayData coherently, has a huge impact on speed, even though the array data for one record is no longer dense
-				uint DataIndex = (ArrayIndex * (uint)(NumGroups.x * NumGroups.y + .5f) + TileIndex) * NUM_CULLED_OBJECT_LISTS + ListIndex;
-
-				RWTileArrayData[DataIndex] = Input.ObjectIndex;
-			}
-		}
-
-		if (ConeAxisDistanceMinMax.x > ConeAxisDepthRanges.z && ConeAxisDistanceMinMax.y < ConeAxisDepthRanges.w)
-		{
-			uint ListIndex = 1;
-			uint ArrayIndex;
-			InterlockedAdd(RWTileHeadDataUnpacked[TileIndex * 4 + 1 + ListIndex], 1U, ArrayIndex);
-
-			if (ArrayIndex < MAX_OBJECTS_PER_TILE)
-			{
-				// Note: indexing so that threads are writing to RWTileArrayData coherently, has a huge impact on speed, even though the array data for one record is no longer dense
-				uint DataIndex = (ArrayIndex * (uint)(NumGroups.x * NumGroups.y + .5f) + TileIndex) * NUM_CULLED_OBJECT_LISTS + ListIndex;
-
-				RWTileArrayData[DataIndex] = Input.ObjectIndex;
-			}
-		}*/
+		uint TotalNumGroups = (uint)(NumGroups.x * NumGroups.y + .5f);
+		IntersectObjectWithConeDepthRange(TileConeVertex, TileConeAxis, TileConeAngleCos, TileConeAngleSin, ConeAxisDepthRanges.xy, ConeAxisDistanceMinMax, TileIndex, 0, Input.ObjectIndex, TotalNumGroups);
+		IntersectObjectWithConeDepthRange(TileConeVertex, TileConeAxis, TileConeAngleCos, TileConeAngleSin, ConeAxisDepthRanges.zw, ConeAxisDistanceMinMax, TileIndex, 1, Input.ObjectIndex, TotalNumGroups);
 	}
 
 #else
diff --git a/Engine/Shaders/DistanceFieldSurfaceCacheLightingCompute.usf b/Engine/Shaders/DistanceFieldSurfaceCacheLightingCompute.usf
--- a/Engine/Shaders/DistanceFieldSurfaceCacheLightingCompute.usf
+++ b/Engine/Shaders/DistanceFieldSurfaceCacheLightingCompute.usf
@@ -259,7 +259,6 @@ float2 NumGroups;
 
 RWBuffer<float4> RWTileConeAxisAndCos;
 RWBuffer<float4> RWTileConeDepthRanges;
-RWBuffer<uint> RWTileHeadDataUnpacked;
 
 /** Builds tile depth ranges and bounding cones. */
 [numthreads(THREADGROUP_SIZEX, THREADGROUP_SIZEY, 1)]
@@ -378,9 +377,7 @@ groupshared uint MediumTileNumObjects;
 groupshared uint LargeTileNumObjects;
 groupshared uint TileArrayDataStart;
 
-RWBuffer<uint> RWTileArrayData;
 RWBuffer<uint> RWTileArrayNextAllocation;
-RWBuffer<uint4> RWTileHeadData;
 
 [numthreads(THREADGROUP_SIZEX, THREADGROUP_SIZEY, 1)]
 void DistanceFieldAOBuildTileListMain(
@@ -472,6 +469,39 @@ void DistanceFieldAOBuildTileListMain(
 		TileConeVertex = float3(0, 0, 0) - TileConeAxis * VertexPullbackLength;
 	}
 	
+#if USE_DEPTH_RANGE_LISTS
+
+	// A value of 1 is conservative, but has a huge impact on performance
+	float RadiusScale = .5f;
+	uint TileIndex = GroupId.y * (uint)NumGroups.x + GroupId.x;
+	uint NumCulledObjects = GetCulledNumObjects();
+
+	// Compute per-tile lists of affecting objects through bounds culling
+	// Each thread now operates on a sample instead of a pixel
+	LOOP
+	for (uint ObjectIndex = ThreadIndex; ObjectIndex < NumCulledObjects; ObjectIndex += THREADGROUP_TOTALSIZE)
+	{
+		float4 WorldSphereCenterAndRadius = LoadObjectPositionAndRadius(ObjectIndex);
+		float3 ViewSpaceSphereCenter = mul(float4(WorldSphereCenterAndRadius.xyz + View.PreViewTranslation.xyz, 1), View.TranslatedWorldToView).xyz;
+		float4 SphereCenterAndRadius = float4(ViewSpaceSphereCenter, WorldSphereCenterAndRadius.w + RadiusScale * AOObjectMaxDistance);
+
+		if (SphereIntersectCone(SphereCenterAndRadius, TileConeVertex, TileConeAxis, TileConeAngleCos, TileConeAngleSin))
+		{
+			float ConeAxisDistance = dot(SphereCenterAndRadius.xyz - TileConeVertex, TileConeAxis);
+			float2 ConeAxisDistanceMinMax = float2(ConeAxisDistance + SphereCenterAndRadius.w, ConeAxisDistance - SphereCenterAndRadius.w);
+
+			uint TotalNumGroups = (uint)(NumGroups.x * NumGroups.y + .5f);
+			IntersectObjectWithConeDepthRange(TileConeVertex, TileConeAxis, TileConeAngleCos, TileConeAngleSin, ConeAxisDepthRanges.xy, ConeAxisDistanceMinMax, TileIndex, 0, ObjectIndex, TotalNumGroups);
+			IntersectObjectWithConeDepthRange(TileConeVertex, TileConeAxis, TileConeAngleCos, TileConeAngleSin, ConeAxisDepthRanges.zw, ConeAxisDistanceMinMax, TileIndex, 1, ObjectIndex, TotalNumGroups);
+		}
+	}
+
+	if (ThreadIndex == 0)
+	{
+		RWTileHeadDataUnpacked[TileIndex * 4 + 0] = TileIndex;
+	}
+
+#else
 	// A value of 1 is conservative, but has huge impact on performance
 	float RadiusScale = .5f;
 	float SmallGroupMaxSampleRadius;
@@ -573,6 +603,7 @@ void DistanceFieldAOBuildTileListMain(
 	{
 		RWTileArrayData[ArrayDataStart + LargeListIndex] = LargeTileObjectIndices[LargeListIndex];
 	}
+#endif
 }
 
 /** View rect min in xy, max in zw. */
diff --git a/Engine/Source/Runtime/Renderer/Private/DistanceFieldGlobalIllumination.cpp b/Engine/Source/Runtime/Renderer/Private/DistanceFieldGlobalIllumination.cpp
--- a/Engine/Source/Runtime/Renderer/Private/DistanceFieldGlobalIllumination.cpp
+++ b/Engine/Source/Runtime/Renderer/Private/DistanceFieldGlobalIllumination.cpp
@@ -1505,7 +1505,6 @@ public:
 		AOParameters.Bind(Initializer.ParameterMap);
 		ScreenGridParameters.Bind(Initializer.ParameterMap);
 		SurfelParameters.Bind(Initializer.ParameterMap);
-		TileHeadData.Bind(Initializer.ParameterMap, TEXT("TileHeadData"));
 		TileHeadDataUnpacked.Bind(Initializer.ParameterMap, TEXT("TileHeadDataUnpacked"));
 		TileArrayData.Bind(Initializer.ParameterMap, TEXT("TileArrayData"));
 		TileConeDepthRanges.Bind(Initializer.ParameterMap, TEXT("TileConeDepthRanges"));
@@ -1540,7 +1539,6 @@ public:
 		FTileIntersectionResources* TileIntersectionResources = View.ViewState->AOTileIntersectionResources;
 
 		SetSRVParameter(RHICmdList, ShaderRHI, TileHeadDataUnpacked, TileIntersectionResources->TileHeadDataUnpacked.SRV);
-		SetSRVParameter(RHICmdList, ShaderRHI, TileHeadData, TileIntersectionResources->TileHeadData.SRV);
 		SetSRVParameter(RHICmdList, ShaderRHI, TileArrayData, TileIntersectionResources->TileArrayData.SRV);
 		SetSRVParameter(RHICmdList, ShaderRHI, TileConeDepthRanges, TileIntersectionResources->TileConeDepthRanges.SRV);
 		SetShaderValue(RHICmdList, ShaderRHI, TileListGroupSize, TileIntersectionResources->TileDimensions);
@@ -1567,7 +1565,6 @@ public:
 		Ar << ScreenGridParameters;
 		Ar << SurfelParameters;
 		Ar << TileHeadDataUnpacked;
-		Ar << TileHeadData;
 		Ar << TileArrayData;
 		Ar << TileConeDepthRanges;
 		Ar << TileListGroupSize;
@@ -1585,7 +1582,6 @@ private:
 	FScreenGridParameters ScreenGridParameters;
 	FSurfelBufferParameters SurfelParameters;
 	FShaderResourceParameter TileHeadDataUnpacked;
-	FShaderResourceParameter TileHeadData;
 	FShaderResourceParameter TileArrayData;
 	FShaderResourceParameter TileConeDepthRanges;
 	FShaderParameter TileListGroupSize;
diff --git a/Engine/Source/Runtime/Renderer/Private/DistanceFieldScreenGridLighting.cpp b/Engine/Source/Runtime/Renderer/Private/DistanceFieldScreenGridLighting.cpp
--- a/Engine/Source/Runtime/Renderer/Private/DistanceFieldScreenGridLighting.cpp
+++ b/Engine/Source/Runtime/Renderer/Private/DistanceFieldScreenGridLighting.cpp
@@ -103,7 +103,6 @@ public:
 		AOParameters.Bind(Initializer.ParameterMap);
 		ScreenGridParameters.Bind(Initializer.ParameterMap);
 		GlobalDistanceFieldParameters.Bind(Initializer.ParameterMap);
-		TileHeadData.Bind(Initializer.ParameterMap, TEXT("TileHeadData"));
 		TileHeadDataUnpacked.Bind(Initializer.ParameterMap, TEXT("TileHeadDataUnpacked"));
 		TileArrayData.Bind(Initializer.ParameterMap, TEXT("TileArrayData"));
 		TileConeDepthRanges.Bind(Initializer.ParameterMap, TEXT("TileConeDepthRanges"));
@@ -153,7 +152,6 @@ public:
 		FTileIntersectionResources* TileIntersectionResources = View.ViewState->AOTileIntersectionResources;
 
 		SetSRVParameter(RHICmdList, ShaderRHI, TileHeadDataUnpacked, TileIntersectionResources->TileHeadDataUnpacked.SRV);
-		SetSRVParameter(RHICmdList, ShaderRHI, TileHeadData, TileIntersectionResources->TileHeadData.SRV);
 		SetSRVParameter(RHICmdList, ShaderRHI, TileArrayData, TileIntersectionResources->TileArrayData.SRV);
 		SetSRVParameter(RHICmdList, ShaderRHI, TileConeDepthRanges, TileIntersectionResources->TileConeDepthRanges.SRV);
 
@@ -197,7 +195,6 @@ public:
 		Ar << ScreenGridParameters;
 		Ar << GlobalDistanceFieldParameters;
 		Ar << TileHeadDataUnpacked;
-		Ar << TileHeadData;
 		Ar << TileArrayData;
 		Ar << TileConeDepthRanges;
 		Ar << TileListGroupSize;
@@ -216,7 +213,6 @@ private:
 	FScreenGridParameters ScreenGridParameters;
 	FGlobalDistanceFieldParameters GlobalDistanceFieldParameters;
 	FShaderResourceParameter TileHeadDataUnpacked;
-	FShaderResourceParameter TileHeadData;
 	FShaderResourceParameter TileArrayData;
 	FShaderResourceParameter TileConeDepthRanges;
 	FShaderParameter TileListGroupSize;
@@ -269,7 +265,6 @@ public:
 		AOParameters.Bind(Initializer.ParameterMap);
 		ScreenGridParameters.Bind(Initializer.ParameterMap);
 		GlobalDistanceFieldParameters.Bind(Initializer.ParameterMap);
-		TileHeadData.Bind(Initializer.ParameterMap, TEXT("TileHeadData"));
 		TileHeadDataUnpacked.Bind(Initializer.ParameterMap, TEXT("TileHeadDataUnpacked"));
 		TileArrayData.Bind(Initializer.ParameterMap, TEXT("TileArrayData"));
 		TileConeDepthRanges.Bind(Initializer.ParameterMap, TEXT("TileConeDepthRanges"));
@@ -314,7 +309,6 @@ public:
 		FTileIntersectionResources* TileIntersectionResources = View.ViewState->AOTileIntersectionResources;
 
 		SetSRVParameter(RHICmdList, ShaderRHI, TileHeadDataUnpacked, TileIntersectionResources->TileHeadDataUnpacked.SRV);
-		SetSRVParameter(RHICmdList, ShaderRHI, TileHeadData, TileIntersectionResources->TileHeadData.SRV);
 		SetSRVParameter(RHICmdList, ShaderRHI, TileArrayData, TileIntersectionResources->TileArrayData.SRV);
 		SetSRVParameter(RHICmdList, ShaderRHI, TileConeDepthRanges, TileIntersectionResources->TileConeDepthRanges.SRV);
 
@@ -352,7 +346,6 @@ public:
 		Ar << ScreenGridParameters;
 		Ar << GlobalDistanceFieldParameters;
 		Ar << TileHeadDataUnpacked;
-		Ar << TileHeadData;
 		Ar << TileArrayData;
 		Ar << TileConeDepthRanges;
 		Ar << TileListGroupSize;
@@ -370,7 +363,6 @@ private:
 	FScreenGridParameters ScreenGridParameters;
 	FGlobalDistanceFieldParameters GlobalDistanceFieldParameters;
 	FShaderResourceParameter TileHeadDataUnpacked;
-	FShaderResourceParameter TileHeadData;
 	FShaderResourceParameter TileArrayData;
 	FShaderResourceParameter TileConeDepthRanges;
 	FShaderParameter TileListGroupSize;
diff --git a/Engine/Source/Runtime/Renderer/Private/DistanceFieldSurfaceCacheLighting.cpp b/Engine/Source/Runtime/Renderer/Private/DistanceFieldSurfaceCacheLighting.cpp
--- a/Engine/Source/Runtime/Renderer/Private/DistanceFieldSurfaceCacheLighting.cpp
+++ b/Engine/Source/Runtime/Renderer/Private/DistanceFieldSurfaceCacheLighting.cpp
@@ -331,7 +331,6 @@ void FTileIntersectionResources::InitDynamicRHI()
 	TileConeDepthRanges.Initialize(sizeof(float)* 4, TileDimensions.X * TileDimensions.Y, PF_A32B32G32R32F, BUF_Static);
 
 	TileHeadDataUnpacked.Initialize(sizeof(uint32), TileDimensions.X * TileDimensions.Y * 4, PF_R32_UINT, BUF_Static);
-	TileHeadData.Initialize(sizeof(uint32)* 4, TileDimensions.X * TileDimensions.Y, PF_R32G32B32A32_UINT, BUF_Static);
 
 	//@todo - handle max exceeded
 	TileArrayData.Initialize(sizeof(uint32), GMaxNumObjectsPerTile * TileDimensions.X * TileDimensions.Y * 3, PF_R32_UINT, BUF_Static);
@@ -748,7 +747,6 @@ public:
 		DeferredParameters.Bind(Initializer.ParameterMap);
 		ObjectParameters.Bind(Initializer.ParameterMap);
 		AOParameters.Bind(Initializer.ParameterMap);
-		TileHeadData.Bind(Initializer.ParameterMap, TEXT("TileHeadData"));
 		TileHeadDataUnpacked.Bind(Initializer.ParameterMap, TEXT("TileHeadDataUnpacked"));
 		TileArrayData.Bind(Initializer.ParameterMap, TEXT("TileArrayData"));
 		TileArrayNextAllocation.Bind(Initializer.ParameterMap, TEXT("TileArrayNextAllocation"));
@@ -770,7 +768,6 @@ public:
 
 		FTileIntersectionResources* TileIntersectionResources = ((FSceneViewState*)View.State)->AOTileIntersectionResources;
 
-		TileHeadData.SetBuffer(RHICmdList, ShaderRHI, TileIntersectionResources->TileHeadData);
 		TileHeadDataUnpacked.SetBuffer(RHICmdList, ShaderRHI, TileIntersectionResources->TileHeadDataUnpacked);
 		TileArrayData.SetBuffer(RHICmdList, ShaderRHI, TileIntersectionResources->TileArrayData);
 		TileArrayNextAllocation.SetBuffer(RHICmdList, ShaderRHI, TileIntersectionResources->TileArrayNextAllocation);
@@ -782,7 +779,6 @@ public:
 
 	void UnsetParameters(FRHICommandList& RHICmdList)
 	{
-		TileHeadData.UnsetUAV(RHICmdList, GetComputeShader());
 		TileHeadDataUnpacked.UnsetUAV(RHICmdList, GetComputeShader());
 		TileArrayData.UnsetUAV(RHICmdList, GetComputeShader());
 		TileArrayNextAllocation.UnsetUAV(RHICmdList, GetComputeShader());
@@ -794,7 +790,6 @@ public:
 		Ar << DeferredParameters;
 		Ar << ObjectParameters;
 		Ar << AOParameters;
-		Ar << TileHeadData;
 		Ar << TileHeadDataUnpacked;
 		Ar << TileArrayData;
 		Ar << TileArrayNextAllocation;
@@ -808,7 +803,6 @@ private:
 	FDeferredPixelShaderParameters DeferredParameters;
 	FDistanceFieldCulledObjectBufferParameters ObjectParameters;
 	FAOParameters AOParameters;
-	FRWShaderParameter TileHeadData;
 	FRWShaderParameter TileHeadDataUnpacked;
 	FRWShaderParameter TileArrayData;
 	FRWShaderParameter TileArrayNextAllocation;
@@ -1426,7 +1420,6 @@ public:
 		SavedStartIndex.Bind(Initializer.ParameterMap, TEXT("SavedStartIndex"));
 		ViewDimensionsParameter.Bind(Initializer.ParameterMap, TEXT("ViewDimensions"));
 		ThreadToCulledTile.Bind(Initializer.ParameterMap, TEXT("ThreadToCulledTile"));
-		TileHeadData.Bind(Initializer.ParameterMap, TEXT("TileHeadData"));
 		TileHeadDataUnpacked.Bind(Initializer.ParameterMap, TEXT("TileHeadDataUnpacked"));
 		TileArrayData.Bind(Initializer.ParameterMap, TEXT("TileArrayData"));
 		TileListGroupSize.Bind(Initializer.ParameterMap, TEXT("TileListGroupSize"));
@@ -1494,7 +1487,6 @@ public:
 		FTileIntersectionResources* TileIntersectionResources = ((FSceneViewState*)View.State)->AOTileIntersectionResources;
 
 		SetSRVParameter(RHICmdList, ShaderRHI, TileHeadDataUnpacked, TileIntersectionResources->TileHeadDataUnpacked.SRV);
-		SetSRVParameter(RHICmdList, ShaderRHI, TileHeadData, TileIntersectionResources->TileHeadData.SRV);
 		SetSRVParameter(RHICmdList, ShaderRHI, TileArrayData, TileIntersectionResources->TileArrayData.SRV);
 		SetSRVParameter(RHICmdList, ShaderRHI, IrradianceCacheTileCoordinate, SurfaceCacheResources.Level[DepthLevel]->TileCoordinate.SRV);
 		SetSRVParameter(RHICmdList, ShaderRHI, TileConeDepthRanges, TileIntersectionResources->TileConeDepthRanges.SRV);
@@ -1566,7 +1558,6 @@ public:
 		Ar << ViewDimensionsParameter;
 		Ar << ThreadToCulledTile;
 		Ar << TileHeadDataUnpacked;
-		Ar << TileHeadData;
 		Ar << TileArrayData;
 		Ar << TileListGroupSize;
 		Ar << TanConeHalfAngle;
@@ -1598,7 +1589,6 @@ private:
 	FShaderParameter ViewDimensionsParameter;
 	FShaderParameter ThreadToCulledTile;
 	FShaderResourceParameter TileHeadDataUnpacked;
-	FShaderResourceParameter TileHeadData;
 	FShaderResourceParameter TileArrayData;
 	FShaderParameter TileListGroupSize;
 	FShaderParameter TanConeHalfAngle;
@@ -2238,7 +2228,7 @@ FIntPoint BuildTileObjectLists(FRHICommandListImmediate& RHICmdList, FScene* Sce
 			// Indicates the clear value for each channel of the UAV format
 			uint32 ClearValues[4] = { 0 };
 			RHICmdList.ClearUAV(TileIntersectionResources->TileArrayNextAllocation.UAV, ClearValues);
-			RHICmdList.ClearUAV(TileIntersectionResources->TileHeadData.UAV, ClearValues);
+			RHICmdList.ClearUAV(TileIntersectionResources->TileHeadDataUnpacked.UAV, ClearValues);
 
 			TShaderMapRef<FDistanceFieldBuildTileListCS > ComputeShader(View.ShaderMap);
 
diff --git a/Engine/Source/Runtime/Renderer/Private/DistanceFieldSurfaceCacheLighting.h b/Engine/Source/Runtime/Renderer/Private/DistanceFieldSurfaceCacheLighting.h
--- a/Engine/Source/Runtime/Renderer/Private/DistanceFieldSurfaceCacheLighting.h
+++ b/Engine/Source/Runtime/Renderer/Private/DistanceFieldSurfaceCacheLighting.h
@@ -342,7 +342,6 @@ public:
 		TileConeAxisAndCos.Release();
 		TileConeDepthRanges.Release();
 		TileHeadDataUnpacked.Release();
-		TileHeadData.Release();
 		TileArrayData.Release();
 		TileArrayNextAllocation.Release();
 	}
@@ -352,13 +351,12 @@ public:
 	FRWBuffer TileConeAxisAndCos;
 	FRWBuffer TileConeDepthRanges;
 	FRWBuffer TileHeadDataUnpacked;
-	FRWBuffer TileHeadData;
 	FRWBuffer TileArrayData;
 	FRWBuffer TileArrayNextAllocation;
 
 	size_t GetSizeBytes() const
 	{
-		return TileConeAxisAndCos.NumBytes + TileConeDepthRanges.NumBytes + TileHeadDataUnpacked.NumBytes + TileHeadData.NumBytes + TileArrayData.NumBytes + TileArrayNextAllocation.NumBytes;
+		return TileConeAxisAndCos.NumBytes + TileConeDepthRanges.NumBytes + TileHeadDataUnpacked.NumBytes + TileArrayData.NumBytes + TileArrayNextAllocation.NumBytes;
 	}
 };
 
