Fixed many transform nodes being broken in decal materials

diff --git a/Engine/Shaders/MaterialTemplate.usf b/Engine/Shaders/MaterialTemplate.usf
--- a/Engine/Shaders/MaterialTemplate.usf
+++ b/Engine/Shaders/MaterialTemplate.usf
@@ -334,8 +334,6 @@ struct FMaterialVertexParameters
 	}
 #endif //DECAL_RENDERSTAGE
 
-#if HAS_PRIMITIVE_UNIFORM_BUFFER
-
 /** Transforms a vector from tangent space to world space, prescaling by an amount calculated previously */
 MaterialFloat3 TransformTangentVectorToWorld_PreScaled(FMaterialTessellationParameters Parameters, MaterialFloat3 InTangentVector)
 {
@@ -351,6 +349,8 @@ MaterialFloat3 TransformTangentVectorToWorld_PreScaled(FMaterialTessellationPara
 #endif // #if FEATURE_LEVEL_SM5
 }
 
+#if HAS_PRIMITIVE_UNIFORM_BUFFER
+
 /** Transforms a vector from tangent space to local space */
 MaterialFloat3 TransformTangentVectorToLocal(FMaterialPixelParameters Parameters, MaterialFloat3 InTangentVector)
 {
@@ -358,6 +358,8 @@ MaterialFloat3 TransformTangentVectorToLocal(FMaterialPixelParameters Parameters
 	return mul(mul(InTangentVector, Parameters.TangentToWorld), (MaterialFloat3x3)Primitive.WorldToLocal);
 }
 
+#endif
+
 /** Transforms a vector from tangent space to view space */
 MaterialFloat3 TransformTangentVectorToView(FMaterialPixelParameters Parameters, MaterialFloat3 InTangentVector)
 {
@@ -394,6 +396,8 @@ MaterialFloat3 TransformLocalVectorToView(MaterialFloat3 InLocalVector)
 	return mul(mul(InLocalVector, GetLocalToWorld3x3()), (MaterialFloat3x3)View.TranslatedWorldToView);
 }
 
+#if HAS_PRIMITIVE_UNIFORM_BUFFER
+
 /** Transforms a vector from world space to local space */
 MaterialFloat3 TransformWorldVectorToLocal(MaterialFloat3 InWorldVector)
 {
@@ -406,6 +410,8 @@ MaterialFloat3 TransformViewVectorToLocal(MaterialFloat3 InViewVector)
 	return mul(mul(InViewVector, (MaterialFloat3x3)View.ViewToTranslatedWorld), (MaterialFloat3x3)Primitive.WorldToLocal);
 }
 
+#endif
+
 /** Transforms a vector from view space to world space */
 MaterialFloat3 TransformViewVectorToWorld(MaterialFloat3 InViewVector)
 {
@@ -418,6 +424,8 @@ MaterialFloat3 TransformWorldVectorToView(MaterialFloat3 InWorldVector)
 	return mul(InWorldVector, (MaterialFloat3x3)View.TranslatedWorldToView);
 }
 
+#if HAS_PRIMITIVE_UNIFORM_BUFFER
+
 /** Transforms a position from local space to absolute world space */
 float3 TransformLocalPositionToWorld(FMaterialPixelParameters Parameters,float3 InLocalPosition)
 {
@@ -440,6 +448,8 @@ float3 TransformWorldPositionToLocal(float3 InWorldPosition)
 	return mul(float4(InWorldPosition, 1), Primitive.WorldToLocal).xyz;
 }
 
+#endif
+
 /** Transforms a position from view space to world space */
 MaterialFloat3 TransformViewPositionToWorld(MaterialFloat3 InViewPosition)
 {
@@ -452,6 +462,8 @@ MaterialFloat3 TransformWorldPositionToView(MaterialFloat3 InWorldPosition)
 	return mul(MaterialFloat4(InWorldPosition + View.PreViewTranslation.xyz, 1), View.TranslatedWorldToView).xyz;
 }
 
+#if HAS_PRIMITIVE_UNIFORM_BUFFER
+
 /** Return the object's position in world space */
 float3 GetObjectWorldPosition(FMaterialPixelParameters Parameters)
 {
@@ -473,6 +485,8 @@ float3 GetObjectWorldPosition(FMaterialVertexParameters Parameters)
 	#endif
 }
 
+#endif
+
 /** Get the per-instance random value when instancing */
 float GetPerInstanceRandom(FMaterialVertexParameters Parameters)
 {
@@ -518,8 +532,6 @@ MaterialFloat GetDistanceCullFade()
 	return saturate( View.RealTime * PrimitiveFade.FadeTimeScaleBias.x + PrimitiveFade.FadeTimeScaleBias.y );
 }
 
-#endif // #if HAS_PRIMITIVE_UNIFORM_BUFFER
-
 /** Rotates Position about the given axis by the given angle, in radians, and returns the offset to Position. */
 float3 RotateAboutAxis(float4 NormalizedRotationAxisAndAngle, float3 PositionOnAxis, float3 Position)
 {
