Removed 'uniform' keywords for gl

diff --git a/Engine/Shaders/GlobalDistanceFieldShared.usf b/Engine/Shaders/GlobalDistanceFieldShared.usf
--- a/Engine/Shaders/GlobalDistanceFieldShared.usf
+++ b/Engine/Shaders/GlobalDistanceFieldShared.usf
@@ -17,14 +17,14 @@ float GlobalVolumeDimension;
 float GlobalVolumeTexelSize;
 float MaxGlobalDistance;
 
-float3 ComputeGlobalUV(float3 WorldPosition, uniform uint ClipmapIndex)
+float3 ComputeGlobalUV(float3 WorldPosition, uint ClipmapIndex)
 {
 	//return ((WorldPosition - GlobalVolumeCenterAndExtent[ClipmapIndex].xyz + GlobalVolumeScollOffset[ClipmapIndex].xyz) / (GlobalVolumeCenterAndExtent[ClipmapIndex].w * 2) + .5f);
 	float4 WorldToUVAddAndMul = GlobalVolumeWorldToUVAddAndMul[ClipmapIndex];
 	return WorldPosition * WorldToUVAddAndMul.www + WorldToUVAddAndMul.xyz;
 }
 
-float GetDistanceToNearestSurfaceGlobalClipmap(float3 WorldPosition, uniform uint ClipmapIndex)
+float GetDistanceToNearestSurfaceGlobalClipmap(float3 WorldPosition, uint ClipmapIndex)
 {
 	float3 GlobalUV = ComputeGlobalUV(WorldPosition, ClipmapIndex);
 	float DistanceToSurface = Texture3DSampleLevel(GlobalDistanceFieldTexture[ClipmapIndex], GlobalDistanceFieldSampler[ClipmapIndex], GlobalUV, 0).x;
@@ -71,7 +71,7 @@ float GetDistanceToNearestSurfaceGlobal(float3 WorldPosition)
 	return DistanceToSurface;
 }
 
-float3 GetDistanceFieldGradientGlobalClipmap(float3 WorldPosition, uniform uint ClipmapIndex)
+float3 GetDistanceFieldGradientGlobalClipmap(float3 WorldPosition, uint ClipmapIndex)
 {
 	float3 GlobalUV = ComputeGlobalUV(WorldPosition, ClipmapIndex);
 	float R = Texture3DSampleLevel(GlobalDistanceFieldTexture[ClipmapIndex], GlobalDistanceFieldSampler[ClipmapIndex], float3(GlobalUV.x + GlobalVolumeTexelSize, GlobalUV.y, GlobalUV.z), 0).x;
