Fix gl compile error

diff --git a/Engine/Shaders/GlobalDistanceFieldShared.usf b/Engine/Shaders/GlobalDistanceFieldShared.usf
--- a/Engine/Shaders/GlobalDistanceFieldShared.usf
+++ b/Engine/Shaders/GlobalDistanceFieldShared.usf
@@ -6,6 +6,21 @@
 
 #if FEATURE_LEVEL >= FEATURE_LEVEL_SM5
 
+//@todo - get the clipmap branches to compile under gl
+#if GL4_PROFILE
+
+float GetDistanceToNearestSurfaceGlobal(float3 WorldPosition)
+{
+	return 0;
+}
+
+float3 GetDistanceFieldGradientGlobal(float3 WorldPosition)
+{
+	return 0;
+}
+
+#else
+
 // Must match C++
 #define MAX_GLOBAL_DF_CLIPMAPS 4
 
@@ -127,4 +142,5 @@ float3 GetDistanceFieldGradientGlobal(float3 WorldPosition)
 	return Gradient;
 }
 
-#endif
\ No newline at end of file
+#endif
+#endif // FEATURE_LEVEL >= FEATURE_LEVEL_SM5
\ No newline at end of file
