Fixed gl compile errors

diff --git a/Engine/Shaders/GlobalDistanceFieldShared.usf b/Engine/Shaders/GlobalDistanceFieldShared.usf
--- a/Engine/Shaders/GlobalDistanceFieldShared.usf
+++ b/Engine/Shaders/GlobalDistanceFieldShared.usf
@@ -6,20 +6,10 @@
 
 #if FEATURE_LEVEL >= FEATURE_LEVEL_SM5
 
-//@todo - get the clipmap branches to compile under gl
-#if GL4_PROFILE
-
-float GetDistanceToNearestSurfaceGlobal(float3 WorldPosition)
-{
-	return 0;
-}
-
-float3 GetDistanceFieldGradientGlobal(float3 WorldPosition)
-{
-	return 0;
-}
+float MaxGlobalDistance;
 
-#else
+//@todo - get the clipmap branches to compile under gl
+#if SM5_PROFILE || PS4_PROFILE
 
 // Must match C++
 #define MAX_GLOBAL_DF_CLIPMAPS 4
@@ -30,7 +20,6 @@ float4 GlobalVolumeCenterAndExtent[MAX_GLOBAL_DF_CLIPMAPS];
 float4 GlobalVolumeWorldToUVAddAndMul[MAX_GLOBAL_DF_CLIPMAPS];
 float GlobalVolumeDimension;
 float GlobalVolumeTexelSize;
-float MaxGlobalDistance;
 
 float3 ComputeGlobalUV(float3 WorldPosition, uint ClipmapIndex)
 {
@@ -142,5 +131,17 @@ float3 GetDistanceFieldGradientGlobal(float3 WorldPosition)
 	return Gradient;
 }
 
+#else
+
+float GetDistanceToNearestSurfaceGlobal(float3 WorldPosition)
+{
+	return MaxGlobalDistance;
+}
+
+float3 GetDistanceFieldGradientGlobal(float3 WorldPosition)
+{
+	return float3(0, 0, 1);
+}
+
 #endif
 #endif // FEATURE_LEVEL >= FEATURE_LEVEL_SM5
\ No newline at end of file
