Fixed DFAO surface cache path

diff --git a/Engine/Shaders/DistanceFieldSurfaceCacheLighting.usf b/Engine/Shaders/DistanceFieldSurfaceCacheLighting.usf
--- a/Engine/Shaders/DistanceFieldSurfaceCacheLighting.usf
+++ b/Engine/Shaders/DistanceFieldSurfaceCacheLighting.usf
@@ -249,7 +249,7 @@ void IrradianceCacheSplatVS(
 
 	Output.PositionRadius = float4(RecordPosition, RecordRadius); 
 
-	Output.NormalAndFade = float4(RecordNormalAndFade.xyz, 0);
+	Output.NormalAndFade = float4(RecordNormalAndFade.xyz, 1);
 
 #if FINAL_INTERPOLATION_PASS
 	Output.BentNormal = IrradianceCacheBentNormal[JobIndex * BENT_NORMAL_STRIDE].xyz;
