fixed shader compiler warning

diff --git a/Engine/Shaders/ScreenSpaceRayCast.usf b/Engine/Shaders/ScreenSpaceRayCast.usf
--- a/Engine/Shaders/ScreenSpaceRayCast.usf
+++ b/Engine/Shaders/ScreenSpaceRayCast.usf
@@ -56,7 +56,7 @@ float4 RayCast( Texture2D Texture, SamplerState Sampler, float3 R, float Roughne
 
 	{
 		// Computes the scale down factor for RayStepScreen required to fit on the X and Y axis in order to clip it in the viewport
-		float2 AbsRayStepScreen = abs(RayStepScreen);
+		float2 AbsRayStepScreen = abs(RayStepScreen.xy);
 		float2 S = (AbsRayStepScreen - max(abs(RayStepScreen.xy + RayStartScreen.xy) - 1.0f, 0.0f)) / AbsRayStepScreen;
 
 		// Rescales RayStepScreen accordingly
