Added max displacement to materials to fix tessellation frustum culling.

diff --git a/Engine/Shaders/MaterialTemplate.usf b/Engine/Shaders/MaterialTemplate.usf
--- a/Engine/Shaders/MaterialTemplate.usf
+++ b/Engine/Shaders/MaterialTemplate.usf
@@ -1363,6 +1363,11 @@ half3 GetMaterialWorldDisplacement(FMaterialTessellationParameters Parameters)
 %s;
 }
 
+half GetMaterialMaxDisplacement()
+{
+%s;
+}
+
 half GetMaterialTessellationMultiplier(FMaterialTessellationParameters Parameters)
 {
 %s;
diff --git a/Engine/Shaders/Tessellation.usf b/Engine/Shaders/Tessellation.usf
--- a/Engine/Shaders/Tessellation.usf
+++ b/Engine/Shaders/Tessellation.usf
@@ -62,11 +62,14 @@ float GetTessellationMultiplier(FPassSpecificVSToDS I)
 	return GetMaterialTessellationMultiplier(MaterialParameters);
 }
 
-int3 GetClipFlag( float3 Origin )
+// Flag set if inside frustum
+int3 GetClipFlag( float3 Origin, float Radius )
 {
+	//float3 ClipExpand = Radius * float3( View.ViewToClip[0][0], View.ViewToClip[1][1], View.ViewToClip[2][2] );
+	float4 ClipExpand = mul( float4( Radius, Radius, Radius, 0 ), View.ViewToClip );
 	float4 ClipPos = mul( float4( Origin, 1.0 ), View.TranslatedWorldToClip );
-    int3 Clip0 = ClipPos.xyz < ClipPos.w;
-    int3 Clip1 = ClipPos.xyz > -ClipPos.w;
+    int3 Clip0 = ClipPos.xyz - ClipExpand.xyz <  ClipPos.w + ClipExpand.w;
+    int3 Clip1 = ClipPos.xyz + ClipExpand.xyz > -ClipPos.w - ClipExpand.w;
     return Clip0 + 2 * Clip1;
 }
 
@@ -81,11 +84,13 @@ float4 CalculateCompositeTessellationFactors(float3 Control0, float3 Control1, f
 {
 #if USE_ADAPTIVE_TESSELLATION_FACTOR
 #if 1
+	half MaxDisplacement = GetMaterialMaxDisplacement();
+
 	// Frustum cull
 	int3 ClipFlag = 0;
-	ClipFlag  = GetClipFlag( Control0 );
-	ClipFlag |= GetClipFlag( Control1 );
-	ClipFlag |= GetClipFlag( Control2 );
+	ClipFlag  = GetClipFlag( Control0, MaxDisplacement );
+	ClipFlag |= GetClipFlag( Control1, MaxDisplacement );
+	ClipFlag |= GetClipFlag( Control2, MaxDisplacement );
 	if( any( ClipFlag != 3 ) )
 	{
 		return 0;
diff --git a/Engine/Source/Runtime/Engine/Classes/Materials/Material.h b/Engine/Source/Runtime/Engine/Classes/Materials/Material.h
--- a/Engine/Source/Runtime/Engine/Classes/Materials/Material.h
+++ b/Engine/Source/Runtime/Engine/Classes/Materials/Material.h
@@ -623,6 +623,9 @@ public:
 	UPROPERTY(EditAnywhere, BlueprintReadOnly, Category=Tessellation, meta=(DisplayName = "Adaptive Tessellation"))
 	uint32 bEnableAdaptiveTessellation:1;
 
+	UPROPERTY(EditAnywhere, Category=Tessellation)
+	float MaxDisplacement;
+
 	/** Enables a wireframe view of the mesh the material is applied to.  */
 	UPROPERTY(EditAnywhere, Category=Material, AdvancedDisplay)
 	uint32 Wireframe:1;
diff --git a/Engine/Source/Runtime/Engine/Private/Materials/HLSLMaterialTranslator.h b/Engine/Source/Runtime/Engine/Private/Materials/HLSLMaterialTranslator.h
--- a/Engine/Source/Runtime/Engine/Private/Materials/HLSLMaterialTranslator.h
+++ b/Engine/Source/Runtime/Engine/Private/Materials/HLSLMaterialTranslator.h
@@ -682,6 +682,7 @@ public:
 		LazyPrintf.PushParam(*GenerateFunctionCode(MP_WorldPositionOffset));
 		LazyPrintf.PushParam(*GenerateFunctionCode(CompiledMP_PrevWorldPositionOffset));
 		LazyPrintf.PushParam(*GenerateFunctionCode(MP_WorldDisplacement));
+		LazyPrintf.PushParam(*FString::Printf(TEXT("return %.5f"),Material->GetMaxDisplacement()));
 		LazyPrintf.PushParam(*GenerateFunctionCode(MP_TessellationMultiplier));
 		LazyPrintf.PushParam(*GenerateFunctionCode(MP_SubsurfaceColor));
 		LazyPrintf.PushParam(*GenerateFunctionCode(MP_ClearCoat));
diff --git a/Engine/Source/Runtime/Engine/Private/Materials/Material.cpp b/Engine/Source/Runtime/Engine/Private/Materials/Material.cpp
--- a/Engine/Source/Runtime/Engine/Private/Materials/Material.cpp
+++ b/Engine/Source/Runtime/Engine/Private/Materials/Material.cpp
@@ -540,6 +540,7 @@ UMaterial::UMaterial(const FObjectInitializer& ObjectInitializer)
 	D3D11TessellationMode = MTM_NoTessellation;
 	bEnableCrackFreeDisplacement = false;
 	bEnableAdaptiveTessellation = true;
+	MaxDisplacement = 0.0f;
 	bOutputVelocityOnBasePass = true;
 	bEnableSeparateTranslucency = true;
 	bEnableResponsiveAA = false;
@@ -2466,6 +2467,7 @@ bool UMaterial::CanEditChange(const UProperty* InProperty) const
 		}
 
 		if (PropertyName == GET_MEMBER_NAME_STRING_CHECKED(UMaterial, bEnableCrackFreeDisplacement) ||
+			PropertyName == GET_MEMBER_NAME_STRING_CHECKED(UMaterial, MaxDisplacement) ||
 			PropertyName == GET_MEMBER_NAME_STRING_CHECKED(UMaterial, bEnableAdaptiveTessellation)
 			)
 		{
diff --git a/Engine/Source/Runtime/Engine/Private/Materials/MaterialShared.cpp b/Engine/Source/Runtime/Engine/Private/Materials/MaterialShared.cpp
--- a/Engine/Source/Runtime/Engine/Private/Materials/MaterialShared.cpp
+++ b/Engine/Source/Runtime/Engine/Private/Materials/MaterialShared.cpp
@@ -869,6 +869,7 @@ float FMaterialResource::GetTranslucentBackscatteringExponent() const { return M
 FLinearColor FMaterialResource::GetTranslucentMultipleScatteringExtinction() const { return Material->TranslucentMultipleScatteringExtinction; }
 float FMaterialResource::GetTranslucentShadowStartOffset() const { return Material->TranslucentShadowStartOffset; }
 float FMaterialResource::GetRefractionDepthBiasValue() const { return Material->RefractionDepthBias; }
+float FMaterialResource::GetMaxDisplacement() const { return Material->MaxDisplacement; }
 bool FMaterialResource::UseTranslucencyVertexFog() const {return Material->bUseTranslucencyVertexFog;}
 FString FMaterialResource::GetFriendlyName() const { return *GetNameSafe(Material); }
 
diff --git a/Engine/Source/Runtime/Engine/Public/MaterialShared.h b/Engine/Source/Runtime/Engine/Public/MaterialShared.h
--- a/Engine/Source/Runtime/Engine/Public/MaterialShared.h
+++ b/Engine/Source/Runtime/Engine/Public/MaterialShared.h
@@ -952,6 +952,7 @@ public:
 	virtual FLinearColor GetTranslucentMultipleScatteringExtinction() const { return FLinearColor::White; }
 	virtual float GetTranslucentShadowStartOffset() const { return 0.0f; }
 	virtual float GetRefractionDepthBiasValue() const { return 0.0f; }
+	virtual float GetMaxDisplacement() const { return 0.0f; }
 	virtual bool UseTranslucencyVertexFog() const { return false; }
 	virtual FString GetFriendlyName() const = 0;
 	virtual bool HasVertexPositionOffsetConnected() const { return false; }
@@ -1516,6 +1517,7 @@ public:
 	ENGINE_API virtual bool RequiresSynchronousCompilation() const override;
 	ENGINE_API virtual bool IsDefaultMaterial() const override;
 	ENGINE_API virtual float GetRefractionDepthBiasValue() const override;
+	ENGINE_API virtual float GetMaxDisplacement() const override;
 	ENGINE_API virtual bool UseTranslucencyVertexFog() const override;
 	/**
 	 * Should shaders compiled for this material be saved to disk?
