Fixes FilterPixelShader.usf compilations failure on HTML5 profile cause by cast from half to float

diff --git a/Engine/Shaders/FilterPixelShader.usf b/Engine/Shaders/FilterPixelShader.usf
--- a/Engine/Shaders/FilterPixelShader.usf
+++ b/Engine/Shaders/FilterPixelShader.usf
@@ -54,6 +54,44 @@ float4 SampleOffsets[(MAX_NUM_SAMPLES + 1) / 2];
 half4 SampleWeights[MAX_NUM_SAMPLES];
 int SampleCount;
 
+# if !(ES2_PROFILE || ES3_1_PROFILE)
+
+void Main(
+	float2 InUV : TEXCOORD0,
+	out half4 OutColor : SV_Target0
+	)
+{
+	half3 Dest = 0;
+	
+	int SampleIndex;
+	LOOP for (SampleIndex = 0; SampleIndex < SampleCount - 1; SampleIndex += 2)
+	{
+		{
+			const float2 UV = InUV + SampleOffsets[SampleIndex / 2].xy;
+			Dest += FilterTexture.Sample(FilterTextureSampler, UV).rgb * SampleWeights[SampleIndex + 0];
+		}
+		
+		{
+			const float2 UV = InUV + SampleOffsets[SampleIndex / 2].wz;
+			Dest += FilterTexture.Sample(FilterTextureSampler, UV).rgb * SampleWeights[SampleIndex + 1];
+		}
+	}
+	BRANCH if (SampleIndex < SampleCount)
+	{
+		const float2 UV = InUV + SampleOffsets[SampleIndex / 2].xy;
+		Dest += FilterTexture.Sample(FilterTextureSampler, UV).rgb * SampleWeights[SampleIndex + 1];
+	}
+
+#  if COMBINE_METHOD == 1
+	Dest += Texture2DSample(AdditiveTexture, AdditiveTextureSampler, InUV).rgb;
+#  endif
+
+	// RETURN_COLOR not needed unless writing to SceneColor;
+	OutColor = half4(Dest, 0);
+}
+
+# else //!(ES2_PROFILE || ES3_1_PROFILE)
+
 void Main(
 	float2 InUV : TEXCOORD0,
 	out half4 OutColor : SV_Target0
@@ -83,18 +121,22 @@ void Main(
 		Combine(Dest, Texture2DSample(FilterTexture, FilterTextureSampler, UV), Mask * SampleWeights[SampleIndex + 0]);
 	}
 
-#if COMBINE_METHOD == 1
+#  if COMBINE_METHOD == 1
 	Dest += Texture2DSample(AdditiveTexture, AdditiveTextureSampler, InUV);
-#endif
+#  endif
 
 	// RETURN_COLOR not needed unless writing to SceneColor;
 	OutColor = Dest;
 }
 
-#elif (ES2_PROFILE || ES3_1_PROFILE)
+# endif //!(ES2_PROFILE || ES3_1_PROFILE)
+
+#else //NUM_SAMPLES!=0
 
 half4 SampleWeights[NUM_SAMPLES];
 
+# if (ES2_PROFILE || ES3_1_PROFILE)
+
 void Main(
 	float2 InUV : TEXCOORD0,
 	float2 InOffsetUVs[NUM_SAMPLES] : TEXCOORD1,
@@ -107,16 +149,14 @@ void Main(
 		Dest += FilterTexture.Sample( FilterTextureSampler, InOffsetUVs[i] ).rgb * SampleWeights[i];
 	}
 
-#if COMBINE_METHOD == 1
+#  if COMBINE_METHOD == 1
 	Dest += AdditiveTexture.Sample( AdditiveTextureSampler, InUV ).rgb;
-#endif
+#  endif
 
 	OutColor.rgb = Dest;
 }
 
-#else
-
-half4 SampleWeights[NUM_SAMPLES];
+# else //!(ES2_PROFILE || ES3_1_PROFILE)
 
 void Main(
 	float2 InUV : TEXCOORD0,
@@ -145,14 +185,16 @@ void Main(
 		Combine(Dest, Texture2DSample(FilterTexture, FilterTextureSampler,UVUV.xy), Mask * SampleWeights[SampleIndex + 0]);
 	}
 
-#if COMBINE_METHOD == 1
+#  if COMBINE_METHOD == 1
 	Dest += Texture2DSample(AdditiveTexture, AdditiveTextureSampler, InUV);
-#endif
+#  endif
 
 	// RETURN_COLOR not needed unless writing to SceneColor;
 	OutColor = Dest;
 }
-#endif
+
+# endif //!(ES2_PROFILE || ES3_1_PROFILE)
+#endif //NUM_SAMPLES!=0
 
 // currently only used for downsampling, doesn't need to be masked as it doesn't read outside bounds
 // reads 4 samples
