Fixes broken builds caused by a shader compilation failure in FilterPixelShader.usf

diff --git a/Engine/Shaders/FilterPixelShader.usf b/Engine/Shaders/FilterPixelShader.usf
--- a/Engine/Shaders/FilterPixelShader.usf
+++ b/Engine/Shaders/FilterPixelShader.usf
@@ -54,7 +54,7 @@ float4 SampleOffsets[(MAX_NUM_SAMPLES + 1) / 2];
 half4 SampleWeights[MAX_NUM_SAMPLES];
 int SampleCount;
 
-# if !(ES2_PROFILE || ES3_1_PROFILE)
+# if (ES2_PROFILE || ES3_1_PROFILE)
 
 void Main(
 	float2 InUV : TEXCOORD0,
