Fix for reading wrong texture channel in WebGL. Fixes UE-15719.

diff --git a/Engine/Shaders/LandscapeVertexFactory.usf b/Engine/Shaders/LandscapeVertexFactory.usf
--- a/Engine/Shaders/LandscapeVertexFactory.usf
+++ b/Engine/Shaders/LandscapeVertexFactory.usf
@@ -204,7 +204,7 @@ float3x3 VertexFactoryGetPerPixelTangentBasis(FVertexFactoryInterpolantsVSToPS I
 	float4 SampleValue = Texture2DSample(NormalmapTexture, NormalmapTextureSampler, Interpolants.WeightMapTexCoord);
 #endif
 #if WEBGL
-	float2 SampleNormal = float2(SampleValue.r, SampleValue.a) * float2(2.0,2.0) - float2(1.0,1.0);
+	float2 SampleNormal = float2(SampleValue.b, SampleValue.a) * float2(2.0,2.0) - float2(1.0,1.0);
 #else
 	float2 SampleNormal = float2(SampleValue.b, SampleValue.a) * float2(2.0,2.0) - float2(1.0,1.0);
 #endif
@@ -358,7 +358,7 @@ FMaterialPixelParameters GetMaterialPixelParameters(FVertexFactoryInterpolantsVS
 #if (ES2_PROFILE || ES3_1_PROFILE) && PIXELSHADER
 
 #if WEBGL
-	Result.LayerWeights.xy = Texture2DSample(NormalmapTexture, NormalmapTextureSampler, Interpolants.WeightMapTexCoord).bg;
+	Result.LayerWeights.xy = Texture2DSample(NormalmapTexture, NormalmapTextureSampler, Interpolants.WeightMapTexCoord).rg;
 #else
 	Result.LayerWeights.xy = Texture2DSample(NormalmapTexture, NormalmapTextureSampler, Interpolants.WeightMapTexCoord).rg;
 #endif 
