Fixes a shader compilation failure when plugin-in const 0 to the material PixelDepthOffset input.

diff --git a/Engine/Shaders/BasePassForForwardShadingPixelShader.usf b/Engine/Shaders/BasePassForForwardShadingPixelShader.usf
--- a/Engine/Shaders/BasePassForForwardShadingPixelShader.usf
+++ b/Engine/Shaders/BasePassForForwardShadingPixelShader.usf
@@ -99,7 +99,7 @@ half DirectionalLightShadowing;
 void Main( 
 	FVertexFactoryInterpolantsVSToPS Interpolants,
 	FForwardShadingBasePassInterpolantsVSToPS BasePassInterpolants,
-	in float4 SvPosition : SV_Position,
+	in float4 SvPosition : SV_Position
 	OPTIONAL_IsFrontFace,
 	out half4 OutColor	: SV_Target0
 	)
diff --git a/Engine/Shaders/BasePassPixelShader.usf b/Engine/Shaders/BasePassPixelShader.usf
--- a/Engine/Shaders/BasePassPixelShader.usf
+++ b/Engine/Shaders/BasePassPixelShader.usf
@@ -597,10 +597,9 @@ void ApplyPixelDepthOffsetForBasePass(inout FMaterialPixelParameters MaterialPar
 void Main(
 	FVertexFactoryInterpolantsVSToPS Interpolants,
 	FBasePassInterpolantsVSToPS BasePassInterpolants,
-	in INPUT_POSITION_QUALIFIERS float4 SvPosition : SV_Position,
-
-	OPTIONAL_IsFrontFace,
-	out float4 OutColor		: SV_Target0
+	in INPUT_POSITION_QUALIFIERS float4 SvPosition : SV_Position
+	OPTIONAL_IsFrontFace
+	,out float4 OutColor		: SV_Target0
 #if FEATURE_LEVEL >= FEATURE_LEVEL_SM4 && (MATERIALBLENDING_SOLID || MATERIALBLENDING_MASKED) && !NO_GBUFFER
 	,out float4 OutGBufferA : SV_Target1 
 	,out float4 OutGBufferB : SV_Target2
diff --git a/Engine/Shaders/DepthOnlyPixelShader.usf b/Engine/Shaders/DepthOnlyPixelShader.usf
--- a/Engine/Shaders/DepthOnlyPixelShader.usf
+++ b/Engine/Shaders/DepthOnlyPixelShader.usf
@@ -12,11 +12,11 @@ void Main(
 #if !MATERIALBLENDING_SOLID || OUTPUT_PIXEL_DEPTH_OFFSET
 	FVertexFactoryInterpolantsVSToPS FactoryInterpolants,
 	float4 PixelPosition : TEXCOORD6,
-	in INPUT_POSITION_QUALIFIERS float4 SvPosition : SV_Position,
+	in INPUT_POSITION_QUALIFIERS float4 SvPosition : SV_Position
 	OPTIONAL_IsFrontFace
-	OPTIONAL_OutDepth
+	OPTIONAL_OutDepth,
 #endif
-	,out float4 OutColor : SV_Target0
+	out float4 OutColor : SV_Target0
 	)
 {
 #if !MATERIALBLENDING_SOLID || OUTPUT_PIXEL_DEPTH_OFFSET
diff --git a/Engine/Shaders/DistortAccumulatePS.usf b/Engine/Shaders/DistortAccumulatePS.usf
--- a/Engine/Shaders/DistortAccumulatePS.usf
+++ b/Engine/Shaders/DistortAccumulatePS.usf
@@ -18,7 +18,7 @@ float4 DistortionParams;
 void Main(
 	FVertexFactoryInterpolantsVSToPS Interpolants,
 	float4	PixelPosition	: TEXCOORD6,
-	in float4 SvPosition	: SV_Position,
+	in float4 SvPosition	: SV_Position
 	OPTIONAL_IsFrontFace,
 	out float4 OutColor		: SV_Target0
 	)
diff --git a/Engine/Shaders/HitProxyPixelShader.usf b/Engine/Shaders/HitProxyPixelShader.usf
--- a/Engine/Shaders/HitProxyPixelShader.usf
+++ b/Engine/Shaders/HitProxyPixelShader.usf
@@ -19,7 +19,7 @@ void Main(
 	float4 InstanceHitProxyId : HIT_PROXY_ID,
 #endif
 	float4 PixelPosition	: TEXCOORD6,
-	in INPUT_POSITION_QUALIFIERS float4 SvPosition	: SV_Position,
+	in INPUT_POSITION_QUALIFIERS float4 SvPosition	: SV_Position
 	OPTIONAL_IsFrontFace
 	OPTIONAL_OutDepth
 	,out float4 OutColor		: SV_Target0
diff --git a/Engine/Shaders/LandscapeGrassWeight.usf b/Engine/Shaders/LandscapeGrassWeight.usf
--- a/Engine/Shaders/LandscapeGrassWeight.usf
+++ b/Engine/Shaders/LandscapeGrassWeight.usf
@@ -55,7 +55,7 @@ int OutputPass;
 // Pixel Shader
 void PSMain( 
 	FLandscapeGrassWeightInterpolantsVSToPS Interpolants,
-	in float4 SvPosition : SV_Position,
+	in float4 SvPosition : SV_Position
 	OPTIONAL_IsFrontFace,
 	out half4 OutColor	: SV_Target0
 	)
diff --git a/Engine/Shaders/LightMapDensityShader.usf b/Engine/Shaders/LightMapDensityShader.usf
--- a/Engine/Shaders/LightMapDensityShader.usf
+++ b/Engine/Shaders/LightMapDensityShader.usf
@@ -153,7 +153,7 @@ SamplerState	GridNormalTextureSampler;
 void MainPixelShader(
 	FVertexFactoryInterpolantsVSToPS FactoryInterpolants,
 	float4 WorldPosition	: TEXCOORD6,
-	in float4 SvPosition	: SV_Position,
+	in float4 SvPosition	: SV_Position
 	OPTIONAL_IsFrontFace,
 	out float4 OutColor		: SV_Target0
 	)
diff --git a/Engine/Shaders/MaterialTemplate.usf b/Engine/Shaders/MaterialTemplate.usf
--- a/Engine/Shaders/MaterialTemplate.usf
+++ b/Engine/Shaders/MaterialTemplate.usf
@@ -1543,7 +1543,7 @@ void GetMaterialClippingShadowDepth(FMaterialPixelParameters Parameters)
 	#define OPTIONAL_IsFrontFace
 	static const FIsFrontFace bIsFrontFace = 1;
 #else
-	#define OPTIONAL_IsFrontFace in FIsFrontFace bIsFrontFace : FrontFaceSemantic
+	#define OPTIONAL_IsFrontFace , in FIsFrontFace bIsFrontFace : FrontFaceSemantic
 #endif
 
 /** Initializes the subset of Parameters that was not set in GetMaterialPixelParameters. */
diff --git a/Engine/Shaders/VelocityShader.usf b/Engine/Shaders/VelocityShader.usf
--- a/Engine/Shaders/VelocityShader.usf
+++ b/Engine/Shaders/VelocityShader.usf
@@ -189,7 +189,7 @@ float4 IndividualVelocityScale;
  
 void MainPixelShader(
 	FVelocityInterpsVSToPS Input,
-	in float4 SvPosition : SV_Position,
+	in float4 SvPosition : SV_Position
 	OPTIONAL_IsFrontFace,
 	out float4 OutColor : SV_Target0
 	)
