Fixes a screen space reflection artifact with r.PreViewTranslation = 0

diff --git a/Engine/Shaders/ScreenSpaceReflections.usf b/Engine/Shaders/ScreenSpaceReflections.usf
--- a/Engine/Shaders/ScreenSpaceReflections.usf
+++ b/Engine/Shaders/ScreenSpaceReflections.usf
@@ -125,7 +125,7 @@ void ScreenSpaceReflectionsPS(in float4 UVAndScreenPos : TEXCOORD0, out float4 O
 	float3 N = GBuffer.WorldNormal;
 	float SceneDepth = CalcSceneDepth(UV);
 
-	float3 PositionTranslatedWorld = mul( float3( ScreenPos * SceneDepth, SceneDepth ), (float3x3)View.ScreenToTranslatedWorld );
+	float3 PositionTranslatedWorld = mul( float4( ScreenPos * SceneDepth, SceneDepth, 1 ), View.ScreenToTranslatedWorld ).xyz;
 
 	float3 V = View.TranslatedViewOrigin.xyz - PositionTranslatedWorld;
 	V = normalize( V );
