Changes FMaterialVertexParameter::WorldPosition from non translated world position to translated world position for precision improvements

diff --git a/Engine/Shaders/GpuSkinVertexFactory.usf b/Engine/Shaders/GpuSkinVertexFactory.usf
--- a/Engine/Shaders/GpuSkinVertexFactory.usf
+++ b/Engine/Shaders/GpuSkinVertexFactory.usf
@@ -206,8 +206,7 @@ struct FVertexFactoryIntermediates
 FMaterialVertexParameters GetMaterialVertexParameters(FVertexFactoryInput Input, FVertexFactoryIntermediates Intermediates, float3 WorldPosition, float3x3 TangentToLocal)
 {
 	FMaterialVertexParameters Result = (FMaterialVertexParameters)0;
-	Result.WorldPosition = WorldPosition - View.PreViewTranslation.xyz;
-	Result.PrevWorldPosition = WorldPosition - View.PrevPreViewTranslation.xyz;
+	Result.WorldPosition = WorldPosition;
 	Result.VertexColor = Intermediates.Color;
 	Result.TangentToWorld = mul(TangentToLocal, GetLocalToWorld3x3());
 
diff --git a/Engine/Shaders/LandscapeVertexFactory.usf b/Engine/Shaders/LandscapeVertexFactory.usf
--- a/Engine/Shaders/LandscapeVertexFactory.usf
+++ b/Engine/Shaders/LandscapeVertexFactory.usf
@@ -375,8 +375,7 @@ FMaterialPixelParameters GetMaterialPixelParameters(FVertexFactoryInterpolantsVS
 FMaterialVertexParameters GetMaterialVertexParameters(FVertexFactoryInput Input, FVertexFactoryIntermediates Intermediates, float3 WorldPosition, float3x3 TangentToLocal)
 {
 	FMaterialVertexParameters Result = (FMaterialVertexParameters)0;
-	Result.WorldPosition = WorldPosition - View.PreViewTranslation;
-	Result.PrevWorldPosition = WorldPosition - View.PrevPreViewTranslation.xyz;
+	Result.WorldPosition = WorldPosition;
 	Result.VertexColor = float4(1,1,1,1);
 	Result.TangentToWorld = mul(TangentToLocal, GetLocalToWorld3x3()); 
 
diff --git a/Engine/Shaders/LocalVertexFactory.usf b/Engine/Shaders/LocalVertexFactory.usf
--- a/Engine/Shaders/LocalVertexFactory.usf
+++ b/Engine/Shaders/LocalVertexFactory.usf
@@ -279,8 +279,7 @@ half3x3 CalcTangentToWorldNoScale(in half3x3 TangentToLocal)
 FMaterialVertexParameters GetMaterialVertexParameters(FVertexFactoryInput Input, FVertexFactoryIntermediates Intermediates, float3 WorldPosition, half3x3 TangentToLocal)
 {
 	FMaterialVertexParameters Result = (FMaterialVertexParameters)0;
-	Result.WorldPosition = WorldPosition - View.PreViewTranslation;
-	Result.PrevWorldPosition = WorldPosition - View.PrevPreViewTranslation.xyz;
+	Result.WorldPosition = WorldPosition;
 	Result.VertexColor = Intermediates.Color;
 
 	// does not handle instancing!
diff --git a/Engine/Shaders/MaterialTemplate.usf b/Engine/Shaders/MaterialTemplate.usf
--- a/Engine/Shaders/MaterialTemplate.usf
+++ b/Engine/Shaders/MaterialTemplate.usf
@@ -237,8 +237,10 @@ struct FMaterialTessellationParameters
  */
 struct FMaterialVertexParameters
 {
+	// Position in the translated world (VertexFactoryGetWorldPosition).
+	// Previous position in the translated world (VertexFactoryGetPreviousWorldPosition) if
+	//    computing material's output for previous frame (See {BasePassVertex,Velocity}Shader.usf).
 	float3 WorldPosition;
-	float3 PrevWorldPosition;
 	// TangentToWorld[2] is WorldVertexNormal
 	half3x3 TangentToWorld;
 #if USE_INSTANCING
@@ -263,6 +265,30 @@ struct FMaterialVertexParameters
 	FMaterialParticleParameters Particle;
 };
 
+float3 GetWorldPosition(FMaterialVertexParameters Parameters)
+{
+	return Parameters.WorldPosition;
+}
+
+float3 GetPrevWorldPosition(FMaterialVertexParameters Parameters)
+{
+	// Previous world position and current world position are sharing the
+	// same attribute in Parameters, because in BasePassVertexShader.usf
+	// and in VelocityShader.usf, we are regenerating a Parameters from
+	// VertexFactoryGetPreviousWorldPosition() instead of
+	// VertexFactoryGetWorldPosition().
+	return GetWorldPosition(Parameters);
+}
+
+float3 GetWorldPosition_NoOffsets(FMaterialVertexParameters Parameters)
+{
+	return GetWorldPosition(Parameters) - View.PreViewTranslation;
+}
+
+float3 GetPrevWorldPosition_NoOffsets(FMaterialVertexParameters Parameters)
+{
+	return GetPrevWorldPosition(Parameters) - View.PrevPreViewTranslation;
+}
 
 #if DECAL_RENDERSTAGE && NUM_MATERIAL_TEXCOORDS
 	/*
@@ -813,7 +839,7 @@ float3 GetSpeedTreeVertexOffsetInner(FMaterialVertexParameters Parameters, int G
 		}
 	#else
 		float3x3 LocalToWorld = (float3x3)Primitive.LocalToWorld;
-		float3 LocalPosition = mul(float4(Parameters.WorldPosition.xyz, 1), Primitive.WorldToLocal).xyz;
+		float3 LocalPosition = mul(float4(GetWorldPosition_NoOffsets(Parameters), 1), Primitive.WorldToLocal).xyz;
 	#endif
 
 	float3 TreePos = GetObjectWorldPosition(Parameters);
diff --git a/Engine/Shaders/MeshParticleVertexFactory.usf b/Engine/Shaders/MeshParticleVertexFactory.usf
--- a/Engine/Shaders/MeshParticleVertexFactory.usf
+++ b/Engine/Shaders/MeshParticleVertexFactory.usf
@@ -218,8 +218,7 @@ FMaterialPixelParameters GetMaterialPixelParameters(FVertexFactoryInterpolantsVS
 FMaterialVertexParameters GetMaterialVertexParameters(FVertexFactoryInput Input, FVertexFactoryIntermediates Intermediates, float3 WorldPosition, float3x3 TangentToLocal)
 {
 	FMaterialVertexParameters Result = (FMaterialVertexParameters)0;
-	Result.WorldPosition = WorldPosition - View.PreViewTranslation;
-	Result.PrevWorldPosition = WorldPosition - View.PrevPreViewTranslation.xyz;
+	Result.WorldPosition = WorldPosition;
 	Result.VertexColor = Intermediates.VertexColor;
 
 	Result.Particle.PositionAndSize.xyz = Intermediates.ParticleWorldPosition.xyz;
diff --git a/Engine/Shaders/ParticleBeamTrailVertexFactory.usf b/Engine/Shaders/ParticleBeamTrailVertexFactory.usf
--- a/Engine/Shaders/ParticleBeamTrailVertexFactory.usf
+++ b/Engine/Shaders/ParticleBeamTrailVertexFactory.usf
@@ -91,8 +91,7 @@ FMaterialPixelParameters GetMaterialPixelParameters(FVertexFactoryInterpolantsVS
 FMaterialVertexParameters GetMaterialVertexParameters(FVertexFactoryInput Input, FVertexFactoryIntermediates Intermediates, float3 WorldPosition, float3x3 TangentToLocal)
 {
 	FMaterialVertexParameters Result = (FMaterialVertexParameters)0;
-	Result.WorldPosition = WorldPosition - View.PreViewTranslation.xyz;
-	Result.PrevWorldPosition = WorldPosition - View.PrevPreViewTranslation.xyz;
+	Result.WorldPosition = WorldPosition;
 	Result.VertexColor = Input.Color;
 	Result.TangentToWorld = mul(TangentToLocal, GetLocalToWorld3x3()); 
 	Result.Particle.Color = Intermediates.Color;
diff --git a/Engine/Shaders/ParticleGPUSpriteVertexFactory.usf b/Engine/Shaders/ParticleGPUSpriteVertexFactory.usf
--- a/Engine/Shaders/ParticleGPUSpriteVertexFactory.usf
+++ b/Engine/Shaders/ParticleGPUSpriteVertexFactory.usf
@@ -410,8 +410,7 @@ FMaterialPixelParameters GetMaterialPixelParameters( FVertexFactoryInterpolantsV
 FMaterialVertexParameters GetMaterialVertexParameters(FVertexFactoryInput Input, FVertexFactoryIntermediates Intermediates, float3 WorldPosition, float3x3 TangentToLocal)
 {
 	FMaterialVertexParameters Result = (FMaterialVertexParameters)0;
-	Result.WorldPosition = WorldPosition - View.PreViewTranslation.xyz;
-	Result.PrevWorldPosition = WorldPosition - View.PrevPreViewTranslation.xyz;
+	Result.WorldPosition = WorldPosition;
 	Result.VertexColor = 1;
 	Result.TangentToWorld = mul(TangentToLocal, GetLocalToWorld3x3()); 
 
diff --git a/Engine/Shaders/ParticleSpriteVertexFactory.usf b/Engine/Shaders/ParticleSpriteVertexFactory.usf
--- a/Engine/Shaders/ParticleSpriteVertexFactory.usf
+++ b/Engine/Shaders/ParticleSpriteVertexFactory.usf
@@ -230,8 +230,7 @@ FMaterialPixelParameters GetMaterialPixelParameters(FVertexFactoryInterpolantsVS
 FMaterialVertexParameters GetMaterialVertexParameters(FVertexFactoryInput Input, FVertexFactoryIntermediates Intermediates, float3 WorldPosition, float3x3 TangentToLocal)
 {
 	FMaterialVertexParameters Result = (FMaterialVertexParameters)0;
-	Result.WorldPosition = WorldPosition - View.PreViewTranslation.xyz;
-	Result.PrevWorldPosition = WorldPosition - View.PrevPreViewTranslation.xyz;
+	Result.WorldPosition = WorldPosition;
 	Result.VertexColor = Input.Color;
 	Result.TangentToWorld = mul(TangentToLocal, GetLocalToWorld3x3()); 
 	Result.Particle.MacroUV = SpriteVF.MacroUVParameters;
diff --git a/Engine/Shaders/ShaderVersion.usf b/Engine/Shaders/ShaderVersion.usf
--- a/Engine/Shaders/ShaderVersion.usf
+++ b/Engine/Shaders/ShaderVersion.usf
@@ -2,4 +2,4 @@
 // Each time the console command is executed it generates a new GUID. As this file is included
 // in common.usf (which should be included in any shader) it allows to invalidate the shader DDC.
 // 
-// GUID = 48155F744FBD837A7BD439A9F40EC4D7
\ No newline at end of file
+// GUID = E106A9034435600215322F98E16C9042
\ No newline at end of file
diff --git a/Engine/Shaders/VectorFieldVisualizationVertexFactory.usf b/Engine/Shaders/VectorFieldVisualizationVertexFactory.usf
--- a/Engine/Shaders/VectorFieldVisualizationVertexFactory.usf
+++ b/Engine/Shaders/VectorFieldVisualizationVertexFactory.usf
@@ -109,8 +109,7 @@ FMaterialPixelParameters GetMaterialPixelParameters( FVertexFactoryInterpolantsV
 FMaterialVertexParameters GetMaterialVertexParameters(FVertexFactoryInput Input, FVertexFactoryIntermediates Intermediates, float3 WorldPosition, float3x3 TangentToLocal)
 {
 	FMaterialVertexParameters Result = (FMaterialVertexParameters)0;
-	Result.WorldPosition = WorldPosition - View.PreViewTranslation.xyz;
-	Result.PrevWorldPosition = WorldPosition - View.PrevPreViewTranslation.xyz;
+	Result.WorldPosition = WorldPosition;
 	Result.VertexColor = float4(1,1,1,1);
 	Result.TangentToWorld = mul(TangentToLocal, GetLocalToWorld3x3()); 
 	return Result;
diff --git a/Engine/Source/Runtime/Engine/Private/Materials/HLSLMaterialTranslator.h b/Engine/Source/Runtime/Engine/Private/Materials/HLSLMaterialTranslator.h
--- a/Engine/Source/Runtime/Engine/Private/Materials/HLSLMaterialTranslator.h
+++ b/Engine/Source/Runtime/Engine/Private/Materials/HLSLMaterialTranslator.h
@@ -2102,12 +2102,12 @@ protected:
 			if (bCompilingPreviousFrame && ShaderFrequency == SF_Vertex)
 			{
 				// Only if not in a pixel shader...
-				return AddInlinedCodeChunk(MCT_Float3, TEXT("Parameters.PrevWorldPosition"));
+				return AddInlinedCodeChunk(MCT_Float3, TEXT("GetPrevWorldPosition_NoOffsets(Parameters)"));
 			}
 			else
 			{
 				// If not in a pixel shader, only the normal WorldPosition is available
-				return AddInlinedCodeChunk(MCT_Float3, TEXT("Parameters.WorldPosition"));
+				return AddInlinedCodeChunk(MCT_Float3, TEXT("GetWorldPosition_NoOffsets(Parameters)"));
 			}
 		}
 	}
