Merge PS4 wave compiler support.

diff --git a/Engine/Shaders/Definitions.usf b/Engine/Shaders/Definitions.usf
--- a/Engine/Shaders/Definitions.usf
+++ b/Engine/Shaders/Definitions.usf
@@ -150,6 +150,10 @@
 #define NUM_VF_PACKED_INTERPOLANTS						0
 #endif
 
+#ifndef INVARIANT
+#define INVARIANT										
+#endif
+
 #ifndef USE_CLEARCOAT
 #define USE_CLEARCOAT									0
 #endif
diff --git a/Engine/Shaders/LocalVertexFactory.usf b/Engine/Shaders/LocalVertexFactory.usf
--- a/Engine/Shaders/LocalVertexFactory.usf
+++ b/Engine/Shaders/LocalVertexFactory.usf
@@ -425,14 +425,16 @@ float4 CalcWorldPosition(float4 Position)
 #if USE_INSTANCING
 	return TransformLocalToTranslatedWorld(mul(Position, transpose(InstanceTransform)).xyz);
 #elif USE_SPLINEDEFORM
+/*
 	// Make transform for this point along spline
 	float4x3 SliceTransform = CalcSliceTransform(dot(Position.xyz, SplineMeshDir));
 
 	// Transform into mesh space
 	float4 LocalPos = float4(mul(Position, SliceTransform), Position.w);
-
 	// Transform from mesh to world space
 	return TransformLocalToTranslatedWorld(LocalPos.xyz);
+	*/
+	return INVARIANT(TransformLocalToTranslatedWorld(float3(mul(Position, CalcSliceTransform(dot(Position.xyz, SplineMeshDir))).xyz)));
 #else
 	return TransformLocalToTranslatedWorld(Position.xyz);
 #endif
@@ -504,8 +506,8 @@ FVertexFactoryIntermediates GetVertexFactoryIntermediates(FVertexFactoryInput In
 		float ViewZZero = length(InstanceLocation - InstancingWorldViewOriginZero.xyz) + RandomLOD;
 		float ViewZOne = length(InstanceLocation - InstancingWorldViewOriginOne.xyz) + RandomLOD;
 		Intermediates.PerInstanceParams.w = 
-			dot(ViewZZero > InstancingViewZCompareZero.xyz, InstancingViewZConstant.xyz) * InstancingWorldViewOriginZero.w +
-			dot(ViewZOne > InstancingViewZCompareOne.xyz, InstancingViewZConstant.xyz) * InstancingWorldViewOriginOne.w;
+			dot(float3(ViewZZero.xxx > InstancingViewZCompareZero.xyz), InstancingViewZConstant.xyz) * InstancingWorldViewOriginZero.w +
+			dot(float3(ViewZOne.xxx > InstancingViewZCompareOne.xyz), InstancingViewZConstant.xyz) * InstancingWorldViewOriginOne.w;
 		Intermediates.PerInstanceParams.z *= abs(Intermediates.PerInstanceParams.w) < .999;
 	#else
 		Intermediates.PerInstanceParams.w = 0;
diff --git a/Engine/Shaders/PostProcessSelectionOutline.usf b/Engine/Shaders/PostProcessSelectionOutline.usf
--- a/Engine/Shaders/PostProcessSelectionOutline.usf
+++ b/Engine/Shaders/PostProcessSelectionOutline.usf
@@ -61,7 +61,7 @@ FPixelInfo GetPixelInfo(int2 PixelPos, int SampleID, int OffsetX, int OffsetY, f
 	float DeviceZ = PostprocessInput1MS.Load(PixelPos, SampleID).r;
 	int Stencil = EditorPrimitivesStencil.Load(PixelPos, SampleID).g;
 
-#if !COMPILER_GLSL
+#if !COMPILER_GLSL && !PS4_PROFILE
 	// not yet supported on OpenGL, slightly better quality
 	ReconstructionOffset += PostprocessInput1MS.GetSamplePosition(SampleID);
 #endif
diff --git a/Engine/Shaders/VertexFactoryCommon.usf b/Engine/Shaders/VertexFactoryCommon.usf
--- a/Engine/Shaders/VertexFactoryCommon.usf
+++ b/Engine/Shaders/VertexFactoryCommon.usf
@@ -5,9 +5,9 @@
 float3 TransformLocalToWorld(float3 LocalPosition)
 {
 	//expand the matrix multiply to avoid a code-gen bug on PS4.
-	//float3 RotatedPosition = mul(LocalPosition.xyz, (float3x3)Primitive.LocalToWorld);
-	float3 RotatedPosition = Primitive.LocalToWorld[0].xyz * LocalPosition.xxx + Primitive.LocalToWorld[1].xyz * LocalPosition.yyy + Primitive.LocalToWorld[2].xyz * LocalPosition.zzz;
-	return RotatedPosition + Primitive.LocalToWorld[3].xyz;
+	//float3 RotatedPosition = mul(LocalPosition.xyz, (float3x3)Primitive.LocalToWorld); 
+	//invariant code!
+	return INVARIANT((Primitive.LocalToWorld[0].xyz * LocalPosition.xxx + Primitive.LocalToWorld[1].xyz * LocalPosition.yyy + Primitive.LocalToWorld[2].xyz * LocalPosition.zzz) + Primitive.LocalToWorld[3].xyz);	
 }
 
 float4 TransformLocalToTranslatedWorld(float3 LocalPosition)
