UE4 - vk shader defines & window name

diff --git a/Engine/Shaders/Common.usf b/Engine/Shaders/Common.usf
--- a/Engine/Shaders/Common.usf
+++ b/Engine/Shaders/Common.usf
@@ -101,6 +101,8 @@
 	#define FEATURE_LEVEL FEATURE_LEVEL_SM5
 #elif SM4_PROFILE
 	#define FEATURE_LEVEL FEATURE_LEVEL_SM4
+#elif VULKAN_PROFILE
+	#define FEATURE_LEVEL FEATURE_LEVEL_ES2
 #elif GL3_PROFILE || GL4_PROFILE
 
 	#if GL3_PROFILE
@@ -161,7 +163,7 @@
 #define USE_DEVELOPMENT_SHADERS (COMPILE_SHADERS_FOR_DEVELOPMENT && PLATFORM_SUPPORTS_DEVELOPMENT_SHADERS)
 
 #ifndef MOBILE_EMULATION
-#define MOBILE_EMULATION ((FEATURE_LEVEL == FEATURE_LEVEL_ES2 || FEATURE_LEVEL == FEATURE_LEVEL_ES3_1) && (!COMPILER_GLSL_ES2 && USE_DEVELOPMENT_SHADERS && !METAL_PROFILE))
+#define MOBILE_EMULATION ((FEATURE_LEVEL == FEATURE_LEVEL_ES2 || FEATURE_LEVEL == FEATURE_LEVEL_ES3_1) && (!COMPILER_GLSL_ES2 && USE_DEVELOPMENT_SHADERS && !METAL_PROFILE && !VULKAN_PROFILE))
 #endif
 
 // If we didn't request color component swizzling, just make it empty
diff --git a/Engine/Shaders/Definitions.usf b/Engine/Shaders/Definitions.usf
--- a/Engine/Shaders/Definitions.usf
+++ b/Engine/Shaders/Definitions.usf
@@ -100,6 +100,10 @@
 #define METAL_MRT_PROFILE								0
 #endif
 
+#ifndef VULKAN_PROFILE
+#define	VULKAN_PROFILE									0
+#endif
+
 #ifndef IOS
 #define IOS												0
 #endif
diff --git a/Engine/Shaders/TonemapCommon.usf b/Engine/Shaders/TonemapCommon.usf
--- a/Engine/Shaders/TonemapCommon.usf
+++ b/Engine/Shaders/TonemapCommon.usf
@@ -5,7 +5,7 @@
 =============================================================================*/
 
 // #todo-rco: Fix in the cross-compiler
-#if COMPILER_GLSL || COMPILER_GLSL_ES2 || METAL_PROFILE
+#if COMPILER_GLSL || COMPILER_GLSL_ES2 || METAL_PROFILE || VULKAN_PROFILE
 #define log10(x)  log(x) / log(10.0)
 #endif
 
diff --git a/Engine/Source/Editor/DetailCustomizations/Private/WindowsTargetSettingsDetails.cpp b/Engine/Source/Editor/DetailCustomizations/Private/WindowsTargetSettingsDetails.cpp
--- a/Engine/Source/Editor/DetailCustomizations/Private/WindowsTargetSettingsDetails.cpp
+++ b/Engine/Source/Editor/DetailCustomizations/Private/WindowsTargetSettingsDetails.cpp
@@ -41,6 +41,10 @@ FText GetFriendlyNameFromRHIName(const FString& InRHIName)
 	{
 		FriendlyRHIName = LOCTEXT("OpenGL4", "OpenGL 4 (SM5, Experimental)");
 	}
+	else if (InRHIName == TEXT("PC_VULKAN_ES2"))
+	{
+		FriendlyRHIName = LOCTEXT("VulkanES2", "Vulkan Mobile(ES2, Experimental)");
+	}
 
 	return FriendlyRHIName;
 }
