Fixes the shader compilation failure of WorldPosition material expression node in tessellation shaders

diff --git a/Engine/Shaders/MaterialTemplate.usf b/Engine/Shaders/MaterialTemplate.usf
--- a/Engine/Shaders/MaterialTemplate.usf
+++ b/Engine/Shaders/MaterialTemplate.usf
@@ -224,6 +224,7 @@ struct FMaterialTessellationParameters
 	float2 TexCoords[NUM_MATERIAL_TEXCOORDS];
 #endif
 	float4 VertexColor;
+	// TODO: Non translated world position
 	float3 WorldPosition;
 	float3 TangentToWorldPreScale;
 
@@ -290,6 +291,12 @@ float3 GetPrevWorldPosition_NoOffsets(FMaterialVertexParameters Parameters)
 	return GetPrevWorldPosition(Parameters) - View.PrevPreViewTranslation;
 }
 
+//TODO(bug UE-17131): We should compute world displacement for the previous frame
+float3 GetWorldPosition_NoOffsets(FMaterialTessellationParameters Parameters)
+{
+	return Parameters.WorldPosition;
+}
+
 #if DECAL_RENDERSTAGE && NUM_MATERIAL_TEXCOORDS
 	/*
 	 * Material node DecalMipmapLevel's code designed to avoid the 2x2 pixels artefacts on the edges around where the decal
