Fix Inverse Squared point lights on ES2

diff --git a/Engine/Shaders/BasePassForForwardShadingPixelShader.usf b/Engine/Shaders/BasePassForForwardShadingPixelShader.usf
--- a/Engine/Shaders/BasePassForForwardShadingPixelShader.usf
+++ b/Engine/Shaders/BasePassForForwardShadingPixelShader.usf
@@ -277,11 +277,25 @@ void Main(
 				half PointNoL = max(0, dot(MaterialParameters.WorldNormal, L));
 				half PointRoL = max(0, dot(MaterialParameters.ReflectionVector, L));
 
-				half DistanceAttenuation = RadialAttenuation(ToLight * LightPositionAndInvRadius[i].w, LightColorAndFalloffExponent[i].w);		
+				half Attenuation;
+
+				if (LightColorAndFalloffExponent[i].w == 0)
+				{
+					// Sphere falloff (technically just 1/d2 but this avoids inf)
+					Attenuation = 1 / ( DistanceSqr + 1 );
+	
+					float LightRadiusMask = Square(saturate(1 - Square(DistanceSqr * (LightPositionAndInvRadius[i].w * LightPositionAndInvRadius[i].w))));
+					Attenuation *= LightRadiusMask;
+				}
+				else
+				{
+					Attenuation = RadialAttenuation(ToLight * LightPositionAndInvRadius[i].w, LightColorAndFalloffExponent[i].w);		
+				}
+
 				#if !MATERIAL_FULLY_ROUGH
-					Color += (DistanceAttenuation * PointNoL) * LightColorAndFalloffExponent[i].rgb * (DiffuseColor + SpecularColor * PhongApprox(Roughness, PointRoL));
+					Color += (Attenuation * PointNoL) * LightColorAndFalloffExponent[i].rgb * (DiffuseColor + SpecularColor * PhongApprox(Roughness, PointRoL));
 				#else
-					Color += (DistanceAttenuation * PointNoL) * LightColorAndFalloffExponent[i].rgb * DiffuseColor;
+					Color += (Attenuation * PointNoL) * LightColorAndFalloffExponent[i].rgb * DiffuseColor;
 				#endif
 			}
 		#endif
