Fix render problems with large radius mobile point lights due to insufficient precision

diff --git a/Engine/Shaders/BasePassForForwardShadingPixelShader.usf b/Engine/Shaders/BasePassForForwardShadingPixelShader.usf
--- a/Engine/Shaders/BasePassForForwardShadingPixelShader.usf
+++ b/Engine/Shaders/BasePassForForwardShadingPixelShader.usf
@@ -270,8 +270,8 @@ void Main(
 				for (int i = 0; i < NUM_DYNAMIC_POINT_LIGHTS; i++)
 			#endif
 			{
-				half3 ToLight = LightPositionAndInvRadius[i].xyz - MaterialParameters.WorldPosition;
-				half DistanceSqr = dot(ToLight, ToLight);
+				float3 ToLight = LightPositionAndInvRadius[i].xyz - MaterialParameters.WorldPosition;
+				float DistanceSqr = dot(ToLight, ToLight);
 				half3 L = ToLight * rsqrt(DistanceSqr);
 
 				half PointNoL = max(0, dot(MaterialParameters.WorldNormal, L));
