Default gradient is non-zero vector

diff --git a/Engine/Shaders/GlobalDistanceFieldShared.usf b/Engine/Shaders/GlobalDistanceFieldShared.usf
--- a/Engine/Shaders/GlobalDistanceFieldShared.usf
+++ b/Engine/Shaders/GlobalDistanceFieldShared.usf
@@ -92,7 +92,7 @@ float3 GetDistanceFieldGradientGlobalClipmap(float3 WorldPosition, uint ClipmapI
 
 float3 GetDistanceFieldGradientGlobal(float3 WorldPosition)
 {
-	float3 Gradient = 0;
+	float3 Gradient = float3(0, 0, .001f);
 	float DistanceFromClipmap = ComputeDistanceFromBoxToPointInside(GlobalVolumeCenterAndExtent[0].xyz, GlobalVolumeCenterAndExtent[0].www, WorldPosition);
 	// Don't sample near the border to avoid sampling out of bounds during partial differencing
 	float BorderTexels = GlobalVolumeTexelSize * 4;
@@ -140,7 +140,7 @@ float GetDistanceToNearestSurfaceGlobal(float3 WorldPosition)
 
 float3 GetDistanceFieldGradientGlobal(float3 WorldPosition)
 {
-	return float3(0, 0, 1);
+	return float3(0, 0, .001f);
 }
 
 #endif
