Fixed compute shader material compile error

diff --git a/Engine/Shaders/Common.usf b/Engine/Shaders/Common.usf
--- a/Engine/Shaders/Common.usf
+++ b/Engine/Shaders/Common.usf
@@ -883,7 +883,7 @@ float AntialiasedTextureMask( Texture2D Tex, SamplerState Sampler, float2 UV, fl
 		float Sample1 = dot(MaskConst, Texture2DSample(Tex, Sampler, UV));
 
 		// compute the derivatives of the texture content
-		float2 TexDD = float2(ddx(Sample1), ddy(Sample1));
+		float2 TexDD = float2(DDX(Sample1), DDY(Sample1));
 
 		float TexDDLength = max(abs(TexDD.x), abs(TexDD.y)); 
 		float Top = InvWidthConst * (Sample1 - ThresholdConst);
