Integrate CL 2560104

diff --git a/Engine/Shaders/ScreenSpaceRayCast.usf b/Engine/Shaders/ScreenSpaceRayCast.usf
--- a/Engine/Shaders/ScreenSpaceRayCast.usf
+++ b/Engine/Shaders/ScreenSpaceRayCast.usf
@@ -339,10 +339,9 @@ float4 SampleScreenColor( Texture2D Texture, SamplerState Sampler, float3 HitUVz
 	// TODO combine to single matrix
 	float4 HitClip = float4( (HitUVz.xy - View.ScreenPositionScaleBias.wz) / View.ScreenPositionScaleBias.xy, HitUVz.z, 1 );
 	float4 HitTranslatedWorld = mul( HitClip, View.ClipToTranslatedWorld );
-	HitTranslatedWorld /= HitTranslatedWorld.w;
 		
-	float3 PrevTranslatedWorld = HitTranslatedWorld.xyz + (View.PrevPreViewTranslation - View.PreViewTranslation);
-	float4 PrevClip = mul( float4( PrevTranslatedWorld, 1 ), View.PrevTranslatedWorldToClip );
+	float4 PrevTranslatedWorld = float4( HitTranslatedWorld.xyz + HitTranslatedWorld.w * (View.PrevPreViewTranslation - View.PreViewTranslation), HitTranslatedWorld.w );
+	float4 PrevClip = mul( PrevTranslatedWorld, View.PrevTranslatedWorldToClip );
 	float2 PrevScreen = PrevClip.xy / PrevClip.w;
 	float2 PrevUV = PrevScreen.xy * View.ScreenPositionScaleBias.xy + View.ScreenPositionScaleBias.wz;
 
