Fix Texture updating shaders for PS4

diff --git a/Engine/Shaders/UpdateTextureShaders.usf b/Engine/Shaders/UpdateTextureShaders.usf
--- a/Engine/Shaders/UpdateTextureShaders.usf
+++ b/Engine/Shaders/UpdateTextureShaders.usf
@@ -8,9 +8,7 @@
 
 Buffer<int> SrcBuffer;
 RWTexture2D<int4> DestTexture;
-uint SrcPitch; //number of SrcBuffer elements in a row of the source area
-
-
+uint2 SrcPitch; //x = number of SrcBuffer elements in a row of the source area, y = Number of components per texel in the source.
 uint4 DestPosSize; // xy = pos, zw = Size of the destination sub rect
 
 [numthreads(8,8,1)]
@@ -21,7 +19,7 @@ void UpdateTexture2DSubresourceCS( uint2 ThreadId : SV_DispatchThreadID )
 		int2 TexturePixelOffset = ThreadId.xy + DestPosSize.xy;
 
 		//supports textures with 1-4 components.  Textures with  < 4 components will simply ignore writes to the extra components via channel masking.		
-		int BufferOffset = ThreadId.x + (ThreadId.y * SrcPitch);
+		int BufferOffset = (ThreadId.x * SrcPitch.y) + (ThreadId.y * SrcPitch.x);
 		DestTexture[TexturePixelOffset] = int4(SrcBuffer[BufferOffset].x, SrcBuffer[BufferOffset + 1].x, SrcBuffer[BufferOffset + 2].x, SrcBuffer[BufferOffset + 3].x);
 	}
 }
@@ -39,8 +37,8 @@ void UpdateTexture3DSubresourceCS(uint3 ThreadId : SV_DispatchThreadID)
 		int3 TexturePixelOffset = ThreadId.xyz + (int3)DestPos.xyz;
 
 		//supports textures with 1-4 components.  Textures with  < 4 components will simply ignore writes to the extra components via channel masking.
-		int BufferOffset = ThreadId.x + (ThreadId.y * SrcPitch) + (ThreadId.z * SrcDepthPitch);		
-		DestTexture3D[TexturePixelOffset] = int4(SrcBuffer[BufferOffset].x, SrcBuffer[BufferOffset + 1].x, SrcBuffer[BufferOffset + 2].x, SrcBuffer[BufferOffset + 3].x);
+		int BufferOffset = (ThreadId.x * SrcPitch.y) + (ThreadId.y * SrcPitch.x) + (ThreadId.z * SrcDepthPitch);
+		DestTexture3D[TexturePixelOffset] = int4(SrcBuffer[BufferOffset].x, SrcBuffer[BufferOffset + 1].x, SrcBuffer[BufferOffset + 2].x, SrcBuffer[BufferOffset + 3].x);		
 	}
 }
 
