Cleans up experimental code for panini projection.

diff --git a/Engine/Shaders/Common.usf b/Engine/Shaders/Common.usf
--- a/Engine/Shaders/Common.usf
+++ b/Engine/Shaders/Common.usf
@@ -79,6 +79,9 @@
 
 // to support the console command "r.InvalidateShaderCache"
 #include "ShaderVersion.usf"
+	
+// uniform buffers specifics
+#include "CommonViewUniformBuffer.usf"
 
 #include "Definitions.usf"
 
diff --git a/Engine/Shaders/CommonViewUniformBuffer.usf b/Engine/Shaders/CommonViewUniformBuffer.usf
new file mode 100644
--- /dev/null
+++ b/Engine/Shaders/CommonViewUniformBuffer.usf
@@ -0,0 +1,21 @@
+// Copyright 1998-2015 Epic Games, Inc. All Rights Reserved.
+
+/*=============================================================================
+	CommonViewUniformBuffer.usf: Common view uniform buffer specifics
+=============================================================================*/
+
+/*
+* @return tan(View.FieldOfViewWideAngles * .5)
+*/
+float2 GetTanHalfFieldOfView()
+{
+	return float2(View.ClipToView[0][0], View.ClipToView[1][1]);
+}
+
+/*
+* @return 1 / tan(View.FieldOfViewWideAngles * .5)
+*/
+float2 GetCotanHalfFieldOfView()
+{
+	return float2(View.ViewToClip[0][0], View.ViewToClip[1][1]);
+}
diff --git a/Engine/Shaders/PaniniProjection.usf b/Engine/Shaders/PaniniProjection.usf
--- a/Engine/Shaders/PaniniProjection.usf
+++ b/Engine/Shaders/PaniniProjection.usf
@@ -18,20 +18,20 @@
 *                                        +
 *                                      / |
 *                        ->          /   |
-*                        z         /     |
-*        - + - - - - - - < - - - + - - - + - - + - - - - - - -
+*                        z         /     |                   (view direction)
+*        - + - - - - - - < - - - + - - - + - - + - - - - - - - - >
 *          D                     O       |A
 *                                        |
 *                                        |
 *                                (panini view plan)
 *
-*		O: origine of the frame
+*		O: origin of the frame
 *		A: origin of the panini view plane (0.0, 0.0, -1.0)
 *		D: panini projection center (0.0, 0.0, d)
 *		OM: the direction of the ray.
 *
-* @param OM - The direction of the ray in the panini frame. The z component is implicitely -1.0f.
-* @param d - The panini projection center's distance form the panini frame's origin.
+* @param OM - The direction of the ray in the panini frame. The z component is implicitely -1.0f in the panini frame.
+* @param d - The panini projection center's distance from the panini frame's origin.
 *		0.0f: Panini projection doesn't change anything, it will return PaniniDirection.
 *		1.0f: Panini
 *		@assert d>=0
@@ -42,7 +42,7 @@
 */
 float2 PaniniProjection(float2 OM, float d, float s)
 {
-	float PaniniDirectionXZInvLength = 1.0f / sqrt(1.0f + OM.x * OM.x);
+	float PaniniDirectionXZInvLength = rsqrt(1.0f + OM.x * OM.x);
 	float SinPhi = OM.x * PaniniDirectionXZInvLength;
 	float TanTheta = OM.y * PaniniDirectionXZInvLength;
 	float CosPhi = sqrt(1.0f - SinPhi * SinPhi);
@@ -50,41 +50,3 @@ float2 PaniniProjection(float2 OM, float d, float s)
 
 	return S * float2(SinPhi, lerp(TanTheta, TanTheta / CosPhi, s));
 }
-
-float2 GetHalfFOVPerAxis()
-{
-	// TODO: should be computed on the CPU side
-	float3 VCenter = normalize(mul(float4(0.0, 0.0, 0.0, 1.0), View.ClipToView)).xyz;
-	float3 VUp = normalize(mul(float4(0.0, 1.0, 0.0, 1.0), View.ClipToView)).xyz;
-	float3 VRight = normalize(mul(float4(1.0, 0.0, 0.0, 1.0), View.ClipToView)).xyz;
-
-	return float2(acos(dot(VCenter, VRight)), acos(dot(VCenter, VUp)));
-}
-
-/*
-* Utility to do the projection 
-*
-*/
-float2 ScreenPosPaniniProjection(float2 ScreenPos, float d, float s)
-{
-	float2 HalfFOV = GetHalfFOVPerAxis();
-	float2 ScreenPosToPaniniFactor = float2(tan(HalfFOV.x), tan(HalfFOV.y));
-
-	// Compute the pixel ray direction into the Panini frame
-	// PaniniDirection.z is implicitly 1.0f
-	float2 PaniniDirection = ScreenPos * ScreenPosToPaniniFactor;
-	float2 PaniniPosition = PaniniProjection(PaniniDirection, d, s);
-
-	// Return the new position back in the screen space frame
-	return PaniniPosition / ScreenPosToPaniniFactor;
-}
-
-// Panini material function node returning a not translated world position according to the panini projection.
-float3 CalcPaniniWorldPosition(float3 InWorldPosition, float d, float s, float PostScale)
-{
-	float4 InClipPosition = mul(float4(InWorldPosition + View.PreViewTranslation, 1.0), View.TranslatedWorldToClip);
-	float2 InScreenPos = InClipPosition.xy / InClipPosition.w;
-	float2 OutScreenPos = ScreenPosPaniniProjection(InScreenPos, d, s) * PostScale;
-	float4 OutClipPosition = float4(OutScreenPos.xy * InClipPosition.w, InClipPosition.z, InClipPosition.w);
-	return mul(OutClipPosition, View.ClipToTranslatedWorld).xyz - View.PreViewTranslation;
-}
diff --git a/Engine/Shaders/PostProcessUpscale.usf b/Engine/Shaders/PostProcessUpscale.usf
--- a/Engine/Shaders/PostProcessUpscale.usf
+++ b/Engine/Shaders/PostProcessUpscale.usf
@@ -35,15 +35,27 @@ float Gaussian(float Scale, float2 Offset)
 	return exp2(Scale * dot(Offset, Offset));
 }
 
-// @param In -1 .. 1 range
-// @return In -1 .. 1 range
-float2 PaniniDistortPos(float2 In)
+// Panini projection in the screen pos
+// @param InScreenPos: the position to project in the screen space
+// @return the panini projected postion in screen space
+float2 PaniniProjectionScreenPos(float2 InScreenPos)
 {
 	const float d = PaniniParams.x;
 	const float s = PaniniParams.y;
 	const float OutScreenPosScale = PaniniParams.z;
+#if 0
+	const float2 ScreenSpaceToPaniniFactor = tan(0.5 * View.FieldOfViewWideAngles);
+	const float2 PaniniToScreenSpaceFactor = 1 / ScreenSpaceToPaniniFactor;
+#else
+	const float2 ScreenSpaceToPaniniFactor = GetTanHalfFieldOfView();
+	const float2 PaniniToScreenSpaceFactor = GetCotanHalfFieldOfView();
+#endif
+
+	const float2 PaniniDirection = InScreenPos * ScreenSpaceToPaniniFactor;
+	const float2 PaniniPosition = PaniniProjection(PaniniDirection, d, s);
 
-	return ScreenPosPaniniProjection(In, d, s) * OutScreenPosScale;
+	// Return the new position back in the screen space frame
+	return PaniniPosition * PaniniToScreenSpaceFactor * OutScreenPosScale;
 }
 
 
@@ -65,7 +77,7 @@ void MainVS(
 	float2 PlanarScreenPos = TexCoord * 2 - 1;
 
 	// distort pos
-	Position.xy = PaniniDistortPos(PlanarScreenPos) * 0.5f + 0.5f;
+	Position.xy = PaniniProjectionScreenPos(PlanarScreenPos) * 0.5f + 0.5f;
 
 	DrawRectangle(Position, TexCoord, OutPosition, OutTexCoord.xy);
 	OutTexCoord.zw =  OutPosition.xy;
