enable experimentation with RE_Z

diff --git a/Engine/Shaders/Definitions.usf b/Engine/Shaders/Definitions.usf
--- a/Engine/Shaders/Definitions.usf
+++ b/Engine/Shaders/Definitions.usf
@@ -158,6 +158,10 @@
 #define INVARIANT										
 #endif
 
+#ifndef ENABLE_RE_Z
+#define ENABLE_RE_Z
+#endif
+
 #ifndef USE_CLEARCOAT
 #define USE_CLEARCOAT									0
 #endif
