Exposes the view property's inverse and the TanHalfFieldOfView view property in the material editor

diff --git a/Engine/Shaders/BasePassForForwardShadingPixelShader.usf b/Engine/Shaders/BasePassForForwardShadingPixelShader.usf
--- a/Engine/Shaders/BasePassForForwardShadingPixelShader.usf
+++ b/Engine/Shaders/BasePassForForwardShadingPixelShader.usf
@@ -86,7 +86,7 @@ float ClipForEditorPrimitives(FMaterialPixelParameters MaterialParameters)
 			bool bIsPerspective = (View.ViewToClip._m33 < 1.0f);
 
 			// dejitter the sample position and make a filtered lookup - for planes this allows to reconstruct a much less jittery depth comparison function, it however doesn't fix silhuetes
-			float DeviceZ = Texture2DSampleLevel(FilteredSceneDepthTexture, FilteredSceneDepthTextureSampler, (MaterialParameters.SVPosition.xy - View.TemporalAAParams.zw) * View.ViewSizeAndSceneTexelSize.zw, 0).r;
+			float DeviceZ = Texture2DSampleLevel(FilteredSceneDepthTexture, FilteredSceneDepthTextureSampler, (MaterialParameters.SVPosition.xy - View.TemporalAAParams.zw) * View.BufferSizeAndInvSize.zw, 0).r;
 
 			float PixelDeviceZ = MaterialParameters.SVPosition.z;
 
diff --git a/Engine/Shaders/BasePassPixelShader.usf b/Engine/Shaders/BasePassPixelShader.usf
--- a/Engine/Shaders/BasePassPixelShader.usf
+++ b/Engine/Shaders/BasePassPixelShader.usf
@@ -540,7 +540,7 @@ float ClipForEditorPrimitives(FMaterialPixelParameters MaterialParameters)
 #endif	// HAS_INVERTED_Z_BUFFER
  		
 		// dejitter the sample position and make a filtered lookup - for planes this allows to reconstruct a much less jittery depth comparison function, it however doesn't fix silhuetes
-		float DeviceZ = Texture2DSampleLevel(FilteredSceneDepthTexture, FilteredSceneDepthTextureSampler, (MaterialParameters.SVPosition.xy - View.TemporalAAParams.zw) * View.ViewSizeAndSceneTexelSize.zw, 0).r;
+		float DeviceZ = Texture2DSampleLevel(FilteredSceneDepthTexture, FilteredSceneDepthTextureSampler, (MaterialParameters.SVPosition.xy - View.TemporalAAParams.zw) * View.BufferSizeAndInvSize.zw, 0).r;
 
 		float PixelDeviceZ = MaterialParameters.SVPosition.z;
 
diff --git a/Engine/Shaders/CommonViewUniformBuffer.usf b/Engine/Shaders/CommonViewUniformBuffer.usf
--- a/Engine/Shaders/CommonViewUniformBuffer.usf
+++ b/Engine/Shaders/CommonViewUniformBuffer.usf
@@ -12,6 +12,11 @@ float2 GetTanHalfFieldOfView()
 	return float2(View.ClipToView[0][0], View.ClipToView[1][1]);
 }
 
+float2 GetPrevTanHalfFieldOfView()
+{
+	return float2(View.PrevClipToView[0][0], View.PrevClipToView[1][1]);
+}
+
 /*
 * @return 1 / tan(View.FieldOfViewWideAngles * .5)
 */
@@ -19,3 +24,8 @@ float2 GetCotanHalfFieldOfView()
 {
 	return float2(View.ViewToClip[0][0], View.ViewToClip[1][1]);
 }
+
+float2 GetPrevCotanHalfFieldOfView()
+{
+	return float2(View.PrevViewToClip[0][0], View.PrevViewToClip[1][1]);
+}
diff --git a/Engine/Shaders/DistanceFieldAOShared.usf b/Engine/Shaders/DistanceFieldAOShared.usf
--- a/Engine/Shaders/DistanceFieldAOShared.usf
+++ b/Engine/Shaders/DistanceFieldAOShared.usf
@@ -220,7 +220,7 @@ float2 GetBaseLevelScreenUVFromScreenGrid(uint2 OutputCoordinate)
 
 float2 GetScreenUVFromScreenGrid(uint2 OutputCoordinate)
 {
-	float2 ScreenUV = ((OutputCoordinate * TRACE_DOWNSAMPLE_FACTOR + JitterOffset) * AO_DOWNSAMPLE_FACTOR + View.ViewRectMin.xy + float2(.5f, .5f)) * View.ViewSizeAndSceneTexelSize.zw;
+	float2 ScreenUV = ((OutputCoordinate * TRACE_DOWNSAMPLE_FACTOR + JitterOffset) * AO_DOWNSAMPLE_FACTOR + View.ViewRectMin.xy + float2(.5f, .5f)) * View.BufferSizeAndInvSize.zw;
 	return ScreenUV;
 }
 
diff --git a/Engine/Shaders/DistanceFieldLightingPost.usf b/Engine/Shaders/DistanceFieldLightingPost.usf
--- a/Engine/Shaders/DistanceFieldLightingPost.usf
+++ b/Engine/Shaders/DistanceFieldLightingPost.usf
@@ -250,7 +250,7 @@ void UpdateHistoryPS(
 	)
 {
 	// Distance field AO was computed at 0,0 regardless of viewrect min
-	float2 DistanceFieldUVs = UVAndScreenPos.xy - View.ViewRectMin.xy * View.ViewSizeAndSceneTexelSize.zw;
+	float2 DistanceFieldUVs = UVAndScreenPos.xy - View.ViewRectMin.xy * View.BufferSizeAndInvSize.zw;
 	float4 NewValue = Texture2DSampleLevel(BentNormalAOTexture, BentNormalAOSampler, DistanceFieldUVs, 0);
 	float SceneDepth = abs(NewValue.w);
 
@@ -269,7 +269,7 @@ void UpdateHistoryPS(
 
 		// Note: have to sample from one of the high res texels exactly to avoid filtering of the velocity buffer
 		// 0 is stored where camera velocity should be used, which must not be filtered with valid object velocities
-		float2 FullResTexel = UVAndScreenPos.xy - .5f * View.ViewSizeAndSceneTexelSize.zw;
+		float2 FullResTexel = UVAndScreenPos.xy - .5f * View.BufferSizeAndInvSize.zw;
 		float2 VelocityN = Texture2DSampleLevel(VelocityTexture, VelocityTextureSampler, FullResTexel, 0).xy;
 
 		if (VelocityN.x > 0)
@@ -283,12 +283,12 @@ void UpdateHistoryPS(
 	
 	float2 OldScreenPosition = (OldUV.xy - View.ScreenPositionScaleBias.wz) / View.ScreenPositionScaleBias.xy;
 
-	float2 OldDistanceFieldUVs = OldUV.xy - View.ViewRectMin.xy * View.ViewSizeAndSceneTexelSize.zw;
+	float2 OldDistanceFieldUVs = OldUV.xy - View.ViewRectMin.xy * View.BufferSizeAndInvSize.zw;
 	float4 HistoryValue = Texture2DSampleLevel(BentNormalHistoryTexture, BentNormalHistorySampler, OldDistanceFieldUVs, 0);
 	
-	float2 MinUV = View.ViewRectMin.xy * View.ViewSizeAndSceneTexelSize.zw;
+	float2 MinUV = View.ViewRectMin.xy * View.BufferSizeAndInvSize.zw;
 	// Pull in the max UV to exclude the region which will read outside the viewport due to bilinear filtering
-	float2 MaxUV = MinUV + (View.ViewSizeAndSceneTexelSize.xy - 1 * DOWNSAMPLE_FACTOR) * View.ViewSizeAndSceneTexelSize.zw;
+	float2 MaxUV = MinUV + (View.ViewSizeAndInvSize.xy - 1 * DOWNSAMPLE_FACTOR) * View.BufferSizeAndInvSize.zw;
 
 	float3 WorldPosition;
 	float3 PrevWorldPosition;
@@ -425,7 +425,7 @@ void AOUpsamplePS(
 	)
 {
 	// Distance field AO was computed at 0,0 regardless of viewrect min
-	float2 DistanceFieldUVs = UVAndScreenPos.xy - View.ViewRectMin.xy * View.ViewSizeAndSceneTexelSize.zw;
+	float2 DistanceFieldUVs = UVAndScreenPos.xy - View.ViewRectMin.xy * View.BufferSizeAndInvSize.zw;
 	float3 Irradiance = 0;
 
 #define BILATERAL_UPSAMPLE 1
diff --git a/Engine/Shaders/DistanceFieldShadowing.usf b/Engine/Shaders/DistanceFieldShadowing.usf
--- a/Engine/Shaders/DistanceFieldShadowing.usf
+++ b/Engine/Shaders/DistanceFieldShadowing.usf
@@ -385,8 +385,8 @@ void CullObjectsToTileWithGather(
 	float2 TanViewFOV = float2(1 / View.ViewToClip[0][0], 1 / View.ViewToClip[1][1]);
 	// tan(FOV) = HalfUnitPlaneWidth / 1, so TanViewFOV * 2 is the size of the whole unit view plane
 	// We are operating on a subset of that defined by ScissorRectMinAndSize
-	float2 TileSize = TanViewFOV * 2 * ScissorRectMinAndSize.zw / ((float2)View.ViewSizeAndSceneTexelSize.xy * NumGroups);
-	float2 UnitPlaneMin = -TanViewFOV + TanViewFOV * 2 * ScissorRectMinAndSize.xy / (float2)View.ViewSizeAndSceneTexelSize.xy;
+	float2 TileSize = TanViewFOV * 2 * ScissorRectMinAndSize.zw / ((float2)View.ViewSizeAndInvSize.xy * NumGroups);
+	float2 UnitPlaneMin = -TanViewFOV + TanViewFOV * 2 * ScissorRectMinAndSize.xy * View.ViewSizeAndInvSize.zw;
 	float2 UnitPlaneTileMin = (GroupId.xy * TileSize + UnitPlaneMin) * float2(1, -1);
 	float2 UnitPlaneTileMax = ((GroupId.xy + 1) * TileSize + UnitPlaneMin) * float2(1, -1);
 
@@ -481,7 +481,7 @@ void DistanceFieldShadowingCS(
 {
 	uint ThreadIndex = GroupThreadId.y * THREADGROUP_SIZEX + GroupThreadId.x;
 
-	float2 ScreenUV = float2((DispatchThreadId.xy * DownsampleFactor + ScissorRectMinAndSize.xy + .5f) * View.ViewSizeAndSceneTexelSize.zw);
+	float2 ScreenUV = float2((DispatchThreadId.xy * DownsampleFactor + ScissorRectMinAndSize.xy + .5f) * View.BufferSizeAndInvSize.zw);
 	float2 ScreenPosition = (ScreenUV.xy - View.ScreenPositionScaleBias.wz) / View.ScreenPositionScaleBias.xy;
 
 	float SceneDepth = CalcSceneDepth(ScreenUV);
@@ -610,7 +610,7 @@ void DistanceFieldShadowingUpsamplePS(
 	out float4 OutColor : SV_Target0)
 {
 	// Distance field shadowing was computed at 0,0 regardless of viewrect min
-	float2 DistanceFieldUVs = UVAndScreenPos.xy - ScissorRectMinAndSize.xy * View.ViewSizeAndSceneTexelSize.zw;
+	float2 DistanceFieldUVs = UVAndScreenPos.xy - ScissorRectMinAndSize.xy * View.BufferSizeAndInvSize.zw;
 	float SceneDepth = CalcSceneDepth(UVAndScreenPos.xy);
 
 #define BILATERAL_UPSAMPLE 1
diff --git a/Engine/Shaders/DistanceFieldSurfaceCacheLighting.usf b/Engine/Shaders/DistanceFieldSurfaceCacheLighting.usf
--- a/Engine/Shaders/DistanceFieldSurfaceCacheLighting.usf
+++ b/Engine/Shaders/DistanceFieldSurfaceCacheLighting.usf
@@ -123,7 +123,7 @@ void ComputeDistanceFieldNormalPS(
 	out float4 OutColor : SV_Target0)
 {
 	// Sample from the center of the top left full resolution texel
-	float2 ScreenUV = float2((floor(SVPos.xy) * DOWNSAMPLE_FACTOR + View.ViewRectMin.xy + .5f) * View.ViewSizeAndSceneTexelSize.zw);
+	float2 ScreenUV = float2((floor(SVPos.xy) * DOWNSAMPLE_FACTOR + View.ViewRectMin.xy + .5f) * View.BufferSizeAndInvSize.zw);
 	float SceneDepth = CalcSceneDepth(ScreenUV);
 	FGBufferData GBufferData = GetGBufferData(ScreenUV);
 
@@ -138,7 +138,7 @@ void ComputeDistanceFieldNormalCS(
 	uint3 DispatchThreadId : SV_DispatchThreadID,
     uint3 GroupThreadId : SV_GroupThreadID) 
 {
-	float2 ScreenUV = float2((DispatchThreadId.xy * DOWNSAMPLE_FACTOR + View.ViewRectMin.xy + .5f) * View.ViewSizeAndSceneTexelSize.zw);
+	float2 ScreenUV = float2((DispatchThreadId.xy * DOWNSAMPLE_FACTOR + View.ViewRectMin.xy + .5f) * View.BufferSizeAndInvSize.zw);
 	float SceneDepth = CalcSceneDepth(ScreenUV);
 	FGBufferData GBufferData = GetGBufferData(ScreenUV);
 
@@ -293,7 +293,7 @@ void IrradianceCacheSplatPS(
 
 #else
 
-	float2 ScreenUV = float2(((floor(SVPos.xy)) * CurrentLevelDownsampleFactor + View.ViewRectMin.xy + float2(.5f, .5f)) * View.ViewSizeAndSceneTexelSize.zw);
+	float2 ScreenUV = float2(((floor(SVPos.xy)) * CurrentLevelDownsampleFactor + View.ViewRectMin.xy + float2(.5f, .5f)) * View.BufferSizeAndInvSize.zw);
 	SceneDepth = CalcSceneDepth(ScreenUV); 
 	float2 ScreenPosition = (ScreenUV.xy - View.ScreenPositionScaleBias.wz) / View.ScreenPositionScaleBias.xy;
 	float4 HomogeneousWorldPosition = mul(float4(ScreenPosition.xy * SceneDepth, SceneDepth, 1), View.ScreenToWorld);
@@ -507,7 +507,7 @@ void VisualizeMeshDistanceFieldCS(
 {
 	uint ThreadIndex = GroupThreadId.y * THREADGROUP_SIZEX + GroupThreadId.x;
 
-	float2 ScreenUV = float2((DispatchThreadId.xy * DOWNSAMPLE_FACTOR + View.ViewRectMin.xy + .5f) * View.ViewSizeAndSceneTexelSize.zw);
+	float2 ScreenUV = float2((DispatchThreadId.xy * DOWNSAMPLE_FACTOR + View.ViewRectMin.xy + .5f) * View.BufferSizeAndInvSize.zw);
 	float2 ScreenPosition = (ScreenUV.xy - View.ScreenPositionScaleBias.wz) / View.ScreenPositionScaleBias.xy;
 
 	float SceneDepth = CalcSceneDepth(ScreenUV);
@@ -615,7 +615,7 @@ SamplerState VisualizeDistanceFieldSampler;
 void VisualizeDistanceFieldUpsamplePS(in float4 UVAndScreenPos : TEXCOORD0, out float4 OutColor : SV_Target0)
 {
 	// Distance field AO was computed at 0,0 regardless of viewrect min
-	float2 DistanceFieldUVs = UVAndScreenPos.xy - View.ViewRectMin.xy * View.ViewSizeAndSceneTexelSize.zw;
+	float2 DistanceFieldUVs = UVAndScreenPos.xy - View.ViewRectMin.xy * View.BufferSizeAndInvSize.zw;
 
 	float3 Value = Texture2DSampleLevel(VisualizeDistanceFieldTexture, VisualizeDistanceFieldSampler, DistanceFieldUVs, 0).xyz;
 
diff --git a/Engine/Shaders/DistanceFieldSurfaceCacheLightingCompute.usf b/Engine/Shaders/DistanceFieldSurfaceCacheLightingCompute.usf
--- a/Engine/Shaders/DistanceFieldSurfaceCacheLightingCompute.usf
+++ b/Engine/Shaders/DistanceFieldSurfaceCacheLightingCompute.usf
@@ -269,10 +269,10 @@ void BuildTileConesMain(
 {
     uint ThreadIndex = GroupThreadId.y * THREADGROUP_SIZEX + GroupThreadId.x;
 
-	float2 BaseLevelScreenUV = (DispatchThreadId.xy + float2(.5f, .5f)) * DOWNSAMPLE_FACTOR * View.ViewSizeAndSceneTexelSize.zw;
+	float2 BaseLevelScreenUV = (DispatchThreadId.xy + float2(.5f, .5f)) * DOWNSAMPLE_FACTOR * View.BufferSizeAndInvSize.zw;
 	float SceneDepth = GetDownsampledDepth(BaseLevelScreenUV);
 
-	//float2 ScreenUV = (DispatchThreadId.xy * DOWNSAMPLE_FACTOR + View.ViewRectMin.xy + float2(.5f, .5f)) * View.ViewSizeAndSceneTexelSize.zw;
+	//float2 ScreenUV = (DispatchThreadId.xy * DOWNSAMPLE_FACTOR + View.ViewRectMin.xy + float2(.5f, .5f)) * View.BufferSizeAndInvSize.zw;
 	//float SceneDepth = CalcSceneDepth(ScreenUV);
 
 	// Initialize per-tile variables
@@ -386,7 +386,7 @@ void DistanceFieldAOBuildTileListMain(
     uint3 GroupThreadId : SV_GroupThreadID) 
 {
     uint ThreadIndex = GroupThreadId.y * THREADGROUP_SIZEX + GroupThreadId.x;
-	float2 ScreenUV = (DispatchThreadId.xy * DOWNSAMPLE_FACTOR + View.ViewRectMin.xy + float2(.5f, .5f)) * View.ViewSizeAndSceneTexelSize.zw;
+	float2 ScreenUV = (DispatchThreadId.xy * DOWNSAMPLE_FACTOR + View.ViewRectMin.xy + float2(.5f, .5f)) * View.BufferSizeAndInvSize.zw;
 	float SceneDepth = CalcSceneDepth(ScreenUV);
 
 	// Initialize per-tile variables
@@ -640,7 +640,7 @@ void PopulateCacheCS(
 
 	GroupMemoryBarrierWithGroupSync();
 
-	float2 ScreenUV = (DispatchThreadId.xy * CurrentLevelDownsampleFactor + View.ViewRectMin.xy + float2(.5f, .5f)) * View.ViewSizeAndSceneTexelSize.zw;
+	float2 ScreenUV = (DispatchThreadId.xy * CurrentLevelDownsampleFactor + View.ViewRectMin.xy + float2(.5f, .5f)) * View.BufferSizeAndInvSize.zw;
 	float2 ScreenPosition = (ScreenUV.xy - View.ScreenPositionScaleBias.wz) / View.ScreenPositionScaleBias.xy;
 	float2 DownsampledScreenUV = (DispatchThreadId.xy + float2(.5f, .5f)) / AOBufferSize;
 
diff --git a/Engine/Shaders/ForEachLight.usf b/Engine/Shaders/ForEachLight.usf
--- a/Engine/Shaders/ForEachLight.usf
+++ b/Engine/Shaders/ForEachLight.usf
@@ -37,7 +37,7 @@ float3 LightingForOneLight(FMaterialPixelParameters Parameters, uint LightIndex,
 	// The only type of shadowing supported for lights using tiled is static shadowing, so the light should only compute shadowing if it has static shadowing
 	LightData.bShadowed = dot(LightData.ShadowMapChannelMask, float4(1, 1, 1, 1));
 	
-	float2 ScreenUV = Parameters.SVPosition.xy * View.ViewSizeAndSceneTexelSize.zw;
+	float2 ScreenUV = Parameters.SVPosition.xy * View.BufferSizeAndInvSize.zw;
 
 	return GetDynamicLighting(Parameters.WorldPosition, CameraVector, ScreenUV, ScreenSpaceData, ScreenSpaceData.GBuffer.ShadingModelID, LightData, float4(1, 1, 1, 1), uint2(0,0)).rgb;
 }
diff --git a/Engine/Shaders/MaterialTemplate.usf b/Engine/Shaders/MaterialTemplate.usf
--- a/Engine/Shaders/MaterialTemplate.usf
+++ b/Engine/Shaders/MaterialTemplate.usf
@@ -817,11 +817,11 @@ float3 GetSpeedTreeVertexOffsetInner(FMaterialVertexParameters Parameters, int G
 #if !USE_INSTANCING || !USE_DITHERED_LOD_TRANSITION
 	if (LODType == SPEEDTREE_LOD_TYPE_SMOOTH) 
 	{
-		const float ScreenMultiple = 0.5 * max(View.ViewSizeAndSceneTexelSize.x * View.ViewToClip[0][0],
-												View.ViewSizeAndSceneTexelSize.y * View.ViewToClip[1][1]);
+		const float ScreenMultiple = 0.5 * max(View.ViewSizeAndInvSize.x * View.ViewToClip[0][0],
+												View.ViewSizeAndInvSize.y * View.ViewToClip[1][1]);
 		const float ScreenRadius = ScreenMultiple * Primitive.ObjectWorldPositionAndRadius.w / 
 									max(1.0, dot(TreePos - View.ViewOrigin.xyz, View.ViewForward.xyz));
-		const float ScreenArea = PI * ScreenRadius * ScreenRadius / (View.ViewSizeAndSceneTexelSize.x * View.ViewSizeAndSceneTexelSize.y);
+		const float ScreenArea = PI * ScreenRadius * ScreenRadius * View.ViewSizeAndInvSize.z * View.ViewSizeAndInvSize.w;
 		LodInterp = saturate((ScreenArea - SpeedTreeLODInfo.x) / SpeedTreeLODInfo.z);
 	}
 #endif
@@ -1562,7 +1562,7 @@ void CalcMaterialParameters(
 
 	{
 		float2 PixelPos = SvPosition.xy - View.ViewRectMin.xy;	
-		Parameters.ScreenPosition.xy = ( PixelPos / View.ViewSizeAndSceneTexelSize.xy - 0.5 ) * float2(2,-2);
+		Parameters.ScreenPosition.xy = ( PixelPos * View.ViewSizeAndInvSize.zw - 0.5 ) * float2(2,-2);
 		Parameters.ScreenPosition.zw = SvPosition.zw;
 		Parameters.ScreenPosition.xyz *= Parameters.ScreenPosition.w;
 	}
diff --git a/Engine/Shaders/PostProcessCompositeEditorPrimitives.usf b/Engine/Shaders/PostProcessCompositeEditorPrimitives.usf
--- a/Engine/Shaders/PostProcessCompositeEditorPrimitives.usf
+++ b/Engine/Shaders/PostProcessCompositeEditorPrimitives.usf
@@ -152,7 +152,7 @@ void MainPS(float4 UVAndScreenPos : TEXCOORD0, out float4 OutColor : SV_Target0)
 //	OutColor = SrcColor.a;	return;
 
 	// dejitter the sample position and make a filtered lookup - for planes this allows to reconstruct a much less jittery depth comparison function, it however doesn't fix silhuetes
-	float DeviceZ = Texture2DSampleLevel(FilteredSceneDepthTexture, FilteredSceneDepthTextureSampler, (PixelPosF - View.TemporalAAParams.zw) * View.ViewSizeAndSceneTexelSize.zw, 0).r;
+	float DeviceZ = Texture2DSampleLevel(FilteredSceneDepthTexture, FilteredSceneDepthTextureSampler, (PixelPosF - View.TemporalAAParams.zw) * View.BufferSizeAndInvSize.zw, 0).r;
 
 	// Soft Bias with DeviceZ for best quality
 	const float DeviceDepthFade = 0.00005f;
diff --git a/Engine/Shaders/PostProcessSelectionOutline.usf b/Engine/Shaders/PostProcessSelectionOutline.usf
--- a/Engine/Shaders/PostProcessSelectionOutline.usf
+++ b/Engine/Shaders/PostProcessSelectionOutline.usf
@@ -220,10 +220,10 @@ void MainPS(float4 UVAndScreenPos : TEXCOORD0, out float4 OutColor : SV_Target0)
 	float3 DeviceZPlane;
 	DeviceZPlane.z = Texture2DSampleLevel(SceneDepthTexture, SceneDepthTextureSampler, UV, 0).r;
 
-	float Left = Texture2DSampleLevel(SceneDepthTexture, SceneDepthTextureSampler, UV + float2(-1, 0) * View.ViewSizeAndSceneTexelSize.zw, 0).r;
-	float Right= Texture2DSampleLevel(SceneDepthTexture, SceneDepthTextureSampler, UV + float2(1, 0) * View.ViewSizeAndSceneTexelSize.zw, 0).r;
-	float Top  = Texture2DSampleLevel(SceneDepthTexture, SceneDepthTextureSampler, UV + float2(0, -1) * View.ViewSizeAndSceneTexelSize.zw, 0).r;
-	float Bottom = Texture2DSampleLevel(SceneDepthTexture, SceneDepthTextureSampler, UV + float2(0, 1) * View.ViewSizeAndSceneTexelSize.zw, 0).r;
+	float Left = Texture2DSampleLevel(SceneDepthTexture, SceneDepthTextureSampler, UV + float2(-1, 0) * View.BufferSizeAndInvSize.zw, 0).r;
+	float Right= Texture2DSampleLevel(SceneDepthTexture, SceneDepthTextureSampler, UV + float2(1, 0) * View.BufferSizeAndInvSize.zw, 0).r;
+	float Top  = Texture2DSampleLevel(SceneDepthTexture, SceneDepthTextureSampler, UV + float2(0, -1) * View.BufferSizeAndInvSize.zw, 0).r;
+	float Bottom = Texture2DSampleLevel(SceneDepthTexture, SceneDepthTextureSampler, UV + float2(0, 1) * View.BufferSizeAndInvSize.zw, 0).r;
 
 	float2 DeviceZMinMax;
 	DeviceZMinMax.x = min(DeviceZPlane.z, min(min(Left, Right), min(Top, Bottom)));
diff --git a/Engine/Shaders/PostProcessSubsurface.usf b/Engine/Shaders/PostProcessSubsurface.usf
--- a/Engine/Shaders/PostProcessSubsurface.usf
+++ b/Engine/Shaders/PostProcessSubsurface.usf
@@ -142,7 +142,7 @@ float ComputeFullResLerp(FScreenSpaceData ScreenSpaceData, float2 UVSceneColor,
 
 	// the viewport is only a fraction of the buffersize, here we compensate for that
 	// *1024 as we store the sample distances / 1000 to allows for scaling
-	float HorizontalScaler = View.ViewSizeAndSceneTexelSize.x * View.ViewSizeAndSceneTexelSize.z * 1024.0f;
+	float HorizontalScaler = View.ViewSizeAndInvSize.x * View.BufferSizeAndInvSize.z * 1024.0f;
 
 	// Calculate the final step to fetch the surrounding pixels:
 	float finalStep = scale * HorizontalScaler;
@@ -174,13 +174,13 @@ void VisualizePS(in float4 UVAndScreenPos : TEXCOORD0, out float4 OutColor : SV_
 
 	int2 PixelPos = (int2)(UVAndScreenPos.zw * ScreenPosToPixel.xy + ScreenPosToPixel.zw + 0.5f);
 
-	float2 ViewLocalUV = (PixelPos - View.ViewRectMin.xy) / View.ViewSizeAndSceneTexelSize.xy;
+	float2 ViewLocalUV = (PixelPos - View.ViewRectMin.xy) * View.ViewSizeAndInvSize.zw;
 
 	float2 IDAreaLocalUV = ViewLocalUV * 2 - 1.0f;
 
 	if (InUnitBox(IDAreaLocalUV))
 	{
-		float2 UV = View.ViewRectMin.xy * View.ViewSizeAndSceneTexelSize.zw + IDAreaLocalUV * (View.ViewSizeAndSceneTexelSize.xy * View.ViewSizeAndSceneTexelSize.zw);
+		float2 UV = View.ViewRectMin.xy * View.BufferSizeAndInvSize.zw + IDAreaLocalUV * (View.ViewSizeAndInvSize.xy * View.BufferSizeAndInvSize.zw);
 
 		FScreenSpaceData ScreenSpaceData = GetScreenSpaceData(UV);
 
@@ -329,7 +329,7 @@ void SetupPS(in float4 UVAndScreenPos : TEXCOORD0, out float4 OutColor : SV_Targ
 
 	// the viewport is only a fraction of the buffersize, here we compensate for that
 	// *1024 as we store the sample distances / 1024 to allows for scaling
-	float HorizontalScaler = View.ViewSizeAndSceneTexelSize.x * View.ViewSizeAndSceneTexelSize.z * 1024.0f;
+	float HorizontalScaler = View.ViewSizeAndInvSize.x * View.BufferSizeAndInvSize.z * 1024.0f;
 
 	// Calculate the final step to fetch the surrounding pixels:
 	float finalStep = scale * HorizontalScaler;
@@ -375,14 +375,14 @@ void MainPS(float4 UVAndScreenPos : TEXCOORD0, out float4 OutColor : SV_Target0)
 
 #if SSS_DIRECTION == 0
 	// horizontal
-	float2 Direction = float2(1, 0) * View.ViewSizeAndSceneTexelSize.x * View.ViewSizeAndSceneTexelSize.z * 1024.0f;
+	float2 Direction = float2(1, 0) * View.ViewSizeAndInvSize.x * View.BufferSizeAndInvSize.z * 1024.0f;
 #else
 	// vertical
-	float2 Direction = float2(0, 1) * View.ViewSizeAndSceneTexelSize.w * View.ViewSizeAndSceneTexelSize.x * 1024.0f;
+	float2 Direction = float2(0, 1) * View.BufferSizeAndInvSize.w * View.ViewSizeAndInvSize.x * 1024.0f;
 #endif
 
 	// can be optimized
-	float2 GBufferUV = (ViewportUV * PostprocessInput0Size.xy * 2 + View.ViewRectMin.xy) * View.ViewSizeAndSceneTexelSize.zw;
+	float2 GBufferUV = (ViewportUV * PostprocessInput0Size.xy * 2 + View.ViewRectMin.xy) * View.BufferSizeAndInvSize.zw;
 
 	bool bIncludeCenter = true;
 
@@ -410,7 +410,7 @@ void SubsurfaceRecombinePS(float4 UVAndScreenPos : TEXCOORD0, out float4 OutColo
 
 	// can be optimized
 	float2 PixelPos = UVAndScreenPos.zw * ScreenPosToPixel.xy + ScreenPosToPixel.zw;
-	float2 UVSceneColor = (PixelPos + 0.5f) * View.ViewSizeAndSceneTexelSize.zw;
+	float2 UVSceneColor = (PixelPos + 0.5f) * View.BufferSizeAndInvSize.zw;
 
 	// specular extracted
 	float4 SceneColor4 = Texture2DSample(PostprocessInput1, PostprocessInput1Sampler, UVSceneColor);
diff --git a/Engine/Shaders/PostProcessVisualizeHDR.usf b/Engine/Shaders/PostProcessVisualizeHDR.usf
--- a/Engine/Shaders/PostProcessVisualizeHDR.usf
+++ b/Engine/Shaders/PostProcessVisualizeHDR.usf
@@ -193,7 +193,7 @@ void MainPS(float4 UVAndScreenPos : TEXCOORD0, out float4 OutColor : SV_Target0)
 		OutColor = bChecker ? float4(0.2f,0,0,0) : float4(0,0.15f,0,0);
 		
 		// SceneColor overlay
-		OutColor += 0.1f * Texture2DSampleLevel(PostprocessInput0, PostprocessInput0Sampler, IDAreaLocalUV * View.ViewSizeAndSceneTexelSize.xy * View.ViewSizeAndSceneTexelSize.zw, 0);
+		OutColor += 0.1f * Texture2DSampleLevel(PostprocessInput0, PostprocessInput0Sampler, IDAreaLocalUV * View.ViewSizeAndInvSize.xy * View.BufferSizeAndInvSize.zw, 0);
 
 		// could be the same as (HistogramY < Value) but a bit more antialiased
 //		float HistogramMask = HistogramY < Value;
diff --git a/Engine/Shaders/ShaderVersion.usf b/Engine/Shaders/ShaderVersion.usf
--- a/Engine/Shaders/ShaderVersion.usf
+++ b/Engine/Shaders/ShaderVersion.usf
@@ -2,4 +2,4 @@
 // Each time the console command is executed it generates a new GUID. As this file is included
 // in common.usf (which should be included in any shader) it allows to invalidate the shader DDC.
 // 
-// GUID = E106A9034435600215322F98E16C9042
\ No newline at end of file
+// GUID = 10DDEA3C48E7FD137E98039F9438498C
\ No newline at end of file
diff --git a/Engine/Shaders/ShadowProjectionPixelShader.usf b/Engine/Shaders/ShadowProjectionPixelShader.usf
--- a/Engine/Shaders/ShadowProjectionPixelShader.usf
+++ b/Engine/Shaders/ShadowProjectionPixelShader.usf
@@ -46,7 +46,7 @@ void Main(
 	out float4 OutColor : SV_Target0
 	)
 {
-	float2 ScreenUV = float2( SVPos.xy * View.ViewSizeAndSceneTexelSize.zw );
+	float2 ScreenUV = float2( SVPos.xy * View.BufferSizeAndInvSize.zw );
 	float SceneW = CalcSceneDepth( ScreenUV );
 	float2 ScreenPosition = ( ScreenUV.xy - View.ScreenPositionScaleBias.wz ) / View.ScreenPositionScaleBias.xy;
 	float4 ShadowPosition = mul(float4(ScreenPosition.xy * SceneW,SceneW,1), ScreenToShadowMatrix);
@@ -169,7 +169,7 @@ void MainOnePassPointLightPS(
 	)
 {
 
-	float2 ScreenUV = float2( SVPos.xy * View.ViewSizeAndSceneTexelSize.zw );
+	float2 ScreenUV = float2( SVPos.xy * View.BufferSizeAndInvSize.zw );
 	float SceneW = CalcSceneDepth( ScreenUV );
 	float2 ScreenPosition = ( ScreenUV.xy - View.ScreenPositionScaleBias.wz ) / View.ScreenPositionScaleBias.xy;
 	float4 HomogeneousWorldPosition = mul(float4(ScreenPosition.xy * SceneW, SceneW, 1), View.ScreenToWorld);
diff --git a/Engine/Source/Runtime/Engine/Classes/Materials/MaterialExpressionViewProperty.h b/Engine/Source/Runtime/Engine/Classes/Materials/MaterialExpressionViewProperty.h
--- a/Engine/Source/Runtime/Engine/Classes/Materials/MaterialExpressionViewProperty.h
+++ b/Engine/Source/Runtime/Engine/Classes/Materials/MaterialExpressionViewProperty.h
@@ -8,12 +8,14 @@
 UENUM()
 enum EMaterialExposedViewProperty
 {	
-	/** Horizontal and vertical fields of view in radian */
+	/** Horizontal and vertical size of the view's buffer in pixels */
+	MEVP_BufferSize UMETA(DisplayName="RenderTargetSize"),
+	/** Horizontal and vertical field of view angles in radian */
 	MEVP_FieldOfView UMETA(DisplayName="FieldOfView"),
+	/** Tan(FieldOfView * 0.5) */
+	MEVP_TanHalfFieldOfView UMETA(DisplayName="TanHalfFieldOfView"),
 	/** Horizontal and vertical size of the view in pixels */
 	MEVP_ViewSize UMETA(DisplayName="ViewSize"),
-	/** Horizontal and vertical size of view's render target's texel in the uv space (1 / ViewSize) */
-	MEVP_TexelSize UMETA(DisplayName="TexelSize"),
 
 	MEVP_MAX,
 };
diff --git a/Engine/Source/Runtime/Engine/Private/Materials/HLSLMaterialTranslator.h b/Engine/Source/Runtime/Engine/Private/Materials/HLSLMaterialTranslator.h
--- a/Engine/Source/Runtime/Engine/Private/Materials/HLSLMaterialTranslator.h
+++ b/Engine/Source/Runtime/Engine/Private/Materials/HLSLMaterialTranslator.h
@@ -1633,7 +1633,7 @@ protected:
 		return AddUniformExpression(new FMaterialUniformExpressionConstant(FLinearColor(X,Y,Z,W),MCT_Float4),MCT_Float4,TEXT("MaterialFloat4(%0.8f,%0.8f,%0.8f,%0.8f)"),X,Y,Z,W);
 	}
 	
-	virtual int32 ViewProperty(EMaterialExposedViewProperty Property) override
+	virtual int32 ViewProperty(EMaterialExposedViewProperty Property, bool InvProperty) override
 	{
 		check(Property < MEVP_MAX);
 
@@ -1642,28 +1642,38 @@ protected:
 		{
 			EMaterialExposedViewProperty EnumValue;
 			EMaterialValueType Type;
-			const TCHAR * CurrentFrameCode;
-			const TCHAR * PreviousFrameCode;
+			const TCHAR * PropertyCode;
+			const TCHAR * InvPropertyCode;
 		};
 
 		static const EMaterialExposedViewPropertyMeta ViewPropertyMetaArray[] = {
-			{MEVP_FieldOfView, MCT_Float2, TEXT("View.FieldOfViewWideAngles"), TEXT("View.PrevFieldOfViewWideAngles")},
-			{MEVP_ViewSize, MCT_Float2, TEXT("View.ViewSizeAndSceneTexelSize.xy"), nullptr},
-			{MEVP_TexelSize, MCT_Float2, TEXT("View.ViewSizeAndSceneTexelSize.zw"), nullptr},
+			{MEVP_BufferSize, MCT_Float2, TEXT("View.BufferSizeAndInvSize.xy"), TEXT("View.BufferSizeAndInvSize.zw")},
+			{MEVP_FieldOfView, MCT_Float2, TEXT("View.<PREV>FieldOfViewWideAngles"), nullptr},
+			{MEVP_TanHalfFieldOfView, MCT_Float2, TEXT("Get<PREV>TanHalfFieldOfView()"), TEXT("Get<PREV>CotanHalfFieldOfView()")},
+			{MEVP_ViewSize, MCT_Float2, TEXT("View.ViewSizeAndInvSize.xy"), TEXT("View.ViewSizeAndInvSize.zw")},
 		};
 		static_assert((sizeof(ViewPropertyMetaArray) / sizeof(ViewPropertyMetaArray[0])) == MEVP_MAX, "incoherency between EMaterialExposedViewProperty and ViewPropertyMetaArray");
 
 		auto& PropertyMeta = ViewPropertyMetaArray[Property];
 		check(Property == PropertyMeta.EnumValue);
 
-		auto Code = PropertyMeta.CurrentFrameCode;
+		FString Code = PropertyMeta.PropertyCode;
 
-		if (bCompilingPreviousFrame && PropertyMeta.PreviousFrameCode)
+		if (InvProperty && PropertyMeta.InvPropertyCode)
 		{
-			Code = PropertyMeta.PreviousFrameCode;
+			Code = PropertyMeta.InvPropertyCode;
 		}
 
-		return AddCodeChunk(PropertyMeta.Type, Code);
+		// Resolved templated code
+		Code.ReplaceInline(TEXT("<PREV>"), bCompilingPreviousFrame ? TEXT("Prev") : TEXT(""));
+		
+		if (InvProperty && !PropertyMeta.InvPropertyCode)
+		{
+			// fall back to compute the property's inverse from PropertyCode
+			return Div(Constant(1.f), AddInlinedCodeChunk(PropertyMeta.Type, *Code));
+		}
+
+		return AddCodeChunk(PropertyMeta.Type, *Code);
 	}
 
 	virtual int32 GameTime(bool bPeriodic, float Period) override
diff --git a/Engine/Source/Runtime/Engine/Private/Materials/MaterialExpressions.cpp b/Engine/Source/Runtime/Engine/Private/Materials/MaterialExpressions.cpp
--- a/Engine/Source/Runtime/Engine/Private/Materials/MaterialExpressions.cpp
+++ b/Engine/Source/Runtime/Engine/Private/Materials/MaterialExpressions.cpp
@@ -4626,11 +4626,16 @@ UMaterialExpressionViewProperty::UMaterialExpressionViewProperty(const FObjectIn
 	Property = MEVP_FieldOfView;
 	MenuCategories.Add(ConstructorStatics.NAME_Constants);
 	bShaderInputData = true;
+	bShowOutputNameOnPin = true;
+	
+	Outputs.Reset();
+	Outputs.Add(FExpressionOutput(TEXT("Property")));
+	Outputs.Add(FExpressionOutput(TEXT("InvProperty")));
 }
 
 int32 UMaterialExpressionViewProperty::Compile(class FMaterialCompiler* Compiler, int32 OutputIndex, int32 MultiplexIndex)
 {
-	return Compiler->ViewProperty(Property);
+	return Compiler->ViewProperty(Property, OutputIndex == 1);
 }
 
 void UMaterialExpressionViewProperty::GetCaption(TArray<FString>& OutCaptions) const
@@ -4695,7 +4700,7 @@ UMaterialExpressionSceneTexelSize::UMaterialExpressionSceneTexelSize(const FObje
 
 int32 UMaterialExpressionSceneTexelSize::Compile(class FMaterialCompiler* Compiler, int32 OutputIndex, int32 MultiplexIndex)
 {
-	return Compiler->ViewProperty(MEVP_TexelSize);
+	return Compiler->ViewProperty(MEVP_BufferSize, /* InvProperty = */ true);
 }
 
 void UMaterialExpressionSceneTexelSize::GetCaption(TArray<FString>& OutCaptions) const
diff --git a/Engine/Source/Runtime/Engine/Public/MaterialCompiler.h b/Engine/Source/Runtime/Engine/Public/MaterialCompiler.h
--- a/Engine/Source/Runtime/Engine/Public/MaterialCompiler.h
+++ b/Engine/Source/Runtime/Engine/Public/MaterialCompiler.h
@@ -58,7 +58,7 @@ public:
 	virtual int32 Constant3(float X,float Y,float Z) = 0;
 	virtual int32 Constant4(float X,float Y,float Z,float W) = 0;
 
-	virtual	int32 ViewProperty(EMaterialExposedViewProperty Property) = 0;
+	virtual	int32 ViewProperty(EMaterialExposedViewProperty Property, bool InvProperty = false) = 0;
 
 	virtual int32 GameTime(bool bPeriodic, float Period) = 0;
 	virtual int32 RealTime(bool bPeriodic, float Period) = 0;
@@ -223,7 +223,7 @@ public:
 	virtual int32 Constant3(float X,float Y,float Z) override { return Compiler->Constant3(X,Y,Z); }
 	virtual int32 Constant4(float X,float Y,float Z,float W) override { return Compiler->Constant4(X,Y,Z,W); }
 	
-	virtual	int32 ViewProperty(EMaterialExposedViewProperty Property) override { return Compiler->ViewProperty(Property); }
+	virtual	int32 ViewProperty(EMaterialExposedViewProperty Property, bool InvProperty) override { return Compiler->ViewProperty(Property, InvProperty); }
 
 	virtual int32 GameTime(bool bPeriodic, float Period) override { return Compiler->GameTime(bPeriodic, Period); }
 	virtual int32 RealTime(bool bPeriodic, float Period) override { return Compiler->RealTime(bPeriodic, Period); }
diff --git a/Engine/Source/Runtime/Engine/Public/SceneView.h b/Engine/Source/Runtime/Engine/Public/SceneView.h
--- a/Engine/Source/Runtime/Engine/Public/SceneView.h
+++ b/Engine/Source/Runtime/Engine/Public/SceneView.h
@@ -324,7 +324,8 @@ BEGIN_UNIFORM_BUFFER_STRUCT_WITH_CONSTRUCTOR(FViewUniformShaderParameters,ENGINE
 	DECLARE_UNIFORM_BUFFER_STRUCT_MEMBER(FVector4,InvDeviceZToWorldZTransform)
 	DECLARE_UNIFORM_BUFFER_STRUCT_MEMBER_EX(FVector4,ScreenPositionScaleBias, EShaderPrecisionModifier::Half)
 	DECLARE_UNIFORM_BUFFER_STRUCT_MEMBER_EX(FVector4,ViewRectMin, EShaderPrecisionModifier::Half)
-	DECLARE_UNIFORM_BUFFER_STRUCT_MEMBER(FVector4,ViewSizeAndSceneTexelSize)
+	DECLARE_UNIFORM_BUFFER_STRUCT_MEMBER(FVector4,ViewSizeAndInvSize)
+	DECLARE_UNIFORM_BUFFER_STRUCT_MEMBER(FVector4,BufferSizeAndInvSize)
 	DECLARE_UNIFORM_BUFFER_STRUCT_MEMBER(FVector4,ViewOrigin)
 	DECLARE_UNIFORM_BUFFER_STRUCT_MEMBER(FVector4,TranslatedViewOrigin)
 	// The exposure scale is just a scalar but needs to be a float4 to workaround a driver bug on IOS.
@@ -372,6 +373,8 @@ BEGIN_UNIFORM_BUFFER_STRUCT_WITH_CONSTRUCTOR(FViewUniformShaderParameters,ENGINE
 	DECLARE_UNIFORM_BUFFER_STRUCT_MEMBER(FMatrix,PrevProjection)
 	DECLARE_UNIFORM_BUFFER_STRUCT_MEMBER(FMatrix,PrevViewProj)
 	DECLARE_UNIFORM_BUFFER_STRUCT_MEMBER(FMatrix,PrevViewRotationProj)
+	DECLARE_UNIFORM_BUFFER_STRUCT_MEMBER(FMatrix,PrevViewToClip)
+	DECLARE_UNIFORM_BUFFER_STRUCT_MEMBER(FMatrix,PrevClipToView)
 	DECLARE_UNIFORM_BUFFER_STRUCT_MEMBER(FMatrix,PrevTranslatedWorldToClip)
 	DECLARE_UNIFORM_BUFFER_STRUCT_MEMBER(FMatrix,PrevTranslatedWorldToView)
 	DECLARE_UNIFORM_BUFFER_STRUCT_MEMBER(FMatrix,PrevViewToTranslatedWorld)
diff --git a/Engine/Source/Runtime/Renderer/Private/SceneRendering.cpp b/Engine/Source/Runtime/Renderer/Private/SceneRendering.cpp
--- a/Engine/Source/Runtime/Renderer/Private/SceneRendering.cpp
+++ b/Engine/Source/Runtime/Renderer/Private/SceneRendering.cpp
@@ -469,7 +469,8 @@ TUniformBufferRef<FViewUniformShaderParameters> FViewInfo::CreateUniformBuffer(
 	ViewUniformShaderParameters.InvDeviceZToWorldZTransform = InvDeviceZToWorldZTransform;
 	ViewUniformShaderParameters.ScreenPositionScaleBias = ScreenPositionScaleBias;
 	ViewUniformShaderParameters.ViewRectMin = FVector4(ViewRect.Min.X, ViewRect.Min.Y, 0.0f, 0.0f);
-	ViewUniformShaderParameters.ViewSizeAndSceneTexelSize = FVector4(ViewRect.Width(), ViewRect.Height(), InvBufferSizeX, InvBufferSizeY);
+	ViewUniformShaderParameters.ViewSizeAndInvSize = FVector4(ViewRect.Width(), ViewRect.Height(), 1.0f / float(ViewRect.Width()), 1.0f / float(ViewRect.Height()));
+	ViewUniformShaderParameters.BufferSizeAndInvSize = FVector4(BufferSize.X, BufferSize.Y, InvBufferSizeX, InvBufferSizeY);
 	ViewUniformShaderParameters.ViewOrigin = ViewMatrices.ViewOrigin;
 	ViewUniformShaderParameters.TranslatedViewOrigin = ViewMatrices.ViewOrigin + ViewMatrices.PreViewTranslation;
 	ViewUniformShaderParameters.DiffuseOverrideParameter = LocalDiffuseOverrideParameter;
@@ -485,6 +486,8 @@ TUniformBufferRef<FViewUniformShaderParameters> FViewInfo::CreateUniformBuffer(
 	ViewUniformShaderParameters.PrevProjection = PrevViewMatrices.ProjMatrix;
 	ViewUniformShaderParameters.PrevViewProj = PrevViewProjMatrix;
 	ViewUniformShaderParameters.PrevViewRotationProj = PrevViewRotationProjMatrix;
+	ViewUniformShaderParameters.PrevViewToClip = PrevViewMatrices.ProjMatrix;
+	ViewUniformShaderParameters.PrevClipToView = PrevViewMatrices.GetInvProjMatrix();
 	ViewUniformShaderParameters.PrevTranslatedWorldToClip = PrevViewMatrices.TranslatedViewProjectionMatrix;
 	ViewUniformShaderParameters.PrevTranslatedWorldToView = PrevViewMatrices.TranslatedViewMatrix;
 	ViewUniformShaderParameters.PrevViewToTranslatedWorld = ViewUniformShaderParameters.PrevTranslatedWorldToView.Inverse();
diff --git a/Engine/Source/Runtime/SlateRHIRenderer/Private/SlateRHIRenderingPolicy.cpp b/Engine/Source/Runtime/SlateRHIRenderer/Private/SlateRHIRenderingPolicy.cpp
--- a/Engine/Source/Runtime/SlateRHIRenderer/Private/SlateRHIRenderingPolicy.cpp
+++ b/Engine/Source/Runtime/SlateRHIRenderer/Private/SlateRHIRenderingPolicy.cpp
@@ -312,7 +312,8 @@ static FSceneView& CreateSceneView( FSceneViewFamilyContext& ViewFamilyContext,
 	ViewUniformShaderParameters.InvDeviceZToWorldZTransform = View->InvDeviceZToWorldZTransform;
 	ViewUniformShaderParameters.ScreenPositionScaleBias = FVector4(0,0,0,0);
 	ViewUniformShaderParameters.ViewRectMin = FVector4(ViewRect.Min.X, ViewRect.Min.Y, 0.0f, 0.0f);
-	ViewUniformShaderParameters.ViewSizeAndSceneTexelSize = FVector4(ViewRect.Width(), ViewRect.Height(), 1.0f/ViewRect.Width(), 1.0f/ViewRect.Height() );
+	ViewUniformShaderParameters.ViewSizeAndInvSize = FVector4(ViewRect.Width(), ViewRect.Height(), 1.0f/ViewRect.Width(), 1.0f/ViewRect.Height() );
+	ViewUniformShaderParameters.BufferSizeAndInvSize = ViewUniformShaderParameters.ViewSizeAndInvSize;
 	ViewUniformShaderParameters.ViewOrigin = View->ViewMatrices.ViewOrigin;
 	ViewUniformShaderParameters.TranslatedViewOrigin = View->ViewMatrices.ViewOrigin + View->ViewMatrices.PreViewTranslation;
 	ViewUniformShaderParameters.DiffuseOverrideParameter = View->DiffuseOverrideParameter;
