Conservative depth support for PS4

diff --git a/Engine/Shaders/MaterialTemplate.usf b/Engine/Shaders/MaterialTemplate.usf
--- a/Engine/Shaders/MaterialTemplate.usf
+++ b/Engine/Shaders/MaterialTemplate.usf
@@ -1642,13 +1642,23 @@ half3x3 AssembleTangentToWorld( half3 TangentToWorld0, half4 TangentToWorld2 )
 
 // Whether to use the hidden d3d11 feature that supports depth writes with ZCull by only pushing into the screen
 //@todo - use for other SM5 platforms
-#define USE_CONSERVATIVE_DEPTH_WRITES (OUTPUT_PIXEL_DEPTH_OFFSET && COMPILER_HLSL && FEATURE_LEVEL >= FEATURE_LEVEL_SM5)
+#define SUPPORTS_CONSERVATIVE_DEPTH_WRITES ((COMPILER_HLSL && FEATURE_LEVEL >= FEATURE_LEVEL_SM5) || (PS4_PROFILE))
+#define USE_CONSERVATIVE_DEPTH_WRITES (OUTPUT_PIXEL_DEPTH_OFFSET && SUPPORTS_CONSERVATIVE_DEPTH_WRITES) 
 
 #if USE_CONSERVATIVE_DEPTH_WRITES
+
+#if COMPILER_HLSL
 	// Note: for some reason using SV_DepthLessEqual without these interpolation modifiers causes a compile error in d3d
 	#define INPUT_POSITION_QUALIFIERS linear noperspective centroid
 	// Use conservative depth output so we still get Z Cull.  Note, this is a reversed Z depth surface.
 	#define DEPTH_WRITE_SEMANTIC SV_DepthLessEqual
+#elif PS4_PROFILE
+	#define INPUT_POSITION_QUALIFIERS
+	#define DEPTH_WRITE_SEMANTIC S_DEPTH_LE_OUTPUT
+#else
+	#error USE_CONSERVATIVE_DEPTH_WRITES enabled for unsupported platform
+#endif
+
 #else
 	#define INPUT_POSITION_QUALIFIERS 
 	#define DEPTH_WRITE_SEMANTIC SV_DEPTH
