Removed unused shader code

diff --git a/Engine/Shaders/ReflectionEnvironmentShaders.usf b/Engine/Shaders/ReflectionEnvironmentShaders.usf
--- a/Engine/Shaders/ReflectionEnvironmentShaders.usf
+++ b/Engine/Shaders/ReflectionEnvironmentShaders.usf
@@ -29,34 +29,8 @@ void CopyToCubeFaceVS(
 	Out.ScreenVector = mul(float4(Out.Position.xy, 1, 0), View.ScreenToTranslatedWorld).xyz;
 }
 
-struct FRouteToCubeFaceGSOut
-{
-	FCopyToCubeFaceVSOutput Vertex;
-	uint LayerIndex : SV_RenderTargetArrayIndex;
-};
-
 int CubeFace;
 
-[maxvertexcount(3)]
-void RouteToCubeFaceGS(triangle FCopyToCubeFaceVSOutput Input[3], inout TriangleStream<FRouteToCubeFaceGSOut> OutStream)
-{
-	FRouteToCubeFaceGSOut Vertex0;
-	Vertex0.Vertex = Input[0];
-	Vertex0.LayerIndex = CubeFace;
-
-	FRouteToCubeFaceGSOut Vertex1;
-	Vertex1.Vertex = Input[1];
-	Vertex1.LayerIndex = CubeFace;
-
-	FRouteToCubeFaceGSOut Vertex2;
-	Vertex2.Vertex = Input[2];
-	Vertex2.LayerIndex = CubeFace;
-
-	OutStream.Append(Vertex0);
-	OutStream.Append(Vertex1);
-	OutStream.Append(Vertex2);
-}
-
 Texture2D InTexture;
 SamplerState InTextureSampler;
 
