Merging a couple of minor, Metal-specific shader changes

diff --git a/Engine/Shaders/Definitions.usf b/Engine/Shaders/Definitions.usf
--- a/Engine/Shaders/Definitions.usf
+++ b/Engine/Shaders/Definitions.usf
@@ -4,6 +4,8 @@
 	Definitions.usf: Defines undefined defines to 0 as Cg can't handle them.
 =============================================================================*/
 
+#pragma once
+
 #ifndef MATERIAL_TWOSIDED
 #define	MATERIAL_TWOSIDED								0
 #endif
@@ -100,6 +102,10 @@
 #define METAL_MRT_PROFILE								0
 #endif
 
+#ifndef METAL_SM5_PROFILE
+#define METAL_SM5_PROFILE								0
+#endif
+
 #ifndef VULKAN_PROFILE
 #define	VULKAN_PROFILE									0
 #endif
diff --git a/Engine/Shaders/MonteCarlo.usf b/Engine/Shaders/MonteCarlo.usf
--- a/Engine/Shaders/MonteCarlo.usf
+++ b/Engine/Shaders/MonteCarlo.usf
@@ -14,7 +14,7 @@ float3 TangentToWorld( float3 Vec, float3 TangentZ )
 
 uint ReverseBits32( uint bits )
 {
-#if SM5_PROFILE
+#if SM5_PROFILE && !METAL_MRT_PROFILE && !METAL_PROFILE && !METAL_SM5_PROFILE
 	return reversebits( bits );
 #else
 	bits = ( bits << 16) | ( bits >> 16);
