Fixes for motion blur fix on cloth.

diff --git a/Engine/Shaders/GpuSkinVertexFactory.usf b/Engine/Shaders/GpuSkinVertexFactory.usf
--- a/Engine/Shaders/GpuSkinVertexFactory.usf
+++ b/Engine/Shaders/GpuSkinVertexFactory.usf
@@ -583,8 +583,8 @@ float4 VertexFactoryGetPreviousWorldPosition(FVertexFactoryInput Input, FVertexF
 #if GPUSKIN_APEX_CLOTH
 	if( IsSimulatedVertex(Input) )
 	{
-		float3 BlendedPrevPos = lerp(PrevSkinPosInWorld.xyz, Intermediates.SimulatedPosition, ClothBlendWeight);		 
-		return float4(BlendedPrevPos + View.PrevPreViewTranslation, 1);
+		float3 BlendedPrevPos = lerp(PrevSkinPosInWorld.xyz, Intermediates.SimulatedPosition + View.PrevPreViewTranslation, ClothBlendWeight);		 
+		return float4(BlendedPrevPos, 1);
 	}
 #endif // GPUSKIN_APEX_CLOTH
 
