Fixes shader compilation failure for OpenGL in ScreenSpaceRayCast.usf

diff --git a/Engine/Shaders/ScreenSpaceRayCast.usf b/Engine/Shaders/ScreenSpaceRayCast.usf
--- a/Engine/Shaders/ScreenSpaceRayCast.usf
+++ b/Engine/Shaders/ScreenSpaceRayCast.usf
@@ -53,7 +53,7 @@ void RayCast(
 	const float4 RayEndV = RayStartV + RayDirV;
 
 	const float RayEndRadiusV = length(RayDirV.xyz) * tan(ConeAngleWorld * 0.4);
-	const float2 RayEndBorderV = RayEndV.xy + RayEndRadiusV / sqrt(2);
+	const float2 RayEndBorderV = RayEndV.xy + RayEndRadiusV / sqrt(2.0);
 
 	const float4 RayStartClip	= mul( RayStartV, View.ViewToClip );
 	const float4 RayEndClip	= mul( RayEndV, View.ViewToClip );
