Move ReverseBits32 from MonteCarlo to Common & use it in BasePassPixelShader for platforms that don't natively support the reversebits intrinsic - fixes [UE-17817].

diff --git a/Engine/Shaders/BasePassPixelShader.usf b/Engine/Shaders/BasePassPixelShader.usf
--- a/Engine/Shaders/BasePassPixelShader.usf
+++ b/Engine/Shaders/BasePassPixelShader.usf
@@ -124,7 +124,7 @@ half3 GetImageBasedReflectionLighting(FMaterialPixelParameters MaterialParameter
 		uint FrameRandom = (View.FrameNumber & 15) * 1551;
 
 		uint Morton = MortonCode( (uint)MaterialParameters.SVPosition.x & 3 ) | ( MortonCode( (uint)MaterialParameters.SVPosition.y & 3 ) * 2 );
-		uint PixelIndex = reversebits( Morton ) >> 28;
+		uint PixelIndex = ReverseBits32( Morton ) >> 28;
 		//uint PixelIndex = ( (uint)MaterialParameters.SVPosition.x & 3 ) | ( ( (uint)MaterialParameters.SVPosition.y & 3 ) * 2 );
 		//PixelIndex = ( PixelIndex * 1551 ) & 15;
 
diff --git a/Engine/Shaders/Common.usf b/Engine/Shaders/Common.usf
--- a/Engine/Shaders/Common.usf
+++ b/Engine/Shaders/Common.usf
@@ -1165,3 +1165,17 @@ float DecodePackedHeight(float2 PackedHeight)
 {
 	return DecodeHeightValue(DecodePackedTwoChannelValue(PackedHeight));
 }
+
+uint ReverseBits32( uint bits )
+{
+#if SM5_PROFILE && !METAL_MRT_PROFILE && !METAL_PROFILE && !METAL_SM5_PROFILE
+	return reversebits( bits );
+#else
+	bits = ( bits << 16) | ( bits >> 16);
+	bits = ( (bits & 0x00ff00ff) << 8 ) | ( (bits & 0xff00ff00) >> 8 );
+	bits = ( (bits & 0x0f0f0f0f) << 4 ) | ( (bits & 0xf0f0f0f0) >> 4 );
+	bits = ( (bits & 0x33333333) << 2 ) | ( (bits & 0xcccccccc) >> 2 );
+	bits = ( (bits & 0x55555555) << 1 ) | ( (bits & 0xaaaaaaaa) >> 1 );
+	return bits;
+#endif
+}
diff --git a/Engine/Shaders/MonteCarlo.usf b/Engine/Shaders/MonteCarlo.usf
--- a/Engine/Shaders/MonteCarlo.usf
+++ b/Engine/Shaders/MonteCarlo.usf
@@ -2,6 +2,8 @@
 
 #pragma once
 
+#include "Common.usf"
+
 /*=============================================================================
 	MonteCarlo.usf: Monte Carlo integration of distributions
 =============================================================================*/
@@ -14,20 +16,6 @@ float3 TangentToWorld( float3 Vec, float3 TangentZ )
 	return TangentX * Vec.x + TangentY * Vec.y + TangentZ * Vec.z;
 }
 
-uint ReverseBits32( uint bits )
-{
-#if SM5_PROFILE && !METAL_MRT_PROFILE && !METAL_PROFILE && !METAL_SM5_PROFILE
-	return reversebits( bits );
-#else
-	bits = ( bits << 16) | ( bits >> 16);
-	bits = ( (bits & 0x00ff00ff) << 8 ) | ( (bits & 0xff00ff00) >> 8 );
-	bits = ( (bits & 0x0f0f0f0f) << 4 ) | ( (bits & 0xf0f0f0f0) >> 4 );
-	bits = ( (bits & 0x33333333) << 2 ) | ( (bits & 0xcccccccc) >> 2 );
-	bits = ( (bits & 0x55555555) << 1 ) | ( (bits & 0xaaaaaaaa) >> 1 );
-	return bits;
-#endif
-}
-
 float2 Hammersley( uint Index, uint NumSamples, uint2 Random )
 {
 	float E1 = frac( (float)Index / NumSamples + float( Random.x & 0xffff ) / (1<<16) );
