Reverting diffuse model back to Lambert

diff --git a/Engine/Shaders/BRDF.usf b/Engine/Shaders/BRDF.usf
--- a/Engine/Shaders/BRDF.usf
+++ b/Engine/Shaders/BRDF.usf
@@ -7,7 +7,6 @@
 #ifndef __BRDF_COMMON__
 #define __BRDF_COMMON__
 
-
 // Physically based shading model
 // parameterized with the below options
 
@@ -40,6 +39,12 @@
 // 2: Fresnel
 #define PHYSICAL_SPEC_F		1
 
+float Pow5( float x )
+{
+	float x2 = x*x;
+	return x2 * x2 * x;
+}
+
 
 float3 Diffuse_Lambert( float3 DiffuseColor )
 {
@@ -47,74 +52,93 @@ float3 Diffuse_Lambert( float3 DiffuseColor )
 }
 
 // [Burley 2012, "Physically-Based Shading at Disney"]
-// [Lagrade et al. 2014, "Moving Frostbite to Physically Based Rendering"]
 float3 Diffuse_Burley( float3 DiffuseColor, float Roughness, float NoV, float NoL, float VoH )
 {
-	float FD90 = ( 0.5 + 2 * VoH * VoH ) * Roughness;
-	
-	float InvNov = 1- NoV;
-	float NoVPow5 = InvNov * InvNov;
-	NoVPow5 = NoVPow5 * NoVPow5 * InvNov;
-
-	float InvNoL = 1 - NoL;
-	float NoLPow5 = InvNoL * InvNoL;
-	NoLPow5 = NoLPow5 * NoLPow5 * InvNoL;
-
-	float FdV = 1 + (FD90 - 1) * NoVPow5;
-	float FdL = 1 + (FD90 - 1) * NoLPow5;
-	return DiffuseColor * ( 1 / PI * FdV * FdL ) * ( 1 - 0.3333 * Roughness );	// TODO premultiply DiffuseColor
+	float FD90 = 0.5 + 2 * VoH * VoH * Roughness;
+	float FdV = 1 + (FD90 - 1) * Pow5( 1 - NoV );
+	float FdL = 1 + (FD90 - 1) * Pow5( 1 - NoL );
+	return DiffuseColor * ( (1 / PI) * FdV * FdL );
 }
 
 // [Gotanda 2012, "Beyond a Simple Physically Based Blinn-Phong Model in Real-Time"]
 float3 Diffuse_OrenNayar( float3 DiffuseColor, float Roughness, float NoV, float NoL, float VoH )
 {
-	float VoL = 2 * VoH - 1;
-	float m = Roughness * Roughness;
-	float m2 = m * m;
-	float C1 = 1 - 0.5 * m2 / (m2 + 0.33);
+	float a = Roughness * Roughness;
+	float s = a;// / ( 1.29 + 0.5 * a );
+	float s2 = s * s;
+	float VoL = 2 * VoH * VoH - 1;		// double angle identity
 	float Cosri = VoL - NoV * NoL;
-	float C2 = 0.45 * m2 / (m2 + 0.09) * Cosri * ( Cosri >= 0 ? min( 1, NoL / NoV ) : NoL );
-	return DiffuseColor / PI * ( NoL * C1 + C2 );
+	float C1 = 1 - 0.5 * s2 / (s2 + 0.33);
+	float C2 = 0.45 * s2 / (s2 + 0.09) * Cosri * ( Cosri >= 0 ? rcp( max( NoL, NoV ) ) : 1 );
+	return DiffuseColor / PI * ( C1 + C2 ) * ( 1 + Roughness * 0.5 );
+}
+
+// [Gotanda 2014, "Designing Reflectance Models for New Consoles"]
+float3 Diffuse_Gotanda( float3 DiffuseColor, float Roughness, float NoV, float NoL, float VoH )
+{
+	float a = Roughness * Roughness;
+	float a2 = a * a;
+	float F0 = 0.04;
+	float VoL = 2 * VoH * VoH - 1;		// double angle identity
+	float Cosri = VoL - NoV * NoL;
+#if 1
+	float a2_13 = a2 + 1.36053;
+	float Fr = ( 1 - ( 0.542026*a2 + 0.303573*a ) / a2_13 ) * ( 1 - pow( 1 - NoV, 5 - 4*a2 ) / a2_13 ) * ( ( -0.733996*a2*a + 1.50912*a2 - 1.16402*a ) * pow( 1 - NoV, 1 + rcp(39*a2*a2+1) ) + 1 );
+	//float Fr = ( 1 - 0.36 * a ) * ( 1 - pow( 1 - NoV, 5 - 4*a2 ) / a2_13 ) * ( -2.5 * Roughness * ( 1 - NoV ) + 1 );
+	float Lm = ( max( 1 - 2*a, 0 ) * ( 1 - Pow5( 1 - NoL ) ) + min( 2*a, 1 ) ) * ( 1 - 0.5*a * (NoL - 1) ) * NoL;
+	float Vd = ( a2 / ( (a2 + 0.09) * (1.31072 + 0.995584 * NoV) ) ) * ( 1 - pow( 1 - NoL, ( 1 - 0.3726732 * NoV * NoV ) / ( 0.188566 + 0.38841 * NoV ) ) );
+	float Bp = Cosri < 0 ? 1.4 * NoV * NoL * Cosri : Cosri;
+	float Lr = (21.0 / 20.0) * (1 - F0) * ( Fr * Lm + Vd + Bp );
+	return DiffuseColor / PI * Lr;
+#else
+	float a2_13 = a2 + 1.36053;
+	float Fr = ( 1 - ( 0.542026*a2 + 0.303573*a ) / a2_13 ) * ( 1 - pow( 1 - NoV, 5 - 4*a2 ) / a2_13 ) * ( ( -0.733996*a2*a + 1.50912*a2 - 1.16402*a ) * pow( 1 - NoV, 1 + rcp(39*a2*a2+1) ) + 1 );
+	float Lm = ( max( 1 - 2*a, 0 ) * ( 1 - Pow5( 1 - NoL ) ) + min( 2*a, 1 ) ) * ( 1 - 0.5*a + 0.5*a * NoL );
+	float Vd = ( a2 / ( (a2 + 0.09) * (1.31072 + 0.995584 * NoV) ) ) * ( 1 - pow( 1 - NoL, ( 1 - 0.3726732 * NoV * NoV ) / ( 0.188566 + 0.38841 * NoV ) ) );
+	float Bp = Cosri < 0 ? 1.4 * NoV * Cosri : Cosri / max( NoL, 1e-8 );
+	float Lr = (21.0 / 20.0) * (1 - F0) * ( Fr * Lm + Vd + Bp );
+	return DiffuseColor / PI * Lr;
+#endif
 }
 
 // [Blinn 1977, "Models of light reflection for computer synthesized pictures"]
 float D_Blinn( float Roughness, float NoH )
 {
-	float m = Roughness * Roughness;
-	float m2 = m * m;
-	float n = 2 / m2 - 2;
+	float a = Roughness * Roughness;
+	float a2 = a * a;
+	float n = 2 / a2 - 2;
 	return (n+2) / (2*PI) * PhongShadingPow( NoH, n );		// 1 mad, 1 exp, 1 mul, 1 log
 }
 
 // [Beckmann 1963, "The scattering of electromagnetic waves from rough surfaces"]
 float D_Beckmann( float Roughness, float NoH )
 {
-	float m = Roughness * Roughness;
-	float m2 = m * m;
+	float a = Roughness * Roughness;
+	float a2 = a * a;
 	float NoH2 = NoH * NoH;
-	return exp( (NoH2 - 1) / (m2 * NoH2) ) / ( PI * m2 * NoH2 * NoH2 );
+	return exp( (NoH2 - 1) / (a2 * NoH2) ) / ( PI * a2 * NoH2 * NoH2 );
 }
 
 // GGX / Trowbridge-Reitz
 // [Walter et al. 2007, "Microfacet models for refraction through rough surfaces"]
 float D_GGX( float Roughness, float NoH )
 {
-	float m = Roughness * Roughness;
-	float m2 = m * m;
-	float d = ( NoH * m2 - NoH ) * NoH + 1;	// 2 mad
-	return m2 / ( PI*d*d );					// 4 mul, 1 rcp
+	float a = Roughness * Roughness;
+	float a2 = a * a;
+	float d = ( NoH * a2 - NoH ) * NoH + 1;	// 2 mad
+	return a2 / ( PI*d*d );					// 4 mul, 1 rcp
 }
 
 // Anisotropic GGX
 // [Burley 2012, "Physically-Based Shading at Disney"]
 float D_GGXaniso( float RoughnessX, float RoughnessY, float NoH, float3 H, float3 X, float3 Y )
 {
-	float mx = RoughnessX * RoughnessX;
-	float my = RoughnessY * RoughnessY;
+	float ax = RoughnessX * RoughnessX;
+	float ay = RoughnessY * RoughnessY;
 	float XoH = dot( X, H );
 	float YoH = dot( Y, H );
-	float d = XoH*XoH / (mx*mx) + YoH*YoH / (my*my) + NoH*NoH;
-	return 1 / ( PI * mx*my * d*d );
+	float d = XoH*XoH / (ax*ax) + YoH*YoH / (ay*ay) + NoH*NoH;
+	return 1 / ( PI * ax*ay * d*d );
 }
 
 float Vis_Implicit()
@@ -172,12 +196,10 @@ float3 F_None( float3 SpecularColor )
 }
 
 // [Schlick 1994, "An Inexpensive BRDF Model for Physically-Based Rendering"]
-// [Lagarde 2012, "Spherical Gaussian approximation for Blinn-Phong, Phong and Fresnel"]
 float3 F_Schlick( float3 SpecularColor, float VoH )
 {
-	float Fc = pow( 1 - VoH, 5 );							// 1 sub, 3 mul
-	//float Fc = exp2( (-5.55473 * VoH - 6.98316) * VoH );	// 1 mad, 1 mul, 1 exp
-	//return Fc + (1 - Fc) * SpecularColor;					// 1 add, 3 mad
+	float Fc = Pow5( 1 - VoH );					// 1 sub, 3 mul
+	//return Fc + (1 - Fc) * SpecularColor;		// 1 add, 3 mad
 	
 	// Anything less than 2% is physically impossible and is instead considered to be shadowing
 	return saturate( 50.0 * SpecularColor.g ) * Fc + (1 - Fc) * SpecularColor;
@@ -290,4 +312,31 @@ half EnvBRDFApproxNonmetal( half Roughness, half NoV )
 }
 
 
+float D_InvBlinn( float Roughness, float NoH )
+{
+	float m = Roughness * Roughness;
+	float m2 = m * m;
+	float A = 4;
+	float Cos2h = NoH * NoH;
+	float Sin2h = 1 - Cos2h;
+	//return rcp( PI * (1 + A*m2) ) * ( 1 + A * ClampedPow( Sin2h, 1 / m2 - 1 ) );
+	return rcp( PI * (1 + A*m2) ) * ( 1 + A * exp( -Cos2h / m2 ) );
+}
+
+float D_InvBeckmann( float Roughness, float NoH )
+{
+	float m = Roughness * Roughness;
+	float m2 = m * m;
+	float A = 4;
+	float Cos2h = NoH * NoH;
+	float Sin2h = 1 - Cos2h;
+	float Sin4h = Sin2h * Sin2h;
+	return rcp( PI * (1 + A*m2) * Sin4h ) * ( Sin4h + A * exp( -Cos2h / (m2 * Sin2h) ) );
+}
+
+float Vis_Cloth( float NoV, float NoL )
+{
+	return rcp( 4 * ( NoL + NoV - NoL * NoV ) );
+}
+
 #endif // __BRDF_COMMON__
\ No newline at end of file
diff --git a/Engine/Shaders/BasePassPixelShader.usf b/Engine/Shaders/BasePassPixelShader.usf
--- a/Engine/Shaders/BasePassPixelShader.usf
+++ b/Engine/Shaders/BasePassPixelShader.usf
@@ -726,10 +726,6 @@ void Main(
 	}
 #endif
 
-	// TODO Replace with proper EnvBRDF for Burley diffuse model.
-	// Any indirect diffuse done in the base pass does not properly account for the Burley diffuse model besides this crude approximation.
-	DiffuseColor *= 1 - 0.3333 * Roughness;
-
 
 	half3 Color = 0;
 	float IndirectIrradiance = 0;
diff --git a/Engine/Shaders/DeferredLightingCommon.usf b/Engine/Shaders/DeferredLightingCommon.usf
--- a/Engine/Shaders/DeferredLightingCommon.usf
+++ b/Engine/Shaders/DeferredLightingCommon.usf
@@ -12,6 +12,7 @@
 #include "BRDF.usf"
 #include "MonteCarlo.usf"
 #include "IESLightProfilesCommon.usf"
+#include "ShadingModels.usf"
 
 /** 
  * Data about a single light.
@@ -46,274 +47,6 @@ struct FSimpleDeferredLightData
 	bool bInverseSquared;
 };
 
-#if 0
-void StandardShadingShared( float3 DiffuseColor, float3 SpecularColor, float Roughness, float3 V, half3 N )
-{
-	//float NoV = saturate( dot(N, V) );
-	float NoV = abs( dot(N, V) ) + 1e-5;
-
-	// Diffuse_Lambert
-	Shared.DiffuseMul = DiffuseColor * (1.0 / PI);
-
-	// D_GGX, Vis_SmithJointApprox
-	float m = Roughness * Roughness;
-	Shared.m2 = m * m;
-	Shared.SpecularMul = (0.5 / PI) * Shared.m2;
-	Shared.VisMad = float2( 2 * NoV * ( 1 - m ) + m, NoV * m );
-	
-	// F_Schlick
-	Shared.SpecularMul *= saturate( 50.0 * SpecularColor.g );
-}
-
-void StandardShadingPerLight( Shared, float3 L, float3 V, half3 N )
-{
-	float3 H = normalize(V + L);			// 3 add, 2 mad, 4 mul, 1 rsqrt
-	float NoL = saturate( dot(N, L) );		// 2 mad, 1 mul
-	float NoH = saturate( dot(N, H) );		// 2 mad, 1 mul
-	float VoH = saturate( dot(V, H) );		// 2 mad, 1 mul
-
-	// D_GGX, Vis_SmithJointApprox
-	float d = ( NoH * Shared.m2 - NoH ) * NoH + 1;			// 2 mad
-	float v = NoL * Shared.VisMad.x + Shared.VisMad.y;		// 1 mad
-	float D_Vis = Shared.SpecularMul * rcp( d * d * v );	// 3 mul, 1 rcp
-
-	// F_Schlick
-	float Fc = pow( 1 - VoH, 5 );							// 1 sub, 3 mul
-	float3 F = Fc + (1 - Fc) * SpecularColor;				// 1 sub, 3 mad
-
-	return Shared.DiffuseMul + D_Vis * F;					// 3 mad
-}
-#endif
-
-// @param DiffSpecMask .r: diffuse, .g:specular e.g. float2(1,1) for both, float2(1,0) for diffuse only
-float3 StandardShading( FGBufferData GBuffer, float3 LobeRoughness, float3 LobeEnergy, float3 L, float3 V, half3 N, float2 DiffSpecMask )
-{
-	float3 H = normalize(V + L);
-	float NoL = saturate( dot(N, L) );
-	//float NoV = saturate( dot(N, V) );
-	float NoV = abs( dot(N, V) ) + 1e-5;
-	float NoH = saturate( dot(N, H) );
-	float VoH = saturate( dot(V, H) );
-	
-	// Generalized microfacet specular
-	float D = D_GGX( LobeRoughness[1], NoH ) * LobeEnergy[1];
-	float Vis = Vis_SmithJointApprox( LobeRoughness[1], NoV, NoL );
-	float3 F = F_Schlick( GBuffer.SpecularColor, VoH );
-
-	//float3 Diffuse = Diffuse_Lambert( GBuffer.DiffuseColor );
-	float3 Diffuse = Diffuse_Burley( GBuffer.DiffuseColor, LobeRoughness[1], NoV, NoL, VoH );
-	
-	return Diffuse * (LobeEnergy[2] * DiffSpecMask.r) + (D * Vis * DiffSpecMask.g) * F;
-}
-
-float3 SimpleShading( FGBufferData GBuffer, float Roughness, float3 L, float3 V, half3 N )
-{
-	float3 H = normalize(V + L);
-	float NoH = saturate( dot(N, H) );
-	
-	// Generalized microfacet specular
-	float D = D_GGX( Roughness, NoH );
-	float Vis = Vis_Implicit();
-	float3 F = F_None( GBuffer.SpecularColor );
-
-	return Diffuse_Lambert( GBuffer.DiffuseColor ) + (D * Vis) * F;
-}
-
-float3 ClearCoatShading( FGBufferData GBuffer, float3 LobeRoughness, float3 LobeEnergy, float3 L, float3 V, half3 N )
-{
-#if USE_CLEARCOAT
-	const float ClearCoat			= GBuffer.CustomData.x;
-	const float ClearCoatRoughness	= GBuffer.CustomData.y;
-	const float Film = 1 * ClearCoat;
-	const float MetalSpec = 0.9;
-
-#if 1
-	float3 H = normalize(V + L);
-	float NoL = saturate( dot(N, L) );
-	float NoV = saturate( dot(N, V) );
-	float NoH = saturate( dot(N, H) );
-	float VoH = saturate( dot(V, H) );
-	
-	// Generalized microfacet specular
-	float D = D_GGX( LobeRoughness[0], NoH ) * LobeEnergy[0];
-	float Vis = Vis_Kelemen( VoH );
-	
-	// F_Schlick
-	float F0 = 0.04;
-	float Fc = pow( 1 - VoH, 5 );
-	float F = Fc + (1 - Fc) * F0;
-	F *= ClearCoat;
-
-	float Fr1 = D * Vis * F;
-
-	float LayerAttenuation = (1 - F);
-
-	// Generalized microfacet specular
-	float D2 = D_GGX( LobeRoughness[1], NoH ) * LobeEnergy[1];
-	float Vis2 = Vis_Schlick( LobeRoughness[1], NoV, NoL );
-	//float3 F2 = F_Schlick( GBuffer.SpecularColor, VoH );
-	float3 F2 = saturate( 50.0 * GBuffer.SpecularColor.g ) * Fc + (1 - Fc) * GBuffer.SpecularColor;
-
-	float3 Fr2 = Diffuse_Burley( GBuffer.DiffuseColor, LobeRoughness[1], NoV, NoL, VoH ) * LobeEnergy[2] + (D2 * Vis2) * F2;
-	
-	return Fr1 + Fr2 * LayerAttenuation;
-#else
-	float3 H = normalize(V + L);
-	float NoL = saturate( dot(N, L) );
-	float NoV = saturate( dot(N, V) );
-	float NoH = saturate( dot(N, H) );
-	float VoH = saturate( dot(V, H) );
-	
-	// Hard coded IOR of 1.5
-
-	// Generalized microfacet specular
-	float D = D_GGX( ClearCoatRoughness, NoH ) * LobeEnergy[0];
-	float Vis = Vis_Kelemen( VoH );
-	
-	// F_Schlick
-	float F0 = 0.04;
-	float Fc = pow( 1 - VoH, 5 );
-	float F = Fc + (1 - Fc) * F0;
-
-	float Fr1 = D * Vis * F;
-
-	// Refract rays
-	//float3 L2 = refract( -L, -H, 1 / 1.5 );
-	//float3 V2 = refract( -V, -H, 1 / 1.5 );
-
-	// LoH == VoH
-	//float RefractBlend = sqrt( 4 * VoH*VoH + 5 ) / 3 + 2.0 / 3 * VoH;
-	//float3 L2 = RefractBlend * H - L / 1.5;
-	//float3 V2 = RefractBlend * H - V / 1.5;
-	//float NoL2 = saturate( dot(N, L2) );
-	//float NoV2 = saturate( dot(N, V2) );
-	
-	// Approximation
-	float RefractBlend = (0.22 * VoH + 0.7) * VoH + 0.745;	// 2 mad
-	// Dot products distribute. No need for L2 and V2.
-	float RefractNoH = RefractBlend * NoH;					// 1 mul
-	float NoL2 = saturate( RefractNoH - (1 / 1.5) * NoL );	// 1 mad
-	float NoV2 = saturate( RefractNoH - (1 / 1.5) * NoV );	// 1 mad
-	// Should refract H too but unimportant
-
-	NoL2 = max( 0.001, NoL2 );
-	NoV2 = max( 0.001, NoV2 );
-
-	float  AbsorptionDist = rcp(NoV2) + rcp(NoL2);
-	float3 Absorption = pow( AbsorptionColor, 0.5 * AbsorptionDist );
-
-	// Approximation
-	//float  AbsorptionDist = ( NoV2 + NoL2 ) / ( NoV2 * NoL2 );
-	//float3 Absorption = AbsorptionColor * ( AbsorptionColor * (AbsorptionDist * 0.5 - 1) + (2 - 0.5 * AbsorptionDist) );
-	//float3 Absorption = AbsorptionColor + AbsorptionColor * (AbsorptionColor - 1) * (AbsorptionDist * 0.5 - 1);	// use for shared version
-	
-	//float F21 = Fresnel( 1 / 1.5, saturate( dot(V2, H) ) );
-	//float TotalInternalReflection = 1 - F21 * G_Schlick( Roughness, NoV2, NoL2 );
-	//float3 LayerAttenuation = ( (1 - F12) * TotalInternalReflection ) * Absorption;
-
-	// Approximation
-	float3 LayerAttenuation = (1 - F) * Absorption;
-
-	// Approximation for IOR == 1.5
-	//SpecularColor = ChangeBaseMedium( SpecularColor, 1.5 );
-	//SpecularColor = saturate( ( 0.55 * SpecularColor + (0.45 * 1.08) ) * SpecularColor - (0.45 * 0.08) );
-	// Treat SpecularColor as relative to IOR. Artist compensates.
-	
-	// Generalized microfacet specular
-	float D2 = D_GGX( Roughness, NoH ) * LobeEnergy[2];
-	float Vis2 = Vis_Schlick( Roughness, NoV2, NoL2 );
-	float3 F2 = F_Schlick( GBuffer.SpecularColor, VoH );
-
-	float3 Fr2 = Diffuse_Lambert( GBuffer.DiffuseColor ) * LobeEnergy[2] + (D2 * Vis2) * F2;
-	
-	return Fr1 + Fr2 * LayerAttenuation;
-#endif
-#else //USE_CLEARCOAT
-	return float3(0.0f, 0.0f, 0.0f);
-#endif //USE_CLEARCOAT
-}
-
-float3 SubsurfaceShadingSubsurface( FGBufferData GBuffer, float3 L, float3 V, half3 N )
-{
-	float3 SubsurfaceColor = ExtractSubsurfaceColor(GBuffer);
-
-	float3 H = normalize(V + L);
-
-	// to get an effect when you see through the material
-	// hard coded pow constant
-	float InScatter = pow(saturate(dot(L, -V)), 12) * lerp(3, .1f, GBuffer.Opacity);
-	// wrap around lighting, /(PI*2) to be energy consistent (hack do get some view dependnt and light dependent effect)
-	float OpacityFactor = GBuffer.Opacity;
-	// Opacity of 0 gives no normal dependent lighting, Opacity of 1 gives strong normal contribution
-	float NormalContribution = saturate(dot(N, H) * OpacityFactor + 1 - OpacityFactor);
-	float BackScatter = GBuffer.GBufferAO * NormalContribution / (PI * 2);
-
-	// lerp to never exceed 1 (energy conserving)
-	return SubsurfaceColor * lerp(BackScatter, 1, InScatter);
-}
-
-float3 SubsurfaceShadingTwoSided( float3 SubsurfaceColor, float3 L, float3 V, half3 N )
-{
-	float Wrap = 0.5;
-	float NoL = saturate( ( dot(-N, L) + Wrap ) / Square( 1 + Wrap ) );
-
-	float VoL = saturate( dot(V, -L) );
-	float a = 0.6;
-	float a2 = a * a;
-	float d = ( VoL * a2 - VoL ) * VoL + 1;	// 2 mad
-	float GGX = (a2 / PI) / (d * d);		// 2 mul, 1 rcp
-	return NoL * GGX * SubsurfaceColor;
-}
-
-Texture2D		PreIntegratedBRDF;
-SamplerState	PreIntegratedBRDFSampler;
-
-float3 SubsurfaceShadingPreintegratedSkin( FGBufferData GBuffer, float3 L, float3 V, half3 N )
-{
-	float3 SubsurfaceColor = ExtractSubsurfaceColor(GBuffer);
-
-	float OpacityFactor = GBuffer.Opacity;
-	float3 PreintegratedBRDF = Texture2DSampleLevel(PreIntegratedBRDF, PreIntegratedBRDFSampler, float2(saturate(dot(N, L) * .5 + .5), 1 - OpacityFactor), 0).rgb;
-	return PreintegratedBRDF * SubsurfaceColor;
-}
-
-// @param DiffSpecMask .r: diffuse, .g:specular e.g. float2(1,1) for both, float2(1,0) for diffuse only
-float3 SurfaceShading( FGBufferData GBuffer, float3 LobeRoughness, float3 LobeEnergy, float3 L, float3 V, half3 N, float2 DiffSpecMask )
-{
-	switch( GBuffer.ShadingModelID )
-	{
-		case SHADINGMODELID_UNLIT:
-		case SHADINGMODELID_DEFAULT_LIT:
-		case SHADINGMODELID_SUBSURFACE:
-		case SHADINGMODELID_PREINTEGRATED_SKIN:
-		case SHADINGMODELID_SUBSURFACE_PROFILE:
-		case SHADINGMODELID_TWOSIDED_FOLIAGE:
-			return StandardShading( GBuffer, LobeRoughness, LobeEnergy, L, V, N, DiffSpecMask);
-		case SHADINGMODELID_CLEAR_COAT:
-			// this path does not support DiffSpecMask yet
-			return ClearCoatShading( GBuffer, LobeRoughness, LobeEnergy, L, V, N );
-		default:
-			return 0;
-	}
-}
-
-float3 SubsurfaceShading( FGBufferData GBuffer, float3 L, float3 V, half3 N )
-{
-	float3 SubsurfaceColor = ExtractSubsurfaceColor(GBuffer);
-
-	switch( GBuffer.ShadingModelID )
-	{
-		case SHADINGMODELID_SUBSURFACE:
-			return SubsurfaceShadingSubsurface( GBuffer, L, V, N );
-		case SHADINGMODELID_PREINTEGRATED_SKIN:
-			return SubsurfaceShadingPreintegratedSkin( GBuffer, L, V, N );
-		case SHADINGMODELID_TWOSIDED_FOLIAGE:
-			return SubsurfaceShadingTwoSided( SubsurfaceColor, L, V, N );
-		default:
-			return 0;
-	}
-}
-
 
 #undef LIGHT_SOURCE_SHAPE
 #define LIGHT_SOURCE_SHAPE 1
diff --git a/Engine/Shaders/PostProcessAmbient.usf b/Engine/Shaders/PostProcessAmbient.usf
--- a/Engine/Shaders/PostProcessAmbient.usf
+++ b/Engine/Shaders/PostProcessAmbient.usf
@@ -94,7 +94,8 @@ float3 StandardShading( FGBufferData GBuffer, float Roughness, float3 L, float3
 {
 	float3 H = normalize(V + L);
 	float NoL = saturate( dot(N, L) );
-	float NoV = saturate( dot(N, V) );
+	//float NoV = saturate( dot(N, V) );
+	float NoV = max( dot(N, V), 1e-5 );
 	float NoH = saturate( dot(N, H) );
 	float VoH = saturate( dot(V, H) );
 	
@@ -103,7 +104,7 @@ float3 StandardShading( FGBufferData GBuffer, float Roughness, float3 L, float3
 	float Vis = Vis_SmithJointApprox( Roughness, NoV, NoL );
 	float3 F = F_Schlick( GBuffer.SpecularColor, VoH );
 
-	float3 Diffuse = Diffuse_Burley( GBuffer.DiffuseColor, Roughness, NoV, NoL, VoH );
+	float3 Diffuse = Diffuse_Lambert( GBuffer.DiffuseColor );
 
 	return Diffuse + (D * Vis) * F;
 }
@@ -461,7 +462,8 @@ void MainPS(in float4 UVAndScreenPos : TEXCOORD0, out float4 OutColor : SV_Targe
 	float3 V = -ScreenVector;
 	float3 R = 2 * dot( V, N ) * N - V;
 
-	float NoV = abs( dot(N, V) ) + 1e-5;
+	//float NoV = max( dot(N, V), 1e-5 );
+	float NoV = saturate( dot(N, V) );
 
 	// Point lobe in off-specular peak direction
 	float a = Square( GBuffer.Roughness );
@@ -474,19 +476,8 @@ void MainPS(in float4 UVAndScreenPos : TEXCOORD0, out float4 OutColor : SV_Targe
 	float3 NonSpecularContribution = 0;
 	float3 SpecularContribution = 0;
 
-#if 1
-	float3 RetroReflectionDir = lerp( N, V, saturate( ( NoV * (1.02341 * a - 1.51174) + -0.511705 * a + 0.755868 ) * a ) );
-	float3 DiffuseLookup =  TextureCubeSampleLevel(AmbientCubemap, AmbientCubemapSampler, RetroReflectionDir, AbsoluteDiffuseMip).rgb;
-
-	float3 ABC = PreIntegratedGF.SampleLevel( PreIntegratedGFSampler, float2( NoV, GBuffer.Roughness ), 0 ).rgb;
-	//float3 ABC = IntegrateBRDF( Random, GBuffer.Roughness, NoV );
-	
-	NonSpecularContribution += GBuffer.DiffuseColor * DiffuseLookup * ABC.z;
-	//NonSpecularContribution += DiffuseIBL( Random, GBuffer.DiffuseColor, GBuffer.Roughness, N, V );// * ABC.z;
-#else
 	float3 DiffuseLookup =  TextureCubeSampleLevel(AmbientCubemap, AmbientCubemapSampler, N, AbsoluteDiffuseMip).rgb;
 	NonSpecularContribution += GBuffer.DiffuseColor * DiffuseLookup;
-#endif
 
 	// Diffuse
 	{
diff --git a/Engine/Shaders/ShadingModels.usf b/Engine/Shaders/ShadingModels.usf
new file mode 100644
--- /dev/null
+++ b/Engine/Shaders/ShadingModels.usf
@@ -0,0 +1,278 @@
+// Copyright 1998-2015 Epic Games, Inc. All Rights Reserved.
+
+#pragma once
+
+#include "DeferredShadingCommon.usf"
+#include "BRDF.usf"
+
+#if 0
+void StandardShadingShared( float3 DiffuseColor, float3 SpecularColor, float Roughness, float3 V, half3 N )
+{
+	//float NoV = saturate( dot(N, V) );
+	float NoV = abs( dot(N, V) ) + 1e-5;
+
+	// Diffuse_Lambert
+	Shared.DiffuseMul = DiffuseColor * (1.0 / PI);
+
+	// D_GGX, Vis_SmithJointApprox
+	float m = Roughness * Roughness;
+	Shared.m2 = m * m;
+	Shared.SpecularMul = (0.5 / PI) * Shared.m2;
+	Shared.VisMad = float2( 2 * NoV * ( 1 - m ) + m, NoV * m );
+	
+	// F_Schlick
+	Shared.SpecularMul *= saturate( 50.0 * SpecularColor.g );
+}
+
+void StandardShadingPerLight( Shared, float3 L, float3 V, half3 N )
+{
+	float3 H = normalize(V + L);			// 3 add, 2 mad, 4 mul, 1 rsqrt
+	float NoL = saturate( dot(N, L) );		// 2 mad, 1 mul
+	float NoH = saturate( dot(N, H) );		// 2 mad, 1 mul
+	float VoH = saturate( dot(V, H) );		// 2 mad, 1 mul
+
+	// D_GGX, Vis_SmithJointApprox
+	float d = ( NoH * Shared.m2 - NoH ) * NoH + 1;			// 2 mad
+	float v = NoL * Shared.VisMad.x + Shared.VisMad.y;		// 1 mad
+	float D_Vis = Shared.SpecularMul * rcp( d * d * v );	// 3 mul, 1 rcp
+
+	// F_Schlick
+	float Fc = pow( 1 - VoH, 5 );							// 1 sub, 3 mul
+	float3 F = Fc + (1 - Fc) * SpecularColor;				// 1 sub, 3 mad
+
+	return Shared.DiffuseMul + D_Vis * F;					// 3 mad
+}
+#endif
+
+// @param DiffSpecMask .r: diffuse, .g:specular e.g. float2(1,1) for both, float2(1,0) for diffuse only
+float3 StandardShading( FGBufferData GBuffer, float3 LobeRoughness, float3 LobeEnergy, float3 L, float3 V, half3 N, float2 DiffSpecMask )
+{
+	float3 H = normalize(V + L);
+	float NoL = saturate( dot(N, L) );
+	//float NoV = saturate( dot(N, V) );
+	float NoV = max( dot(N, V), 1e-5 );
+	float NoH = saturate( dot(N, H) );
+	float VoH = saturate( dot(V, H) );
+	
+	// Generalized microfacet specular
+	float D = D_GGX( LobeRoughness[1], NoH ) * LobeEnergy[1];
+	float Vis = Vis_SmithJointApprox( LobeRoughness[1], NoV, NoL );
+	float3 F = F_Schlick( GBuffer.SpecularColor, VoH );
+
+	float3 Diffuse = Diffuse_Lambert( GBuffer.DiffuseColor );
+	//float3 Diffuse = Diffuse_Burley( GBuffer.DiffuseColor, LobeRoughness[1], NoV, NoL, VoH );
+	//float3 Diffuse = Diffuse_OrenNayar( GBuffer.DiffuseColor, LobeRoughness[1], NoV, NoL, VoH );
+	
+	return Diffuse * (LobeEnergy[2] * DiffSpecMask.r) + (D * Vis * DiffSpecMask.g) * F;
+}
+
+float3 SimpleShading( FGBufferData GBuffer, float Roughness, float3 L, float3 V, half3 N )
+{
+	float3 H = normalize(V + L);
+	float NoH = saturate( dot(N, H) );
+	
+	// Generalized microfacet specular
+	float D = D_GGX( Roughness, NoH );
+	float Vis = Vis_Implicit();
+	float3 F = F_None( GBuffer.SpecularColor );
+
+	return Diffuse_Lambert( GBuffer.DiffuseColor ) + (D * Vis) * F;
+}
+
+float3 ClearCoatShading( FGBufferData GBuffer, float3 LobeRoughness, float3 LobeEnergy, float3 L, float3 V, half3 N )
+{
+#if USE_CLEARCOAT
+	const float ClearCoat			= GBuffer.CustomData.x;
+	const float ClearCoatRoughness	= GBuffer.CustomData.y;
+	const float Film = 1 * ClearCoat;
+	const float MetalSpec = 0.9;
+
+#if 1
+	float3 H = normalize(V + L);
+	float NoL = saturate( dot(N, L) );
+	float NoV = saturate( dot(N, V) );
+	float NoH = saturate( dot(N, H) );
+	float VoH = saturate( dot(V, H) );
+	
+	// Generalized microfacet specular
+	float D = D_GGX( LobeRoughness[0], NoH ) * LobeEnergy[0];
+	float Vis = Vis_Kelemen( VoH );
+	
+	// F_Schlick
+	float F0 = 0.04;
+	float Fc = Pow5( 1 - VoH );
+	float F = Fc + (1 - Fc) * F0;
+	F *= ClearCoat;
+
+	float Fr1 = D * Vis * F;
+
+	float LayerAttenuation = (1 - F);
+
+	// Generalized microfacet specular
+	float D2 = D_GGX( LobeRoughness[1], NoH ) * LobeEnergy[1];
+	float Vis2 = Vis_Schlick( LobeRoughness[1], NoV, NoL );
+	//float3 F2 = F_Schlick( GBuffer.SpecularColor, VoH );
+	float3 F2 = saturate( 50.0 * GBuffer.SpecularColor.g ) * Fc + (1 - Fc) * GBuffer.SpecularColor;
+
+	//float3 Fr2 = Diffuse_Burley( GBuffer.DiffuseColor, LobeRoughness[1], NoV, NoL, VoH ) * LobeEnergy[2] + (D2 * Vis2) * F2;
+	float3 Fr2 = Diffuse_Lambert( GBuffer.DiffuseColor ) * LobeEnergy[2] + (D2 * Vis2) * F2;
+	
+	return Fr1 + Fr2 * LayerAttenuation;
+#else
+	float3 H = normalize(V + L);
+	float NoL = saturate( dot(N, L) );
+	float NoV = saturate( dot(N, V) );
+	float NoH = saturate( dot(N, H) );
+	float VoH = saturate( dot(V, H) );
+	
+	// Hard coded IOR of 1.5
+
+	// Generalized microfacet specular
+	float D = D_GGX( ClearCoatRoughness, NoH ) * LobeEnergy[0];
+	float Vis = Vis_Kelemen( VoH );
+	
+	// F_Schlick
+	float F0 = 0.04;
+	float Fc = Pow5( 1 - VoH );
+	float F = Fc + (1 - Fc) * F0;
+
+	float Fr1 = D * Vis * F;
+
+	// Refract rays
+	//float3 L2 = refract( -L, -H, 1 / 1.5 );
+	//float3 V2 = refract( -V, -H, 1 / 1.5 );
+
+	// LoH == VoH
+	//float RefractBlend = sqrt( 4 * VoH*VoH + 5 ) / 3 + 2.0 / 3 * VoH;
+	//float3 L2 = RefractBlend * H - L / 1.5;
+	//float3 V2 = RefractBlend * H - V / 1.5;
+	//float NoL2 = saturate( dot(N, L2) );
+	//float NoV2 = saturate( dot(N, V2) );
+	
+	// Approximation
+	float RefractBlend = (0.22 * VoH + 0.7) * VoH + 0.745;	// 2 mad
+	// Dot products distribute. No need for L2 and V2.
+	float RefractNoH = RefractBlend * NoH;					// 1 mul
+	float NoL2 = saturate( RefractNoH - (1 / 1.5) * NoL );	// 1 mad
+	float NoV2 = saturate( RefractNoH - (1 / 1.5) * NoV );	// 1 mad
+	// Should refract H too but unimportant
+
+	NoL2 = max( 0.001, NoL2 );
+	NoV2 = max( 0.001, NoV2 );
+
+	float  AbsorptionDist = rcp(NoV2) + rcp(NoL2);
+	float3 Absorption = pow( AbsorptionColor, 0.5 * AbsorptionDist );
+
+	// Approximation
+	//float  AbsorptionDist = ( NoV2 + NoL2 ) / ( NoV2 * NoL2 );
+	//float3 Absorption = AbsorptionColor * ( AbsorptionColor * (AbsorptionDist * 0.5 - 1) + (2 - 0.5 * AbsorptionDist) );
+	//float3 Absorption = AbsorptionColor + AbsorptionColor * (AbsorptionColor - 1) * (AbsorptionDist * 0.5 - 1);	// use for shared version
+	
+	//float F21 = Fresnel( 1 / 1.5, saturate( dot(V2, H) ) );
+	//float TotalInternalReflection = 1 - F21 * G_Schlick( Roughness, NoV2, NoL2 );
+	//float3 LayerAttenuation = ( (1 - F12) * TotalInternalReflection ) * Absorption;
+
+	// Approximation
+	float3 LayerAttenuation = (1 - F) * Absorption;
+
+	// Approximation for IOR == 1.5
+	//SpecularColor = ChangeBaseMedium( SpecularColor, 1.5 );
+	//SpecularColor = saturate( ( 0.55 * SpecularColor + (0.45 * 1.08) ) * SpecularColor - (0.45 * 0.08) );
+	// Treat SpecularColor as relative to IOR. Artist compensates.
+	
+	// Generalized microfacet specular
+	float D2 = D_GGX( Roughness, NoH ) * LobeEnergy[2];
+	float Vis2 = Vis_Schlick( Roughness, NoV2, NoL2 );
+	float3 F2 = F_Schlick( GBuffer.SpecularColor, VoH );
+
+	float3 Fr2 = Diffuse_Lambert( GBuffer.DiffuseColor ) * LobeEnergy[2] + (D2 * Vis2) * F2;
+	
+	return Fr1 + Fr2 * LayerAttenuation;
+#endif
+#else //USE_CLEARCOAT
+	return float3(0.0f, 0.0f, 0.0f);
+#endif //USE_CLEARCOAT
+}
+
+float3 SubsurfaceShadingSubsurface( FGBufferData GBuffer, float3 L, float3 V, half3 N )
+{
+	float3 SubsurfaceColor = ExtractSubsurfaceColor(GBuffer);
+
+	float3 H = normalize(V + L);
+
+	// to get an effect when you see through the material
+	// hard coded pow constant
+	float InScatter = pow(saturate(dot(L, -V)), 12) * lerp(3, .1f, GBuffer.Opacity);
+	// wrap around lighting, /(PI*2) to be energy consistent (hack do get some view dependnt and light dependent effect)
+	float OpacityFactor = GBuffer.Opacity;
+	// Opacity of 0 gives no normal dependent lighting, Opacity of 1 gives strong normal contribution
+	float NormalContribution = saturate(dot(N, H) * OpacityFactor + 1 - OpacityFactor);
+	float BackScatter = GBuffer.GBufferAO * NormalContribution / (PI * 2);
+
+	// lerp to never exceed 1 (energy conserving)
+	return SubsurfaceColor * lerp(BackScatter, 1, InScatter);
+}
+
+float3 SubsurfaceShadingTwoSided( float3 SubsurfaceColor, float3 L, float3 V, half3 N )
+{
+	// http://blog.stevemcauley.com/2011/12/03/energy-conserving-wrapped-diffuse/
+	float Wrap = 0.5;
+	float NoL = saturate( ( dot(-N, L) + Wrap ) / Square( 1 + Wrap ) );
+
+	// GGX scatter distribution
+	float VoL = saturate( dot(V, -L) );
+	float a = 0.6;
+	float a2 = a * a;
+	float d = ( VoL * a2 - VoL ) * VoL + 1;	// 2 mad
+	float GGX = (a2 / PI) / (d * d);		// 2 mul, 1 rcp
+	return NoL * GGX * SubsurfaceColor;
+}
+
+Texture2D		PreIntegratedBRDF;
+SamplerState	PreIntegratedBRDFSampler;
+
+float3 SubsurfaceShadingPreintegratedSkin( FGBufferData GBuffer, float3 L, float3 V, half3 N )
+{
+	float3 SubsurfaceColor = ExtractSubsurfaceColor(GBuffer);
+
+	float OpacityFactor = GBuffer.Opacity;
+	float3 PreintegratedBRDF = Texture2DSampleLevel(PreIntegratedBRDF, PreIntegratedBRDFSampler, float2(saturate(dot(N, L) * .5 + .5), 1 - OpacityFactor), 0).rgb;
+	return PreintegratedBRDF * SubsurfaceColor;
+}
+
+// @param DiffSpecMask .r: diffuse, .g:specular e.g. float2(1,1) for both, float2(1,0) for diffuse only
+float3 SurfaceShading( FGBufferData GBuffer, float3 LobeRoughness, float3 LobeEnergy, float3 L, float3 V, half3 N, float2 DiffSpecMask )
+{
+	switch( GBuffer.ShadingModelID )
+	{
+		case SHADINGMODELID_UNLIT:
+		case SHADINGMODELID_DEFAULT_LIT:
+		case SHADINGMODELID_SUBSURFACE:
+		case SHADINGMODELID_PREINTEGRATED_SKIN:
+		case SHADINGMODELID_SUBSURFACE_PROFILE:
+		case SHADINGMODELID_TWOSIDED_FOLIAGE:
+			return StandardShading( GBuffer, LobeRoughness, LobeEnergy, L, V, N, DiffSpecMask);
+		case SHADINGMODELID_CLEAR_COAT:
+			// this path does not support DiffSpecMask yet
+			return ClearCoatShading( GBuffer, LobeRoughness, LobeEnergy, L, V, N );
+		default:
+			return 0;
+	}
+}
+
+float3 SubsurfaceShading( FGBufferData GBuffer, float3 L, float3 V, half3 N )
+{
+	float3 SubsurfaceColor = ExtractSubsurfaceColor(GBuffer);
+
+	switch( GBuffer.ShadingModelID )
+	{
+		case SHADINGMODELID_SUBSURFACE:
+			return SubsurfaceShadingSubsurface( GBuffer, L, V, N );
+		case SHADINGMODELID_PREINTEGRATED_SKIN:
+			return SubsurfaceShadingPreintegratedSkin( GBuffer, L, V, N );
+		case SHADINGMODELID_TWOSIDED_FOLIAGE:
+			return SubsurfaceShadingTwoSided( SubsurfaceColor, L, V, N );
+		default:
+			return 0;
+	}
+}
\ No newline at end of file
diff --git a/Engine/Shaders/SkyLighting.usf b/Engine/Shaders/SkyLighting.usf
--- a/Engine/Shaders/SkyLighting.usf
+++ b/Engine/Shaders/SkyLighting.usf
@@ -38,22 +38,6 @@ void SkyLightDiffusePS(
 		float3 SkyLightingNormal = GBuffer.WorldNormal;
 		float3 DiffuseIrradiance = 0;
 
-		// screen position in [-1, 1] screen space
-		float2 ScreenSpacePos = UVAndScreenPos.zw;
-		float3 ScreenVector = normalize(mul(float4(ScreenSpacePos, 1, 0), View.ScreenToWorld).xyz);
-
-		float3 N = GBuffer.WorldNormal;
-		float3 V = -ScreenVector;
-		float NoV = abs( dot(N, V) ) + 1e-5;
-
-		float a = Square( GBuffer.Roughness );
-		float3 DominantDiffuseDir = lerp( N, V, saturate( ( NoV * (1.02341 * a - 1.51174) + -0.511705 * a + 0.755868 ) * a ) );
-		DominantDiffuseDir = normalize( DominantDiffuseDir );
-		SkyLightingNormal = DominantDiffuseDir;
-
-		float3 ABC = PreIntegratedGF.SampleLevel( PreIntegratedGFSampler, float2( NoV, GBuffer.Roughness ), 0 ).rgb;
-		GBuffer.DiffuseColor *= ABC.z;
-
 #if APPLY_SHADOWING
 		float3 BentNormal = Texture2DSampleLevel(BentNormalAOTexture, BentNormalAOSampler, UV, 0).xyz;
 
@@ -67,8 +51,8 @@ void SkyLightDiffusePS(
 			// Use more bent normal in corners
 			float BentNormalWeightFactor = SkyVisibility;
 
-			SkyLightingNormal = lerp(NormalizedBentNormal, DominantDiffuseDir, BentNormalWeightFactor);
-			DotProductFactor = lerp(dot(NormalizedBentNormal, DominantDiffuseDir), 1, BentNormalWeightFactor);
+			SkyLightingNormal = lerp(NormalizedBentNormal, GBuffer.WorldNormal, BentNormalWeightFactor);
+			DotProductFactor = lerp(dot(NormalizedBentNormal, GBuffer.WorldNormal), 1, BentNormalWeightFactor);
 		}
 		#else
 		{
diff --git a/Engine/Source/Runtime/Renderer/Private/DistanceFieldSurfaceCacheLighting.cpp b/Engine/Source/Runtime/Renderer/Private/DistanceFieldSurfaceCacheLighting.cpp
--- a/Engine/Source/Runtime/Renderer/Private/DistanceFieldSurfaceCacheLighting.cpp
+++ b/Engine/Source/Runtime/Renderer/Private/DistanceFieldSurfaceCacheLighting.cpp
@@ -3078,8 +3078,6 @@ public:
 		: FGlobalShader(Initializer)
 	{
 		DeferredParameters.Bind(Initializer.ParameterMap);
-		PreIntegratedGF.Bind(Initializer.ParameterMap, TEXT("PreIntegratedGF"));
-		PreIntegratedGFSampler.Bind(Initializer.ParameterMap, TEXT("PreIntegratedGFSampler"));
 		DynamicBentNormalAOTexture.Bind(Initializer.ParameterMap, TEXT("BentNormalAOTexture"));
 		DynamicBentNormalAOSampler.Bind(Initializer.ParameterMap, TEXT("BentNormalAOSampler"));
 		DynamicIrradianceTexture.Bind(Initializer.ParameterMap, TEXT("IrradianceTexture"));
@@ -3094,8 +3092,6 @@ public:
 		FGlobalShader::SetParameters(RHICmdList, ShaderRHI, View);
 		DeferredParameters.Set(RHICmdList, ShaderRHI, View);
 
-		SetTextureParameter(RHICmdList, ShaderRHI, PreIntegratedGF, PreIntegratedGFSampler, TStaticSamplerState<SF_Bilinear,AM_Clamp,AM_Clamp,AM_Clamp>::GetRHI(), GSystemTextures.PreintegratedGF->GetRenderTargetItem().ShaderResourceTexture );
-
 		SetTextureParameter(RHICmdList, ShaderRHI, DynamicBentNormalAOTexture, DynamicBentNormalAOSampler, TStaticSamplerState<SF_Point>::GetRHI(), DynamicBentNormalAO);
 
 		if (DynamicIrradianceTexture.IsBound())
@@ -3121,7 +3117,6 @@ public:
 	{
 		bool bShaderHasOutdatedParameters = FGlobalShader::Serialize(Ar);
 		Ar << DeferredParameters;
-		Ar << PreIntegratedGF << PreIntegratedGFSampler;
 		Ar << DynamicBentNormalAOTexture;
 		Ar << DynamicBentNormalAOSampler;
 		Ar << DynamicIrradianceTexture;
@@ -3134,8 +3129,6 @@ public:
 private:
 
 	FDeferredPixelShaderParameters DeferredParameters;
-	FShaderResourceParameter PreIntegratedGF;
-	FShaderResourceParameter PreIntegratedGFSampler;
 	FShaderResourceParameter DynamicBentNormalAOTexture;
 	FShaderResourceParameter DynamicBentNormalAOSampler;
 	FShaderResourceParameter DynamicIrradianceTexture;
diff --git a/Engine/Source/Runtime/Renderer/Private/SystemTextures.cpp b/Engine/Source/Runtime/Renderer/Private/SystemTextures.cpp
--- a/Engine/Source/Runtime/Renderer/Private/SystemTextures.cpp
+++ b/Engine/Source/Runtime/Renderer/Private/SystemTextures.cpp
@@ -258,11 +258,9 @@ void FSystemTextures::InitializeTextures(FRHICommandListImmediate& RHICmdList, E
 
 			EPixelFormat Format = PF_R8G8;
 			// for low roughness we would get banding with PF_R8G8 but for low spec it could be used, for now we don't do this optimization
-			if (GPixelFormats[PF_A16B16G16R16].Supported)
+			if (GPixelFormats[PF_G16R16].Supported)
 			{
-				// 3rd channel required for diffuse
-				// TODO try 11:11:10
-				Format = PF_A16B16G16R16;
+				Format = PF_G16R16;
 			}
 
 			FPooledRenderTargetDesc Desc(FPooledRenderTargetDesc::Create2DDesc(FIntPoint(128, 32), Format, TexCreate_FastVRAM, TexCreate_None, false));
@@ -346,10 +344,10 @@ void FSystemTextures::InitializeTextures(FRHICommandListImmediate& RHICmdList, E
 							float NoH = FMath::Max(H.Z, 0.0f);
 							float VoH = FMath::Max(V | H, 0.0f);
 
-							float FD90 = ( 0.5f + 2.0f * VoH * VoH ) * Roughness;
+							float FD90 = 0.5f + 2.0f * VoH * VoH * Roughness;
 							float FdV = 1.0f + (FD90 - 1.0f) * pow( 1.0f - NoV, 5 );
 							float FdL = 1.0f + (FD90 - 1.0f) * pow( 1.0f - NoL, 5 );
-							C += FdV * FdL * ( 1.0f - 0.3333f * Roughness );
+							C += FdV * FdL;// * ( 1.0f - 0.3333f * Roughness );
 						}
 					}
 					A /= NumSamples;
