More fixes for shader compile errors in OpenGL 3

diff --git a/Engine/Shaders/FastMath.usf b/Engine/Shaders/FastMath.usf
--- a/Engine/Shaders/FastMath.usf
+++ b/Engine/Shaders/FastMath.usf
@@ -73,9 +73,13 @@ float sqrtFast( float x )
 // 1 ALU
 float rcpFast( float x )
 {
+#if !GL3_PROFILE
 	int i = asint(x);
 	i = 0x7EF311C2 - i;
 	return asfloat(i);
+#else
+	return rcp(x);
+#endif
 }
 
 // Using 1 Newton Raphson iterations
@@ -84,11 +88,15 @@ float rcpFast( float x )
 // 3 ALU
 float rcpFastNR1( float x )
 {
+#if !GL3_PROFILE
 	int i = asint(x);
 	i = 0x7EF311C3 - i;
 	float xRcp = asfloat(i);
 	xRcp = xRcp * (-xRcp * x + 2.0f);
 	return xRcp;
+#else
+	return rcp(x);
+#endif
 }
 
 float lengthFast( float3 v )
