Fixed dynamic point light aliasing issue due to half float precision

diff --git a/Engine/Shaders/BasePassForForwardShadingPixelShader.usf b/Engine/Shaders/BasePassForForwardShadingPixelShader.usf
--- a/Engine/Shaders/BasePassForForwardShadingPixelShader.usf
+++ b/Engine/Shaders/BasePassForForwardShadingPixelShader.usf
@@ -272,12 +272,12 @@ void Main(
 			{
 				float3 ToLight = LightPositionAndInvRadius[i].xyz - MaterialParameters.WorldPosition;
 				float DistanceSqr = dot(ToLight, ToLight);
-				half3 L = ToLight * rsqrt(DistanceSqr);
+				float3 L = ToLight * rsqrt(DistanceSqr);
 
-				half PointNoL = max(0, dot(MaterialParameters.WorldNormal, L));
-				half PointRoL = max(0, dot(MaterialParameters.ReflectionVector, L));
+				float PointNoL = max(0, dot(MaterialParameters.WorldNormal, L));
+				float PointRoL = max(0, dot(MaterialParameters.ReflectionVector, L));
 
-				half Attenuation;
+				float Attenuation;
 
 				if (LightColorAndFalloffExponent[i].w == 0)
 				{
