missing file for former checkin

diff --git a/Engine/Shaders/MiniFontCommon.usf b/Engine/Shaders/MiniFontCommon.usf
--- a/Engine/Shaders/MiniFontCommon.usf
+++ b/Engine/Shaders/MiniFontCommon.usf
@@ -58,4 +58,24 @@ void PrintFloat(int2 PixelPos, inout float3 OutColor, float3 FontColor, int2 Lef
 	PrintCharacter(PixelPos, OutColor, FontColor, Cursor, ExtractDigitFromFloat(Number, 0.001));
 }
 
+// for printf debugging in the shader, has to be positive
+// outputs a float number in the form: xxx.yyy
+// @param LeftTop - in pixels
+void PrintFloatNoFraction(int2 PixelPos, inout float3 OutColor, float3 FontColor, int2 LeftTop, float Number)
+{
+	int2 Cursor = LeftTop;
+
+	float BorderDistance = ComputeDistanceToRect(PixelPos, LeftTop, int2(7, 1) * 8 - 1);
+
+	// black border around number
+	//	OutColor = lerp(0, OutColor, saturate(BorderDistance - 2));
+
+	// before period
+	PrintCharacter(PixelPos, OutColor, FontColor, Cursor, ExtractDigitFromFloat(Number, 10000));
+	PrintCharacter(PixelPos, OutColor, FontColor, Cursor, ExtractDigitFromFloat(Number, 1000));
+	PrintCharacter(PixelPos, OutColor, FontColor, Cursor, ExtractDigitFromFloat(Number, 100));
+	PrintCharacter(PixelPos, OutColor, FontColor, Cursor, ExtractDigitFromFloat(Number, 10));
+	PrintCharacter(PixelPos, OutColor, FontColor, Cursor, ExtractDigitFromFloat(Number, 1));
+}
+
 #endif // __TINY_FONT_COMMON__
