Improved quality Temporal AA. Now using bicubic filtering.

High precision reprojection.

diff --git a/Engine/Shaders/PostProcessTemporalAA.usf b/Engine/Shaders/PostProcessTemporalAA.usf
--- a/Engine/Shaders/PostProcessTemporalAA.usf
+++ b/Engine/Shaders/PostProcessTemporalAA.usf
@@ -6,7 +6,6 @@
 
 #include "Common.usf"
 #include "PostProcessCommon.usf"
-#include "DeferredShadingCommon.usf"		// FGBufferData
 
 #ifndef ENABLE_TEMPORAL_AA
 #define ENABLE_TEMPORAL_AA	1
@@ -20,6 +19,30 @@ float2 RandomOffset;
 // TODO: This can be removed.
 float ClampFeebackMix;
 
+float3 RGBToYCoCg( float3 RGB )
+{
+	float Y  = dot( RGB, float3(  1, 2,  1 ) );
+	float Co = dot( RGB, float3(  2, 0, -2 ) );
+	float Cg = dot( RGB, float3( -1, 2, -1 ) );
+	
+	float3 YCoCg = float3( Y, Co, Cg );
+	return YCoCg;
+}
+
+float3 YCoCgToRGB( float3 YCoCg )
+{
+	float Y  = YCoCg.x * 0.25;
+	float Co = YCoCg.y * 0.25;
+	float Cg = YCoCg.z * 0.25;
+
+	float R = Y + Co - Cg;
+	float G = Y + Cg;
+	float B = Y - Co - Cg;
+
+	float3 RGB = float3( R, G, B );
+	return RGB;
+}
+
 // Faster but less accurate luma computation. 
 // Luma includes a scaling by 4.
 float Luma4(float3 Color)
@@ -184,7 +207,42 @@ float4 HdrLerp(float4 ColorA, float4 ColorB, float Blend, float Exposure)
 	return ColorA * BlendA + ColorB * BlendB;
 }
 
+void Bicubic2DCatmullRom( in float2 UV, in float2 Size, out float2 Sample[3], out float2 Weight[3] )
+{
+	const float2 InvSize = 1.0 / Size;
+ 
+	UV *= Size;
+
+	float2 tc = floor( UV - 0.5 ) + 0.5;
+	float2 f = UV - tc;
+	float2 f2 = f * f;
+	float2 f3 = f2 * f;
+
+	float2 w0 = f2 - 0.5 * (f3 + f);
+	float2 w1 = 1.5 * f3 - 2.5 * f2 + 1;
+	float2 w3 = 0.5 * (f3 - f2);
+	float2 w2 = 1 - w0 - w1 - w3;
+
+	Weight[0] = w0;
+	Weight[1] = w1 + w2;
+	Weight[2] = w3;
+
+	Sample[0] = tc - 1;
+	Sample[1] = tc + w2 / Weight[1];
+	Sample[2] = tc + 2;
+ 
+	Sample[0] *= InvSize;
+	Sample[1] *= InvSize;
+	Sample[2] *= InvSize;
+}
 
+float RandFast( uint2 PixelPos, float Magic = 3571.0 )
+{
+	float2 Random = ( 1.0 / 4320.0 ) * PixelPos + float2( 0.25, 0.0 );
+	Random = frac( dot( Random * Random, Magic ) );
+	Random = frac( dot( Random * Random, Magic ) );
+	return Random.x;
+}
 
 void SSRTemporalAAPS( float4 UVAndScreenPos : TEXCOORD0, float3 InExposureScaleVignette : TEXCOORD1, out float4 OutColor : SV_Target0 )
 {
@@ -238,39 +296,59 @@ void LightShaftTemporalAAPS( float4 UVAndScreenPos : TEXCOORD0, float3 InExposur
 #endif
 }
 
-void MainTemporalAAPS( float4 UVAndScreenPos : TEXCOORD0, float3 InExposureScaleVignette : TEXCOORD1, out float4 OutColor : SV_Target0 )
+void MainTemporalAAPS(
+	in float4 UVAndScreenPos : TEXCOORD0,
+	in float3 InExposureScaleVignette : TEXCOORD1,
+	in float4 SvPosition : SV_Position,
+	out float4 OutColor : SV_Target0 )
 {
 #if ENABLE_TEMPORAL_AA
 	float InExposureScale = InExposureScaleVignette.x;
 	#define AA_FILTERED 1
 	#define AA_BORDER 1
+	#define AA_ALPHA 0
 	#define AA_CROSS 2
 	#define AA_GREEN_AS_LUMA 1
 	#define AA_AABB 1
 	#define AA_LOWPASS 0
 	#define AA_DEBUG 0
-	#define AA_VELOCITY_WEIGHTING 1
+	#define AA_VELOCITY_WEIGHTING 0
 	#define AA_YCOCG 1
+	#define AA_BICUBIC 1
 	#include "PostProcessTemporalCommon.usf"
+	
+	uint2 PixelPos = SvPosition.xy + uint2( View.FrameNumber & 3, (View.FrameNumber / 4) & 3 );
+	float Random = RandFast( PixelPos );
+	OutColor += (Random - 0.5) * (1.0 / 1024.0) / InExposureScale;
 #else 
 	OutColor = PostprocessInput0.SampleLevel(PostprocessInput0Sampler, UVAndScreenPos.xy, 0);
 #endif
 }
 
-void MainFastTemporalAAPS( float4 UVAndScreenPos : TEXCOORD0, float3 InExposureScaleVignette : TEXCOORD1, out float4 OutColor : SV_Target0 )
+void MainFastTemporalAAPS(
+	in float4 UVAndScreenPos : TEXCOORD0,
+	in float3 InExposureScaleVignette : TEXCOORD1,
+	in float4 SvPosition : SV_Position,
+	out float4 OutColor : SV_Target0 )
 {
 #if ENABLE_TEMPORAL_AA
 	float InExposureScale = InExposureScaleVignette.x;
 	#define AA_FILTERED 1
 	#define AA_BORDER 1
+	#define AA_ALPHA 0
 	#define AA_CROSS 2
-	#define AA_GREEN_AS_LUMA 1
+	#define AA_GREEN_AS_LUMA 0
 	#define AA_AABB 1
 	#define AA_LOWPASS 0
 	#define AA_DEBUG 0
 	#define AA_VELOCITY_WEIGHTING 0
 	#define AA_YCOCG 1
+	#define AA_BICUBIC 0
 	#include "PostProcessTemporalCommon.usf"
+
+	uint2 PixelPos = SvPosition.xy + uint2( View.FrameNumber & 3, (View.FrameNumber / 4) & 3 );
+	float Random = RandFast( PixelPos );
+	OutColor += (Random - 0.5) * (1.0 / 1024.0) / InExposureScale;
 #else 
 	OutColor = PostprocessInput0.SampleLevel(PostprocessInput0Sampler, UVAndScreenPos.xy, 0);
 #endif
diff --git a/Engine/Shaders/PostProcessTemporalCommon.usf b/Engine/Shaders/PostProcessTemporalCommon.usf
--- a/Engine/Shaders/PostProcessTemporalCommon.usf
+++ b/Engine/Shaders/PostProcessTemporalCommon.usf
@@ -54,23 +54,6 @@ which can be configured via defines for various temporal AA passes.
 	#define AA_FILTERED 1
 #endif
 
-// 1 = Improve quality on converged edges.
-// 0 = Default.
-#ifndef AA_GRAIN
-	#define AA_GRAIN 0
-#endif
-
-// 1 = Use higher quality HDR weighted filtering.
-// 0 = Don't use.
-#ifndef AA_HDR
-	#define AA_HDR 1
-#endif
-
-// 1 = Use manual HDR weighted bilinear filtering for history.
-#if AA_HDR_HISTORY
-	#define AA_HDR_HISTORY 1
-#endif
-
 // 0 = Dynamic motion based lerp value (default).
 // non-zero = Use 1/LERP fixed lerp value (used for reflections).
 #ifndef AA_LERP
@@ -105,11 +88,6 @@ which can be configured via defines for various temporal AA passes.
 	#define AA_FORCE_ALPHA_CLAMP 0
 #endif
 
-// Use velocity weighting.
-#ifndef AA_VELOCITY_WEIGHTING
-	#define AA_VELOCITY_WEIGHTING 0
-#endif
-
 // Use YCoCg path.
 #ifndef AA_YCOCG
 	#define AA_YCOCG 0
@@ -120,21 +98,21 @@ which can be configured via defines for various temporal AA passes.
 	#define AA_GREEN_AS_LUMA AA_YCOCG
 #endif
 
+// Bicubic filter history
+#ifndef AA_BICUBIC
+	#define AA_BICUBIC 0
+#endif
+
 // Special adjustments for DOF.
 #ifndef AA_DOF
 	#define AA_DOF 0
 #endif
 
-// Completely non-optimized prototype for solving noise problem.
-#ifndef AA_TEST
-	#define AA_TEST 0
-#endif
-
 
 	// FIND MOTION OF PIXEL AND NEAREST IN NEIGHBORHOOD
     // ------------------------------------------------
 	float3 PosN; // Position of this pixel, possibly later nearest pixel in neighborhood.
-	PosN.xy = UVAndScreenPos.zw * float2(0.5, -0.5) + 0.5; // View position [0 to 1] flipped in Y.
+	PosN.xy = UVAndScreenPos.zw;
 	PosN.z = SceneDepthTexture.SampleLevel(SceneDepthTextureSampler, UVAndScreenPos.xy, 0).r;
 	// Screen position of minimum depth.
 	float2 VelocityOffset = float2(0.0, 0.0);
@@ -143,18 +121,14 @@ which can be configured via defines for various temporal AA passes.
 		// This enables better quality outline on foreground against different motion background.
 		// Larger 2 pixel distance "x" works best (because AA dilates surface).
 		float4 Depths;
-		#if FEATURE_LEVEL >= FEATURE_LEVEL_SM5
-			Depths = SceneDepthTexture.GatherRed(SceneDepthTextureSampler, UVAndScreenPos.xy, int2(-AA_CROSS, -AA_CROSS), int2(AA_CROSS, -AA_CROSS), int2(-AA_CROSS, AA_CROSS), int2(AA_CROSS, AA_CROSS));
-		#else
-			Depths.x = SceneDepthTexture.SampleLevel(SceneDepthTextureSampler, UVAndScreenPos.xy, 0, int2(-AA_CROSS, -AA_CROSS)).r;
-			Depths.y = SceneDepthTexture.SampleLevel(SceneDepthTextureSampler, UVAndScreenPos.xy, 0, int2( AA_CROSS, -AA_CROSS)).r;
-			Depths.z = SceneDepthTexture.SampleLevel(SceneDepthTextureSampler, UVAndScreenPos.xy, 0, int2(-AA_CROSS,  AA_CROSS)).r;
-			Depths.w = SceneDepthTexture.SampleLevel(SceneDepthTextureSampler, UVAndScreenPos.xy, 0, int2( AA_CROSS,  AA_CROSS)).r;
-		#endif
+		Depths.x = SceneDepthTexture.SampleLevel(SceneDepthTextureSampler, UVAndScreenPos.xy, 0, int2(-AA_CROSS, -AA_CROSS)).r;
+		Depths.y = SceneDepthTexture.SampleLevel(SceneDepthTextureSampler, UVAndScreenPos.xy, 0, int2( AA_CROSS, -AA_CROSS)).r;
+		Depths.z = SceneDepthTexture.SampleLevel(SceneDepthTextureSampler, UVAndScreenPos.xy, 0, int2(-AA_CROSS,  AA_CROSS)).r;
+		Depths.w = SceneDepthTexture.SampleLevel(SceneDepthTextureSampler, UVAndScreenPos.xy, 0, int2( AA_CROSS,  AA_CROSS)).r;
+
 		float2 DepthOffset = float2(AA_CROSS, AA_CROSS);
 		float DepthOffsetXx = float(AA_CROSS);
 #if HAS_INVERTED_Z_BUFFER
-		//@todo-briank
 		// Nearest depth is the largest depth (depth surface 0=far, 1=near).
 		if(Depths.x > Depths.y) 
 		{
@@ -179,34 +153,19 @@ which can be configured via defines for various temporal AA passes.
 			// With the assumption that UV position scales between velocity and color.
 			VelocityOffset = DepthOffset * PostprocessInput0Size.zw;
 			// This is [0 to 1] flipped in Y.
-			PosN.xy = (UVAndScreenPos.zw + DepthOffset * ViewportSize.zw * 2.0) * float2(0.5, -0.5) + 0.5;
+			//PosN.xy = UVAndScreenPos.zw + DepthOffset * ViewportSize.zw * 2.0;
 			PosN.z = DepthsXYZW;
 		}
 #else
 #error Fix me!
 #endif // HAS_INVERTED_Z_BUFFER
 	#endif	// AA_CROSS
-	// Camera motion for pixel or nearest pixel (in ScreenPos space).
-	float ScaleM = 1.0 / (dot(PosN, CameraMotion.Value[0].xyz) + CameraMotion.Value[0].w);
-	float2 BackN;
-	BackN.x = -2.0 * ((PosN.x * ((CameraMotion.Value[1].x * PosN.y) + (CameraMotion.Value[1].y * PosN.z) + CameraMotion.Value[1].z)) + (CameraMotion.Value[1].w * PosN.y) + (CameraMotion.Value[2].x * PosN.x * PosN.x) + (CameraMotion.Value[2].y * PosN.z) + CameraMotion.Value[2].z) * ScaleM;
-	BackN.y =  2.0 * ((PosN.y * ((CameraMotion.Value[3].x * PosN.x) + (CameraMotion.Value[3].y * PosN.z) + CameraMotion.Value[3].z)) + (CameraMotion.Value[3].w * PosN.x) + (CameraMotion.Value[4].x * PosN.y * PosN.y) + (CameraMotion.Value[4].y * PosN.z) + CameraMotion.Value[4].z) * ScaleM;
-
-#if 0
-	//Fix for VR 'wobble' Imprecision in this method makes it not good enough to have on by default.  A proper fix is coming soon, but in the meantime, VR users can put this in locally.
-	{
-		float4 ThisClip = float4( (UVAndScreenPos.xy - View.ScreenPositionScaleBias.wz) / View.ScreenPositionScaleBias.xy, PosN.z, 1 );
-		float4 ThisTranslatedWorld = mul( ThisClip, View.ClipToTranslatedWorld );
-		ThisTranslatedWorld /= ThisTranslatedWorld.w;
-		
-		float3 PrevTranslatedWorld = ThisTranslatedWorld.xyz + (View.PrevPreViewTranslation - View.PreViewTranslation);
-		float4 PrevClip = mul( float4( PrevTranslatedWorld, 1 ), View.PrevTranslatedWorldToClip );
-		float2 PrevScreen = PrevClip.xy / PrevClip.w;
-		//float2 PrevUV = PrevScreen.xy * View.ScreenPositionScaleBias.xy + View.ScreenPositionScaleBias.wz;
 
-		BackN = UVAndScreenPos.zw - PrevScreen - ( View.ViewToClip[2].xy - View.PrevProjection[2].xy );
-	}
-#endif
+	// Camera motion for pixel or nearest pixel (in ScreenPos space).
+	float4 ThisClip = float4( PosN.xy, PosN.z, 1 );
+	float4 PrevClip = mul( ThisClip, View.ClipToPrevClip );
+	float2 PrevScreen = PrevClip.xy / PrevClip.w;
+	float2 BackN = PosN.xy - PrevScreen;
 
 	float2 BackTemp = BackN * ViewportSize.xy;
 	#if AA_DYNAMIC
@@ -223,10 +182,12 @@ which can be configured via defines for various temporal AA passes.
 		}
 		BackTemp = BackN * ViewportSize.xy;
 	#endif
-	// Save the amount of pixel offset of just camera motion, used later as the amount of blur introduced by history.
-	float HistoryBlurAmp = 2.0;
-	float HistoryBlur = saturate(abs(BackTemp.x) * HistoryBlurAmp + abs(BackTemp.y) * HistoryBlurAmp);
-	float Velocity = sqrt(dot(BackTemp, BackTemp));
+	#if !AA_BICUBIC
+		// Save the amount of pixel offset of just camera motion, used later as the amount of blur introduced by history.
+		float HistoryBlurAmp = 2.0;
+		float HistoryBlur = saturate(abs(BackTemp.x) * HistoryBlurAmp + abs(BackTemp.y) * HistoryBlurAmp);
+		float Velocity = sqrt(dot(BackTemp, BackTemp));
+	#endif
 	// Easier to do off screen check before conversion.
 	// BackN is in units of 2pixels/viewportWidthInPixels
 	// This converts back projection vector to [-1 to 1] offset in viewport.
@@ -256,19 +217,12 @@ which can be configured via defines for various temporal AA passes.
 		float4 Neighbor4 = PostprocessInput0.SampleLevel(PostprocessInput0Sampler, UVAndScreenPos.xy, 0);
 		float4 Neighbor5 = PostprocessInput0.SampleLevel(PostprocessInput0Sampler, UVAndScreenPos.xy, 0, int2( 1,  0));
 		float4 Neighbor7 = PostprocessInput0.SampleLevel(PostprocessInput0Sampler, UVAndScreenPos.xy, 0, int2( 0,  1));
-		#if 1
-			Neighbor1.rgb *= HdrWeightG(Neighbor1.rgb, InExposureScale);
-			Neighbor3.rgb *= HdrWeightG(Neighbor3.rgb, InExposureScale);
-			Neighbor4.rgb *= HdrWeightG(Neighbor4.rgb, InExposureScale);
-			Neighbor5.rgb *= HdrWeightG(Neighbor5.rgb, InExposureScale);
-			Neighbor7.rgb *= HdrWeightG(Neighbor7.rgb, InExposureScale);
-		#endif
 		Neighbor1.rgb = RGBToYCoCg(Neighbor1.rgb);
 		Neighbor3.rgb = RGBToYCoCg(Neighbor3.rgb);
 		Neighbor4.rgb = RGBToYCoCg(Neighbor4.rgb);
 		Neighbor5.rgb = RGBToYCoCg(Neighbor5.rgb);
 		Neighbor7.rgb = RGBToYCoCg(Neighbor7.rgb);
-		#if 0
+		#if 1
 			Neighbor1.xyz *= HdrWeightY(Neighbor1.x, InExposureScale);
 			Neighbor3.xyz *= HdrWeightY(Neighbor3.x, InExposureScale);
 			Neighbor4.xyz *= HdrWeightY(Neighbor4.x, InExposureScale);
@@ -393,11 +347,28 @@ which can be configured via defines for various temporal AA passes.
 		#endif
 	#endif
 
-
-
 	// FETCH HISTORY
 	// -------------
-	OutColor = PostprocessInput1.SampleLevel(PostprocessInput1Sampler, BackN.xy, 0);
+	#if AA_BICUBIC
+		float2 Weight[3];
+		float2 Sample[3];
+		Bicubic2DCatmullRom( BackN.xy, PostprocessInput1Size.xy, Sample, Weight );
+
+		OutColor  = PostprocessInput1.SampleLevel( PostprocessInput1Sampler, float2( Sample[0].x, Sample[0].y ), 0 ) * Weight[0].x * Weight[0].y;
+		OutColor += PostprocessInput1.SampleLevel( PostprocessInput1Sampler, float2( Sample[1].x, Sample[0].y ), 0 ) * Weight[1].x * Weight[0].y;
+		OutColor += PostprocessInput1.SampleLevel( PostprocessInput1Sampler, float2( Sample[2].x, Sample[0].y ), 0 ) * Weight[2].x * Weight[0].y;
+
+		OutColor += PostprocessInput1.SampleLevel( PostprocessInput1Sampler, float2( Sample[0].x, Sample[1].y ), 0 ) * Weight[0].x * Weight[1].y;
+		OutColor += PostprocessInput1.SampleLevel( PostprocessInput1Sampler, float2( Sample[1].x, Sample[1].y ), 0 ) * Weight[1].x * Weight[1].y;
+		OutColor += PostprocessInput1.SampleLevel( PostprocessInput1Sampler, float2( Sample[2].x, Sample[1].y ), 0 ) * Weight[2].x * Weight[1].y;
+
+		OutColor += PostprocessInput1.SampleLevel( PostprocessInput1Sampler, float2( Sample[0].x, Sample[2].y ), 0 ) * Weight[0].x * Weight[2].y;
+		OutColor += PostprocessInput1.SampleLevel( PostprocessInput1Sampler, float2( Sample[1].x, Sample[2].y ), 0 ) * Weight[1].x * Weight[2].y;
+		OutColor += PostprocessInput1.SampleLevel( PostprocessInput1Sampler, float2( Sample[2].x, Sample[2].y ), 0 ) * Weight[2].x * Weight[2].y;
+	#else
+		OutColor = PostprocessInput1.SampleLevel(PostprocessInput1Sampler, BackN.xy, 0);
+	#endif
+
 	#if AA_DEBUG
 		Neighbor4.rg = float2(0.0, 0.0);
 		NeighborMin.rg = float2(0.0, 0.0);
@@ -407,11 +378,8 @@ which can be configured via defines for various temporal AA passes.
 		float DebugDiffCurrent = Filtered.b;
 	#endif
 	#if AA_YCOCG
-		#if 1
-			OutColor.rgb *= HdrWeightG(OutColor.rgb, InExposureScale);
-		#endif
 		OutColor.rgb = RGBToYCoCg(OutColor.rgb);
-		#if 0
+		#if 1
 			OutColor.xyz *= HdrWeightY(OutColor.x, InExposureScale);
 		#endif
 	#else
@@ -427,10 +395,6 @@ which can be configured via defines for various temporal AA passes.
 	#endif
 
 
-
-
-
-
 	// FIND LUMA OF CLAMPED HISTORY
 	// ----------------------------
 	// Save off luma of history before the clamp.
@@ -475,7 +439,7 @@ which can be configured via defines for various temporal AA passes.
 
 	// ADD BACK IN ALIASING TO SHARPEN
     // -------------------------------
-	#if AA_FILTERED
+	#if AA_FILTERED && !AA_BICUBIC
 		// Blend in non-filtered based on the amount of sub-pixel motion.
 		float AddAliasing = saturate(HistoryBlur) * 0.5;
 		float LumaContrastFactor = 32.0;
@@ -498,26 +462,20 @@ which can be configured via defines for various temporal AA passes.
 		
 	// COMPUTE BLEND AMOUNT 
     // --------------------
-	// Replace history with minimum difference of history from min or max neighborhood.
-	LumaHistory = min(abs(LumaMin-LumaHistory), abs(LumaMax-LumaHistory));
-	float HistoryAmount = (1.0/8.0) + HistoryBlur * (1.0/8.0);
-	float HistoryFactor = LumaHistory * HistoryAmount * (1.0 + HistoryBlur * HistoryAmount * 8.0);
-	float BlendFinal = saturate(HistoryFactor * rcp(LumaHistory + LumaContrast));
+	float DistToClamp = min(abs(LumaMin-LumaHistory), abs(LumaMax-LumaHistory));
+	#if AA_BICUBIC
+		float HistoryFactor = 0.125 * DistToClamp;
+	#else
+		float HistoryAmount = (1.0/8.0) + HistoryBlur * (1.0/8.0);
+		float HistoryFactor = DistToClamp * HistoryAmount * (1.0 + HistoryBlur * HistoryAmount * 8.0);
+	#endif
+	float BlendFinal = saturate(HistoryFactor / (DistToClamp + LumaMax - LumaMin));
 	#if RESPONSIVE
 		// Responsive forces 1/4 of new frame.
 		BlendFinal = 1.0/4.0;
 	#endif
 
-   	#if (AA_ALPHA == 1) && (AA_VELOCITY_WEIGHTING == 1)
-		// Velocity weighting.
-		// Decay prior value.
-		float PriorVelocity = OutColor.a;
-		float VelocityDecay = 0.5;
-		OutColor.a = max(OutColor.a * VelocityDecay, Velocity * rcp(VelocityDecay));
-		float VelocityDiff = abs(PriorVelocity - Velocity) / max(1.0, max(PriorVelocity, Velocity));
-		BlendFinal = max(BlendFinal, VelocityDiff * (1.0/8.0));
-		BlendFinal = min(1.0/2.0, BlendFinal);
-	#endif
+	BlendFinal = 0.04;
 
 	// Offscreen feedback resets.
 	#if AA_LERP 
@@ -553,13 +511,10 @@ which can be configured via defines for various temporal AA passes.
 		OutColor = lerp(OutColor, Filtered, FixedLerp);
 	#endif
 	#if AA_YCOCG
-		#if 0
+		#if 1
 			OutColor.xyz *= HdrWeightInvY(OutColor.x, InExposureScale);
 		#endif
 		OutColor.rgb = YCoCgToRGB(OutColor.rgb);
-		#if 1
-			OutColor.rgb *= HdrWeightInvG(OutColor.rgb, InExposureScale);
-		#endif
 	#else
 		// Convert back into linear.
 		#if AA_GREEN_AS_LUMA
@@ -586,9 +541,6 @@ which can be configured via defines for various temporal AA passes.
 #undef AA_DILATE
 #undef AA_DYNAMIC
 #undef AA_FILTERED
-#undef AA_GRAIN
-#undef AA_HDR
-#undef AA_HDR_HISTORY
 #undef AA_LERP
 #undef AA_LOWPASS
 #undef AA_ROUND
@@ -596,7 +548,6 @@ which can be configured via defines for various temporal AA passes.
 #undef AA_BORDER
 #undef AA_FORCE_ALPHA_CLAMP
 #undef AA_GREEN_AS_LUMA
-#undef AA_VELOCITY_WEIGHTING
 #undef AA_YCOCG
-#undef AA_DOF
-#undef AA_TEST
\ No newline at end of file
+#undef AA_BICUBIC
+#undef AA_DOF
\ No newline at end of file
diff --git a/Engine/Source/Runtime/Engine/Private/SceneView.cpp b/Engine/Source/Runtime/Engine/Private/SceneView.cpp
--- a/Engine/Source/Runtime/Engine/Private/SceneView.cpp
+++ b/Engine/Source/Runtime/Engine/Private/SceneView.cpp
@@ -311,9 +311,11 @@ FSceneView::FSceneView(const FSceneViewInitOptions& InitOptions)
 	// Compute the view projection matrix and its inverse.
 	ViewProjectionMatrix = ViewMatrices.GetViewProjMatrix();
 
+	FMatrix InvProjectionMatrix = ViewMatrices.GetInvProjMatrix();
+
 	// For precision reasons the view matrix inverse is calculated independently.
-	InvViewMatrix = ViewMatrices.ViewMatrix.Inverse();
-	InvViewProjectionMatrix = ViewMatrices.GetInvProjMatrix() * InvViewMatrix;
+	InvViewMatrix = InitOptions.ViewRotationMatrix.GetTransposed() * FTranslationMatrix(InitOptions.ViewOrigin);
+	InvViewProjectionMatrix = InvProjectionMatrix * InvViewMatrix;
 
 	bool bApplyPreViewTranslation = true;
 	bool bViewOriginIsFudged = false;
@@ -340,6 +342,7 @@ FSceneView::FSceneView(const FSceneViewInitOptions& InitOptions)
 
 	/** The view transform, starting from world-space points translated by -ViewOrigin. */
 	FMatrix TranslatedViewMatrix = InitOptions.ViewRotationMatrix;
+	FMatrix InvTranslatedViewMatrix = TranslatedViewMatrix.GetTransposed();
 
 	// Translate world-space so its origin is at ViewOrigin for improved precision.
 	// Note that this isn't exactly right for orthogonal projections (See the above special case), but we still use ViewOrigin
@@ -371,6 +374,7 @@ FSceneView::FSceneView(const FSceneViewInitOptions& InitOptions)
 	{
 		// If not applying PreViewTranslation then we need to use the view matrix directly.
 		TranslatedViewMatrix = ViewMatrices.ViewMatrix;
+		InvTranslatedViewMatrix = InvViewMatrix;
 	}
 
 	// When the view origin is fudged for faux ortho view position the translations don't cancel out.
@@ -383,7 +387,7 @@ FSceneView::FSceneView(const FSceneViewInitOptions& InitOptions)
 	// Compute a transform from view origin centered world-space to clip space.
 	ViewMatrices.TranslatedViewMatrix = TranslatedViewMatrix;
 	ViewMatrices.TranslatedViewProjectionMatrix = TranslatedViewMatrix * ViewMatrices.ProjMatrix;
-	ViewMatrices.InvTranslatedViewProjectionMatrix = ViewMatrices.TranslatedViewProjectionMatrix.Inverse();
+	ViewMatrices.InvTranslatedViewProjectionMatrix = InvProjectionMatrix * InvTranslatedViewMatrix;
 
 	// Compute screen scale factors.
 	// Stereo renders at half horizontal resolution, but compute shadow resolution based on full resolution.
diff --git a/Engine/Source/Runtime/Engine/Public/SceneView.h b/Engine/Source/Runtime/Engine/Public/SceneView.h
--- a/Engine/Source/Runtime/Engine/Public/SceneView.h
+++ b/Engine/Source/Runtime/Engine/Public/SceneView.h
@@ -231,18 +231,17 @@ struct FViewMatrices
 
 	FMatrix GetInvProjMatrix() const
 	{
-		return ProjMatrix.Inverse();
+		return InvertProjMatrix( ProjMatrix );
 	}
 	
 	FMatrix GetInvProjNoAAMatrix() const
 	{
-		return GetProjNoAAMatrix().Inverse();
+		return InvertProjMatrix( GetProjNoAAMatrix() );
 	}
 
 	FMatrix GetInvViewMatrix() const
 	{
-		// can be optimized: it's not a perspective matrix so transpose would be enough
-		return ViewMatrix.Inverse();
+		return ViewMatrix.RemoveTranslation().GetTransposed() * FTranslationMatrix( ViewMatrix.GetOrigin() );
 	}
 
 	FMatrix GetInvViewProjMatrix() const
@@ -265,6 +264,49 @@ struct FViewMatrices
 
 		return FVector2D(FMath::Acos(VCenter | VRight), FMath::Acos(VCenter | VUp));
 	}
+
+private:
+	static FMatrix InvertProjMatrix( const FMatrix& M )
+	{
+		if( M.M[1][0] == 0.0f &&
+			M.M[3][0] == 0.0f &&
+			M.M[0][1] == 0.0f &&
+			M.M[3][1] == 0.0f &&
+			M.M[0][2] == 0.0f &&
+			M.M[1][2] == 0.0f &&
+			M.M[0][3] == 0.0f &&
+			M.M[1][3] == 0.0f &&
+			M.M[2][3] == 1.0f &&
+			M.M[3][3] == 0.0f )
+		{
+			// Solve the common case directly with very high precision.
+			/*
+			M = 
+			| a | 0 | 0 | 0 |
+			| 0 | b | 0 | 0 |
+			| s | t | c | 1 |
+			| 0 | 0 | d | 0 |
+			*/
+
+			double a = M.M[0][0];
+			double b = M.M[1][1];
+			double c = M.M[2][2];
+			double d = M.M[3][2];
+			double s = M.M[2][0];
+			double t = M.M[2][1];
+
+			return FMatrix(
+				FPlane( 1.0 / a, 0.0f, 0.0f, 0.0f ),
+				FPlane( 0.0f, 1.0 / b, 0.0f, 0.0f ),
+				FPlane( 0.0f, 0.0f, 0.0f, 1.0 / d ),
+				FPlane( -s/a, -t/b, 1.0f, -c/d )
+			);
+		}
+		else
+		{
+			return M.Inverse();
+		}
+	}
 };
 
 #include "ShaderParameters.h"
@@ -383,6 +425,7 @@ BEGIN_UNIFORM_BUFFER_STRUCT_WITH_CONSTRUCTOR(FViewUniformShaderParameters,ENGINE
 	DECLARE_UNIFORM_BUFFER_STRUCT_MEMBER(FVector,PrevPreViewTranslation)
 	DECLARE_UNIFORM_BUFFER_STRUCT_MEMBER(FMatrix,PrevInvViewProj)
 	DECLARE_UNIFORM_BUFFER_STRUCT_MEMBER(FMatrix,PrevScreenToTranslatedWorld)
+	DECLARE_UNIFORM_BUFFER_STRUCT_MEMBER(FMatrix,ClipToPrevClip)
 	DECLARE_UNIFORM_BUFFER_STRUCT_MEMBER(FVector,IndirectLightingColorScale)
 	DECLARE_UNIFORM_BUFFER_STRUCT_MEMBER_EX(float,HdrMosaic, EShaderPrecisionModifier::Half)
 	DECLARE_UNIFORM_BUFFER_STRUCT_MEMBER(FVector,AtmosphericFogSunDirection)
diff --git a/Engine/Source/Runtime/Renderer/Private/SceneRendering.cpp b/Engine/Source/Runtime/Renderer/Private/SceneRendering.cpp
--- a/Engine/Source/Runtime/Renderer/Private/SceneRendering.cpp
+++ b/Engine/Source/Runtime/Renderer/Private/SceneRendering.cpp
@@ -523,6 +523,12 @@ TUniformBufferRef<FViewUniformShaderParameters> FViewInfo::CreateUniformBuffer(
 		FPlane(0,0,ProjectionMatrixUnadjustedForRHI.M[3][2],0))
 		* PrevViewMatrices.InvTranslatedViewProjectionMatrix;
 
+	FVector DeltaTranslation = PrevViewMatrices.PreViewTranslation - ViewMatrices.PreViewTranslation;
+	FMatrix InvViewProj = ViewMatrices.GetInvProjNoAAMatrix() * ViewMatrices.TranslatedViewMatrix.GetTransposed();
+	FMatrix PrevViewProj = FTranslationMatrix( DeltaTranslation ) * PrevViewMatrices.TranslatedViewMatrix * PrevViewMatrices.GetProjNoAAMatrix();
+
+	ViewUniformShaderParameters.ClipToPrevClip = InvViewProj * PrevViewProj;
+
 	// is getting clamped in the shader to a value larger than 0 (we don't want the triangles to disappear)
 	ViewUniformShaderParameters.AdaptiveTessellationFactor = 0.0f;
 
