Fixed shader compile error in PostProcessVisualizeDOF.usf

diff --git a/Engine/Shaders/PostProcessVisualizeDOF.usf b/Engine/Shaders/PostProcessVisualizeDOF.usf
--- a/Engine/Shaders/PostProcessVisualizeDOF.usf
+++ b/Engine/Shaders/PostProcessVisualizeDOF.usf
@@ -79,13 +79,13 @@ void VisualizeDOFPS(in float4 UVAndScreenPos : TEXCOORD0, out float4 OutColor :
 
 		{
 			int2 LeftTop = PrintCursor - int2(11, 0);
-			PrintCharacter(PixelPos, OutColor.xyz, float3(1, 1, 1), LeftTop, 'D' - 'A' + 10);
+			PrintCharacter(PixelPos, OutColor.xyz, float3(1, 1, 1), LeftTop, 13);	// 'D' - 'A' + 10);
 		}
 		PrintFloatNoFraction(PixelPos, OutColor.xyz, float3(1, 1, 1)*0.6, PrintCursor, SceneDepthAtCursor);
 		PrintCursor.y += 16;
 		{
 			int2 LeftTop = PrintCursor - int2(11, 0);
-			PrintCharacter(PixelPos, OutColor.xyz, float3(1, 1, 1), LeftTop, 'R' - 'A' + 10);
+			PrintCharacter(PixelPos, OutColor.xyz, float3(1, 1, 1), LeftTop, 27);	// 'R' - 'A' + 10);
 		}
 		PrintFloat(PixelPos, OutColor.xyz, float3(1, 1, 1)*0.6, PrintCursor, CoCInPixel);
 	}
