Fixed d3d debug error from MainTemporalAAPS inputs (specifically SV_Position) not matching the VS output order

diff --git a/Engine/Shaders/PostProcessTemporalAA.usf b/Engine/Shaders/PostProcessTemporalAA.usf
--- a/Engine/Shaders/PostProcessTemporalAA.usf
+++ b/Engine/Shaders/PostProcessTemporalAA.usf
@@ -299,6 +299,9 @@ void LightShaftTemporalAAPS( float4 UVAndScreenPos : TEXCOORD0, float3 InExposur
 void MainTemporalAAPS(
 	in float4 UVAndScreenPos : TEXCOORD0,
 	in float3 InExposureScaleVignette : TEXCOORD1,
+	// Note: these are needed to match up with the vertex shader outputs which is FPostProcessTonemapVS, otherwise a d3d debug error is issued
+	in float4 Unused1 : TEXCOORD2,
+	in float4 Unused2 : TEXCOORD3,
 	in float4 SvPosition : SV_Position,
 	out float4 OutColor : SV_Target0 )
 {
