fixed random quantization bias in base pixel shader

diff --git a/Engine/Shaders/BasePassPixelShader.usf b/Engine/Shaders/BasePassPixelShader.usf
--- a/Engine/Shaders/BasePassPixelShader.usf
+++ b/Engine/Shaders/BasePassPixelShader.usf
@@ -905,7 +905,8 @@ void Main(
 			float4 OutGBufferVelocity = 0;
 		#endif
 
-		float QuantizationBias = PseudoRandom( MaterialParameters.SVPosition.xy ) * 0.5 - 0.5;
+		// -0.5 .. 0.5, could be optimzed as lower quality noise would be sufficient
+		float QuantizationBias = PseudoRandom( MaterialParameters.SVPosition.xy ) - 0.5f;
 		EncodeGBuffer(GBuffer, OutGBufferA, OutGBufferB, OutGBufferC, OutGBufferD, OutGBufferE, OutGBufferVelocity, QuantizationBias);
 	#endif 
 
