Remove harmless leftover.

diff --git a/Engine/Shaders/MediaShaders.usf b/Engine/Shaders/MediaShaders.usf
--- a/Engine/Shaders/MediaShaders.usf
+++ b/Engine/Shaders/MediaShaders.usf
@@ -35,7 +35,6 @@ void YCbCrConvertPS(
 	float2 InUV : TEXCOORD0,
 	out float4 OutColor : SV_Target0)
 {
-	OutColor = float4(0, 0, 0, 0);
 	float3 YCbCr;
 	YCbCr.x = YCbCrConvertUB.LumaTexture.Sample(YCbCrConvertUB.LumaSampler, InUV.xy).x - 0.0625; //e.g. 16 / 255
 	YCbCr.yz = YCbCrConvertUB.CbCrTexture.Sample(YCbCrConvertUB.CbCrSampler, InUV.xy).xy - float2(0.5, 0.5);
