Bug UE-19018: Clobbers ShaderVersion.usf to avoid state tracking error between the previously materials' generated HLSL code and the material graph itself, causing a survival of incorrect generated HLSL code by UE-18715.

diff --git a/Engine/Shaders/ShaderVersion.usf b/Engine/Shaders/ShaderVersion.usf
--- a/Engine/Shaders/ShaderVersion.usf
+++ b/Engine/Shaders/ShaderVersion.usf
@@ -2,4 +2,4 @@
 // Each time the console command is executed it generates a new GUID. As this file is included
 // in common.usf (which should be included in any shader) it allows to invalidate the shader DDC.
 // 
-// GUID = 10DDEA3C48E7FD138E98039F9438498D
\ No newline at end of file
+// GUID = 614697354776EE6EA312B887F2FFD5EF
\ No newline at end of file
