added comments

diff --git a/Engine/Shaders/PostProcessSubsurface.usf b/Engine/Shaders/PostProcessSubsurface.usf
--- a/Engine/Shaders/PostProcessSubsurface.usf
+++ b/Engine/Shaders/PostProcessSubsurface.usf
@@ -80,11 +80,14 @@ float GetSubsurfaceStrength(float2 UV)
 // 1: no color bleeding in z direction
 #define SSSS_FOLLOW_SURFACE 1
 
+// @param In needs to be in 0..1 range 
+// @return 0..1
 float4 Quantize8Bit(float4 In)
 {
 	return floor(In * 255.999f) / 255.0f;
 }
 
+// @return .rgb is the weight for color channel, .a is the sample location
 float4 GetKernel(uint SampleIndex, uint SubsurfaceProfileInt)
 {
 	const float4 TableMax = float4(1, 1, 1, 3.0f);
