Improved ScreenSpaceSubsurface quality (minor colorshift on object borders, more visible with large radius)

diff --git a/Engine/Shaders/SeparableSSS.usf b/Engine/Shaders/SeparableSSS.usf
--- a/Engine/Shaders/SeparableSSS.usf
+++ b/Engine/Shaders/SeparableSSS.usf
@@ -192,7 +192,7 @@ float4 SSSSBlurPS(float2 GBufferUV, float2 ViewportUV, float2 dir, bool initSten
     float4 colorM = SSSSSampleSceneColorPoint(ViewportUV);
 
 	// we store the depth in alpha
-	float depthM = colorM.a;
+	float OutDepth = colorM.a;
 
 	colorM.a = ComputeMaskFromDepthInAlpha(colorM.a);
 
@@ -212,7 +212,7 @@ float4 SSSSBlurPS(float2 GBufferUV, float2 ViewportUV, float2 dir, bool initSten
 
 	float SSSScaleX = SSSParams.x;
 
-	float scale = SSSScaleX / depthM;
+	float scale = SSSScaleX / OutDepth;
 
     // Calculate the final step to fetch the surrounding pixels:
     float2 finalStep = scale * dir;
@@ -227,14 +227,16 @@ float4 SSSSBlurPS(float2 GBufferUV, float2 ViewportUV, float2 dir, bool initSten
 	uint SubsurfaceProfileInt = ExtractSubsurfaceProfileInt(GBufferData);
 
     // Accumulate the center sample:
-    float4 colorBlurred = 0;
+    float3 colorAccum = 0;
+	// >0 to avoid division by 0, not 100% correct to not visible
+	float3 colorInvDiv = 0.00001f;
 
 	// center sample is added in Recombine pass
-	colorBlurred.a = colorM.a;
+	colorInvDiv += GetKernel(SSSS_N_KERNELWEIGHTOFFSET, SubsurfaceProfileInt).rgb;
 
 	if(bIncludeCenter)
 	{
-		colorBlurred.rgb = colorM.rgb * GetKernel(SSSS_N_KERNELWEIGHTOFFSET, SubsurfaceProfileInt).rgb;
+		colorAccum = colorM.rgb * GetKernel(SSSS_N_KERNELWEIGHTOFFSET, SubsurfaceProfileInt).rgb;
 	}
 
 	// Accumulate the other samples:
@@ -260,35 +262,33 @@ float4 SSSSBlurPS(float2 GBufferUV, float2 ViewportUV, float2 dir, bool initSten
 			
 			color.a = ComputeMaskFromDepthInAlpha(color.a);
 
-			// approximation, ideally we would reconstruct the mask with ComputeMaskFromDepthInAlpha() and do manual bilinear filter
-			// needed?
-//			color.rgb *= color.a;
-
-
 #if SSSS_FOLLOW_SURFACE == 1
 			// If the difference in depth is huge, we weight the sample less or not at all
 			float s = saturate(12000.0f / 400000 * SSSParams.y *
 	//        float s = saturate(300.0f/400000 * SSSParams.y *
-				abs(depthM - LocalDepth));
+				abs(OutDepth - LocalDepth));
 
-			color *= 1 - s;
+			color.a *= 1 - s;
 #endif
+			// approximation, ideally we would reconstruct the mask with ComputeMaskFromDepthInAlpha() and do manual bilinear filter
+			// needed?
+			color.rgb *= color.a;
 
 			// Accumulate left and right 
 			LocalAccum += color;
 		}
 
 		// Accumulate to final value (left and right sample with the same weight)
-		colorBlurred += Kernel * LocalAccum;
+		colorAccum += Kernel.rgb * LocalAccum.rgb;
+		colorInvDiv += Kernel.rgb * LocalAccum.a;
 	}
 
 	// normalize (some samples are rejected because of depth or the other material is no SSS, compensate for that)
-	colorBlurred.rgb *= 1.0f / max(colorBlurred.a, 0.00001f);
+	// done for each color channel to avoid color shift
+	float3 OutColor = colorAccum / colorInvDiv; 
 
 	// alpha stored the SceneDepth (0 if there is no subsurface scattering)
-	colorBlurred.a = depthM;
-
-    return colorBlurred;
+    return float4(OutColor, OutDepth);
 }
 
 //-----------------------------------------------------------------------------
\ No newline at end of file
