PR #1239 - Contributed by [REDACTED_PERSON]
UE-16895 - UE4 - Fix cloth motion blur

diff --git a/Engine/Shaders/GpuSkinVertexFactory.usf b/Engine/Shaders/GpuSkinVertexFactory.usf
--- a/Engine/Shaders/GpuSkinVertexFactory.usf
+++ b/Engine/Shaders/GpuSkinVertexFactory.usf
@@ -583,7 +583,8 @@ float4 VertexFactoryGetPreviousWorldPosition(FVertexFactoryInput Input, FVertexF
 #if GPUSKIN_APEX_CLOTH
 	if( IsSimulatedVertex(Input) )
 	{
-		float3 BlendedPrevPos = lerp(PrevSkinPosInWorld.xyz, Intermediates.SimulatedPosition + View.PrevPreViewTranslation, ClothBlendWeight);		 
+		float3 PrevSimulatedPos = mul(mul(float4(Intermediates.SimulatedPosition.xyz, 1), Primitive.WorldToLocal), PreviousLocalToWorld).xyz;
+		float3 BlendedPrevPos = lerp(PrevSkinPosInWorld.xyz, PrevSimulatedPos.xyz, ClothBlendWeight);
 		return float4(BlendedPrevPos, 1);
 	}
 #endif // GPUSKIN_APEX_CLOTH
