Fixed clear coat bugs.
Fixed clear coat when sky but no reflection captures are present.

Fixed clear coat values were between 0 and 1 going gray.

diff --git a/Engine/Shaders/BasePassPixelShader.usf b/Engine/Shaders/BasePassPixelShader.usf
--- a/Engine/Shaders/BasePassPixelShader.usf
+++ b/Engine/Shaders/BasePassPixelShader.usf
@@ -855,24 +855,18 @@ void Main(
 
 			float ClearCoat				= GetMaterialClearCoat(MaterialParameters);
 			float ClearCoatRoughness	= GetMaterialClearCoatRoughness(MaterialParameters);
-			float Film = 1 * ClearCoat;
 			float MetalSpec = 0.9;
 
 			float NoV = saturate( dot( MaterialParameters.WorldNormal, MaterialParameters.CameraVector ) );
 
 			// Approximation of refraction's effect on EnvBRDF
 			float RefractionScale = ( (NoV * 0.5 + 0.5) * NoV - 1 ) * saturate( 1.25 - 1.25 * Roughness ) + 1;
-			RefractionScale = lerp( 1, RefractionScale, ClearCoat );
-
-			float3 AbsorptionColor = (1 - Film) + BaseColor * ( Film * (1 / MetalSpec) );
-			GBuffer.BaseColor = lerp( BaseColor, MetalSpec, Film );
-
+		
 			// Approximation of absorption integral, tuned for Roughness=0.4
+			float3 AbsorptionColor = BaseColor * (1 / MetalSpec);
 			float3 Absorption = AbsorptionColor * ( (NoV - 1) * 0.85 * ( 1 - lerp( AbsorptionColor, Square(AbsorptionColor), -0.78 ) ) + 1 );
-			//float Toe = 1 - exp2( -9 * NoV - 3.5 );
-			//float3 Absorption = Toe * AbsorptionColor * ( (NoV - 1) * 0.85 * ( 1 - lerp( AbsorptionColor, Square(AbsorptionColor), -0.78 ) ) + lerp( 0.93, 0.96, AbsorptionColor * 2 - 1 ) );
 
-			GBuffer.BaseColor *= Absorption * lerp( 1, RefractionScale, Metallic );
+			GBuffer.BaseColor = lerp( BaseColor, MetalSpec * Absorption * lerp( 1, RefractionScale, Metallic ), ClearCoat );
 			GBuffer.Specular *= Luminance( Absorption ) * RefractionScale;
 
 			GBuffer.CustomData.x = ClearCoat;
diff --git a/Engine/Shaders/ReflectionEnvironmentComputeShaders.usf b/Engine/Shaders/ReflectionEnvironmentComputeShaders.usf
--- a/Engine/Shaders/ReflectionEnvironmentComputeShaders.usf
+++ b/Engine/Shaders/ReflectionEnvironmentComputeShaders.usf
@@ -673,7 +673,7 @@ void ReflectionEnvironmentTiledDeferredMain(
 		
 			// F_Schlick
 			float F0 = 0.04;
-			float Fc = pow( 1 - NoV, 5 );
+			float Fc = Pow5( 1 - NoV );
 			float F = Fc + (1 - Fc) * F0;
 			F *= ClearCoat;
 			
@@ -682,10 +682,13 @@ void ReflectionEnvironmentTiledDeferredMain(
 			Color.rgb *= LayerAttenuation;
 			Color.a = F;
 		
-			Color.rgb += SSR.rgb * F;
-			Color.a *= 1 - SSR.a;
+			#if 1//APPLY_SSR
+				Color.rgb += SSR.rgb * F;
+				Color.a *= 1 - SSR.a;
+			#endif
 			Color.a *= SpecularOcclusion;
 
+			R = 2 * dot( V, N ) * N - V;
 			GatherRadiance(Color, WorldPosition, R, ClearCoatRoughness, ScreenPosition, IndirectIrradiance, GBuffer.ShadingModelID);			
 		}
 		else
diff --git a/Engine/Shaders/ReflectionEnvironmentShaders.usf b/Engine/Shaders/ReflectionEnvironmentShaders.usf
--- a/Engine/Shaders/ReflectionEnvironmentShaders.usf
+++ b/Engine/Shaders/ReflectionEnvironmentShaders.usf
@@ -489,6 +489,49 @@ void ReflectionApplyPS(
 #endif
 #endif
 
+#if APPLY_SSR
+		float4 SSR = Texture2DSample( ScreenSpaceReflectionsTexture, ScreenSpaceReflectionsSampler, UV );
+#endif
+
+		BRANCH
+		if( GBuffer.ShadingModelID == SHADINGMODELID_CLEAR_COAT )
+		{
+			const float ClearCoat			= GBuffer.CustomData.x;
+			const float ClearCoatRoughness	= GBuffer.CustomData.y;
+
+			// TODO EnvBRDF should have a mask param
+			float2 AB = PreIntegratedGF.SampleLevel( PreIntegratedGFSampler, float2( NoV, GBuffer.Roughness ), 0 ).rg;
+			SpecularColor = GBuffer.SpecularColor * AB.x + AB.y * saturate( 50 * GBuffer.SpecularColor.g ) * (1 - ClearCoat);
+
+			// F_Schlick
+			float F0 = 0.04;
+			float Fc = Pow5( 1 - NoV );
+			float F = Fc + (1 - Fc) * F0;
+			F *= ClearCoat;
+
+			float LayerAttenuation = (1 - F);
+			SpecularColor *= LayerAttenuation;
+			SpecularColor *= SpecularOcclusion;
+
+			#if APPLY_SSR
+				SpecularLighting.rgb *= lerp( 1 - SSR.a, 1, ClearCoat );
+				SpecularLighting.rgb += lerp( SSR.rgb, 0, ClearCoat );
+			#endif
+
+			SpecularLighting.rgb *= SpecularColor;
+
+			#if APPLY_SSR
+				SpecularLighting.rgb += SSR.rgb * F;
+			#endif
+		}
+		else
+		{
+			#if APPLY_SSR
+				SpecularLighting.rgb = SSR.rgb + SpecularLighting.rgb * (1 - SSR.a);
+			#endif
+			SpecularLighting.rgb *= SpecularColor;
+		}
+
 #if APPLY_SKYLIGHT
 
 		BRANCH
@@ -498,6 +541,27 @@ void ReflectionApplyPS(
 			// Normalize for static skylight types which mix with lightmaps
 			bool bNormalize = SkyLightParameters.z < 1 && View.UseLightmaps > 0 && !ALLOW_STATIC_LIGHTING;
 			float3 SkyLighting = GetSkyLightReflectionSupportingBlend(ReflectionVector, GBuffer.Roughness, bNormalize);
+			SkyLighting *= SpecularColor;
+
+			BRANCH
+			if( GBuffer.ShadingModelID == SHADINGMODELID_CLEAR_COAT )
+			{
+				const float ClearCoat			= GBuffer.CustomData.x;
+				const float ClearCoatRoughness	= GBuffer.CustomData.y;
+
+				// F_Schlick
+				float F0 = 0.04;
+				float Fc = Pow5( 1 - NoV );
+				float F = Fc + (1 - Fc) * F0;
+				F *= ClearCoat;
+
+				#if APPLY_SSR
+					F *= 1 - SSR.a;
+				#endif
+
+				float3 SampleColor = GetSkyLightReflectionSupportingBlend(ReflectionVector, ClearCoatRoughness, bNormalize);	
+				SkyLighting += SampleColor * F;
+			}
 
 		    float DirectionalOcclusion = ScreenSpaceData.DirectionalOcclusion.g;
 		    SkyLighting *= DirectionalOcclusion;
@@ -514,12 +578,7 @@ void ReflectionApplyPS(
 		}
 #endif
 
-#if APPLY_SSR
-		float4 SSR = Texture2DSample( ScreenSpaceReflectionsTexture, ScreenSpaceReflectionsSampler, UV );
-		SpecularLighting.rgb = SpecularLighting.rgb * (1 - SSR.a) + SSR.rgb;
-#endif
-
-		OutColor.rgb = SpecularLighting.rgb * SpecularColor;
+		OutColor.rgb = SpecularLighting.rgb;
 
 		// Transform NaNs to black, transform negative colors to black.
 		OutColor.rgb = -min(-OutColor.rgb, 0.0);
diff --git a/Engine/Shaders/ScreenSpaceReflections.usf b/Engine/Shaders/ScreenSpaceReflections.usf
--- a/Engine/Shaders/ScreenSpaceReflections.usf
+++ b/Engine/Shaders/ScreenSpaceReflections.usf
@@ -42,7 +42,7 @@ float GetRoughness(const in FGBufferData GBuffer)
 {
 	float Roughness = GBuffer.Roughness;
 
-#if USE_CLEARCOAT
+#if 1//USE_CLEARCOAT
 	BRANCH
 	if( GBuffer.ShadingModelID == SHADINGMODELID_CLEAR_COAT )
 	{
