Improved SubsurfaceScatteringProfile quality, base color is now applied after scattering which gives crisper results and more control.

Scanned textures (e.g. human face) include scattering so this is more correct for such assets. Affects the look of emissive lighting (rarely needed and can be fixed by content). All diffuse lighting for SubsurfaceScatteringProfile is now done without the BaseColor (minor performance cost) and reading the Lighting before the SubsurfaceScattering pass would be different.

diff --git a/Engine/Shaders/DeferredShadingCommon.usf b/Engine/Shaders/DeferredShadingCommon.usf
--- a/Engine/Shaders/DeferredShadingCommon.usf
+++ b/Engine/Shaders/DeferredShadingCommon.usf
@@ -240,8 +240,10 @@ struct FGBufferData
 	float3 DiffuseColor;
 	// 0..1 (derived from BaseColor, Metalness, Specular)
 	float3 SpecularColor;
-	// 0..1
+	// 0..1, white for SHADINGMODELID_SUBSURFACE_PROFILE (apply BaseColor after scattering is more correct and less blurry)
 	float3 BaseColor;
+	// 0..1, only needed by SHADINGMODELID_SUBSURFACE_PROFILE which applied BaseColor later
+	float3 StoredBaseColor;
 	// 0..1
 	float Metallic;
 	// 0..1
@@ -389,7 +391,9 @@ FGBufferData DecodeGBufferData(
 	GBuffer.DecalMask = 0;
 	GBuffer.ShadingModelID = DecodeShadingModelId(Decode71(InGBufferB.a, GBuffer.DecalMask));
 	
-	GBuffer.BaseColor = DecodeBaseColor(InGBufferC.rgb);
+	GBuffer.StoredBaseColor = DecodeBaseColor(InGBufferC.rgb);
+	GBuffer.BaseColor = (GBuffer.ShadingModelID == SHADINGMODELID_SUBSURFACE_PROFILE) ? float3(1,1,1) : GBuffer.StoredBaseColor;
+
 	GBuffer.Metallic = InGBufferB.r;
 	GBuffer.Specular = InGBufferB.g;
 	GBuffer.Opacity = InGBufferD.r;
diff --git a/Engine/Shaders/PostProcessSubsurface.usf b/Engine/Shaders/PostProcessSubsurface.usf
--- a/Engine/Shaders/PostProcessSubsurface.usf
+++ b/Engine/Shaders/PostProcessSubsurface.usf
@@ -408,6 +408,7 @@ void SubsurfaceRecombinePS(float4 UVAndScreenPos : TEXCOORD0, out float4 OutColo
 
 	// 0: no SSS,  1: full SSS
 	float BlendFactor = (ScreenSpaceData.GBuffer.ShadingModelID == SHADINGMODELID_SUBSURFACE_PROFILE);
+	// 0: no SSS,  1: full SSS
 	float LerpFactor = 1;
 
 #if HALF_RES
@@ -416,29 +417,35 @@ void SubsurfaceRecombinePS(float4 UVAndScreenPos : TEXCOORD0, out float4 OutColo
 	// minor quality improvement (faces are more detailed in distance)
 	LerpFactor = ComputeFullResLerp(ScreenSpaceData, UVSceneColor, PostprocessInput1Size);
 
-	// hack to debug if the fade is happening at the same time the VisualizeSSS shows it
-//	SSSColor = 0;
 #endif
 #endif // HALF_RES
 
 	float4 SceneColor4 = Texture2DSample(PostprocessInput0, PostprocessInput0Sampler, UVSceneColor);
 
 	float3 ExtractedSpecular = 0;
-	{
+
 #if SSS_RECOMBINE_METHOD
 		// we took the specular highlights out, now we add them back in
-		ExtractedSpecular += LightAccumulator_ReconstructNonDiffuseLighting(SceneColor4);
+		ExtractedSpecular = LightAccumulator_ReconstructNonDiffuseLighting(SceneColor4);
 #endif
 
-		ExtractedSpecular = lerp(SceneColor4.rgb, ExtractedSpecular, BlendFactor);
-	}
+	float3 ExtractedNonSubsurface = lerp(SceneColor4.rgb, ExtractedSpecular, BlendFactor);
 
+	// asset specific color
 	float3 SubsurfaceColor = GetProfileSubsurfaceColor(ScreenSpaceData.GBuffer);
+	float3 FadedSubsurfaceColor = GetProfileSubsurfaceColor(ScreenSpaceData.GBuffer) * LerpFactor;
+
+	float3 ExtractedDiffuse = LightAccumulator_ReconstructDiffuseLighting(SceneColor4);
+
+	// hack to debug if the fade is happening at the same time the VisualizeSSS shows it
+//	SSSColor = float3(0,1,0);
+
+	// combine potentially half res with full res
+	float3 SubsurfaceLighting = lerp(ExtractedDiffuse, SSSColor, FadedSubsurfaceColor);
 
-	// only affect pixels that have SSS
-	LerpFactor *= BlendFactor;
+	// we multiply the base color later in to get more crips human skin textures (scanned data always has Subsurface included)
+	float3 StoredBaseColor = ScreenSpaceData.GBuffer.StoredBaseColor;
 
-	// add back in non subsurface pixels and specular, alpha: 0=don't blend in, 1:fully blend in
-	OutColor = float4(lerp(SceneColor4.rgb, SSSColor * BlendFactor + ExtractedSpecular, SubsurfaceColor * LerpFactor), 0); 
+	OutColor = float4(SubsurfaceLighting * StoredBaseColor + ExtractedNonSubsurface, 0); 
 }
 
