fixed shader warning

diff --git a/Engine/Shaders/LPVGeometryVolumeCommon.usf b/Engine/Shaders/LPVGeometryVolumeCommon.usf
--- a/Engine/Shaders/LPVGeometryVolumeCommon.usf
+++ b/Engine/Shaders/LPVGeometryVolumeCommon.usf
@@ -102,7 +102,7 @@ float3 GvGetDominantDirectionApprox( in GeometryVolumeEntry cell )
 	return float3(0,0,0);
 #else
 	float3 dir = float3( -cell.SH[3], -cell.SH[1], -cell.SH[2] );
-	float3 len = length( dir );
+	float len = length( dir );
 	float den = ( len > 0.01 ) ? len : 1.0f;
 	return dir/den;
 #endif
