Fixed blacks of old tonemapper not being black with log space LUT.

diff --git a/Engine/Shaders/TonemapCommon.usf b/Engine/Shaders/TonemapCommon.usf
--- a/Engine/Shaders/TonemapCommon.usf
+++ b/Engine/Shaders/TonemapCommon.usf
@@ -67,6 +67,9 @@ half3 FilmPostProcess(half3 LinearColor)
 		half3 CurveColor = (MatrixColorH*ColorCurve_Ch1_Ch2.xxx + ColorCurve_Ch1_Ch2.yyy) * rcp(MatrixColorH + ColorCurve_Cm0Cd0_Cd2_Ch0Cm1_Ch3.www) + MatrixColorM;
 	#endif
 
+	// Must hit black by log10(-2.474)
+	CurveColor -= 0.002;
+
 	return CurveColor;
 }
 
