Fixed responsive aa

diff --git a/Engine/Shaders/PostProcessTemporalCommon.usf b/Engine/Shaders/PostProcessTemporalCommon.usf
--- a/Engine/Shaders/PostProcessTemporalCommon.usf
+++ b/Engine/Shaders/PostProcessTemporalCommon.usf
@@ -475,8 +475,6 @@ which can be configured via defines for various temporal AA passes.
 		BlendFinal = 1.0/4.0;
 	#endif
 
-	BlendFinal = 0.04;
-
 	// Offscreen feedback resets.
 	#if AA_LERP 
 		float FixedLerp = 1.0/float(AA_LERP);
