fixed Specular not adjustable for SubSurfaceProfile after recent change

diff --git a/Engine/Shaders/DeferredShadingCommon.usf b/Engine/Shaders/DeferredShadingCommon.usf
--- a/Engine/Shaders/DeferredShadingCommon.usf
+++ b/Engine/Shaders/DeferredShadingCommon.usf
@@ -437,7 +437,7 @@ FGBufferData DecodeGBufferData(
 	{
 		// we apply the BaseColor later, this is why we need StoredBaseColor
 		GBuffer.BaseColor = bChecker;
-		GBuffer.Specular = !bChecker;
+		GBuffer.Specular *= !bChecker;
 	}
 
 	// derived from BaseColor, Metalness, Specular
