Fix ES2 shader build issue.

diff --git a/Engine/Shaders/DeferredShadingCommon.usf b/Engine/Shaders/DeferredShadingCommon.usf
--- a/Engine/Shaders/DeferredShadingCommon.usf
+++ b/Engine/Shaders/DeferredShadingCommon.usf
@@ -364,22 +364,26 @@ void EncodeGBuffer(
 }
 
 // High frequency Checkerboard pattern
-// @param UVSceneColor at pixel center
+// @param PixelPos relative to left top of the rendertarget (not viewport)
 // @return true/false, todo: profile if float 0/1 would be better (need to make sure it's 100% the same)
-bool CheckerFromSceneColorUV(float2 UVSceneColor)
+bool CheckerFromPixelPos(uint2 PixelPos)
 {
-	// relative to left top of the rendertarget (not viewport)
-	uint2 PixelPos = uint2(UVSceneColor * View.BufferSizeAndInvSize.xy);
-
+#if FEATURE_LEVEL >= FEATURE_LEVEL_SM4
 	return (PixelPos.x + PixelPos.y) % 2;
+#else
+	return (uint)(fmod(PixelPos.x + PixelPos.y, 2)) != 0;
+#endif
 }
 
 // High frequency Checkerboard pattern
-// @param PixelPos relative to left top of the rendertarget (not viewport)
+// @param UVSceneColor at pixel center
 // @return true/false, todo: profile if float 0/1 would be better (need to make sure it's 100% the same)
-bool CheckerFromPixelPos(uint2 PixelPos)
+bool CheckerFromSceneColorUV(float2 UVSceneColor)
 {
-	return (PixelPos.x + PixelPos.y) % 2;
+	// relative to left top of the rendertarget (not viewport)
+	uint2 PixelPos = uint2(UVSceneColor * View.BufferSizeAndInvSize.xy);
+
+	return CheckerFromPixelPos(PixelPos);
 }
 
 /** Populates FGBufferData */
