UE4 - Rename GpuSkinCacheCommon.usf to GpuSkinCommon.usf in prep for changes

diff --git a/Engine/Shaders/GpuSkinCacheCommon.usf b/Engine/Shaders/GpuSkinCacheCommon.usf
deleted file mode 100644
diff --git a/Engine/Shaders/GpuSkinCacheComputeShader.usf b/Engine/Shaders/GpuSkinCacheComputeShader.usf
--- a/Engine/Shaders/GpuSkinCacheComputeShader.usf
+++ b/Engine/Shaders/GpuSkinCacheComputeShader.usf
@@ -7,7 +7,7 @@
 #include "Common.usf"
 #include "Definitions.usf"
 #include "VertexFactoryCommon.usf"
-#include "GpuSkinCacheCommon.usf"
+#include "GpuSkinCommon.usf"
 
 #define FBoneMatrix float3x4
 
diff --git a/Engine/Shaders/GpuSkinCommon.usf b/Engine/Shaders/GpuSkinCommon.usf
new file mode 100644
--- /dev/null
+++ b/Engine/Shaders/GpuSkinCommon.usf
@@ -0,0 +1,31 @@
+//// Copyright 1998-2015 Epic Games, Inc. All Rights Reserved.
+
+/*=============================================================================
+	GpuSkinCacheCommon.usf: Common properties and defines for GPU Skinning Cache
+=============================================================================*/
+
+#ifndef GPUSKIN_USE_EXTRA_INFLUENCES
+#define GPUSKIN_USE_EXTRA_INFLUENCES 0
+#endif
+
+// These offsets are in uint32/floats; they attempt to match TGPUSkinVertexBase<> & derived classes:
+//	TGPUSkinVertexBase<false>
+//	TGPUSkinVertexBase<true>
+//	TGPUSkinVertexFloat16Uvs<false, 1..N>
+//	TGPUSkinVertexFloat16Uvs<true, 1..N>
+
+// As we know the properties of the Vertices, the layout looks like this (offset, size in u32):
+//	MEMBER				TGPUSkinVertexBase<false>	TGPUSkinVertexFloat16Uvs<false, 1..N>	TGPUSkinVertexBase<true>	TGPUSkinVertexFloat16Uvs<true, 1..N>	TGPUSkinVertexFloat32Uvs<true, 1..N>
+//	TangentX			0,1							0,1										0,1							0,1										0,1
+//	TangentZ			1,1							1,1										1,1							1,1										1,1
+//	InfluenceBones		2,1							2,1										2,2							2,2										2,2
+//	InfluenceWeights	3,1							3,1										4,2							4,2										4,2
+//	Position			4,3							4,3										6,3							6,3										6,3
+//	UVs[1..N]			N/A							7,1..N									N/A							9,1..N									9,(1..N)*2
+
+#define	GPUSKIN_VB_OFFSET_TANGENT_X			0
+#define	GPUSKIN_VB_OFFSET_TANGENT_Z			1
+#define	GPUSKIN_VB_OFFSET_INFLUENCEBONES	2
+#define	GPUSKIN_VB_OFFSET_INFLUENCEWEIGHTS	(3 + GPUSKIN_USE_EXTRA_INFLUENCES)
+#define	GPUSKIN_VB_OFFSET_POSITION			(GPUSKIN_VB_OFFSET_INFLUENCEWEIGHTS + GPUSKIN_USE_EXTRA_INFLUENCES + 1)
+#define	GPUSKIN_VB_OFFSET_UVS				(GPUSKIN_VB_OFFSET_POSITION + 3)
diff --git a/Engine/Shaders/GpuSkinVertexFactory.usf b/Engine/Shaders/GpuSkinVertexFactory.usf
--- a/Engine/Shaders/GpuSkinVertexFactory.usf
+++ b/Engine/Shaders/GpuSkinVertexFactory.usf
@@ -23,7 +23,7 @@
 =============================================================================*/
 
 #include "VertexFactoryCommon.usf"
-#include "GpuSkinCacheCommon.usf"
+#include "GpuSkinCommon.usf"
 
 #ifndef GPUSKIN_MORPH_BLEND
 #define GPUSKIN_MORPH_BLEND 0
