Improved SubsurfaceScatteringProfile quality by applying Specular Color in full res

diff --git a/Engine/Shaders/DeferredShadingCommon.usf b/Engine/Shaders/DeferredShadingCommon.usf
--- a/Engine/Shaders/DeferredShadingCommon.usf
+++ b/Engine/Shaders/DeferredShadingCommon.usf
@@ -242,8 +242,6 @@ struct FGBufferData
 	float3 SpecularColor;
 	// 0..1, white for SHADINGMODELID_SUBSURFACE_PROFILE (apply BaseColor after scattering is more correct and less blurry)
 	float3 BaseColor;
-	// 0..1, only needed by SHADINGMODELID_SUBSURFACE_PROFILE which applied BaseColor later
-	float3 StoredBaseColor;
 	// 0..1
 	float Metallic;
 	// 0..1
@@ -278,6 +276,11 @@ struct FGBufferData
 	float Depth;
 	// Velocity for motion blur (only used when WRITES_VELOCITY_TO_GBUFFER is enabled)
 	float4 Velocity;
+
+	// 0..1, only needed by SHADINGMODELID_SUBSURFACE_PROFILE which applies BaseColor later
+	float3 StoredBaseColor;
+	// 0..1, only needed by SHADINGMODELID_SUBSURFACE_PROFILE which applies Specular later
+	float3 StoredSpecular;
 };
 
 
@@ -420,7 +423,8 @@ FGBufferData DecodeGBufferData(
 	GBuffer.BaseColor = GBuffer.StoredBaseColor;
 
 	GBuffer.Metallic = InGBufferB.r;
-	GBuffer.Specular = InGBufferB.g;
+	GBuffer.StoredSpecular = InGBufferB.g;
+	GBuffer.Specular = GBuffer.StoredSpecular;
 	GBuffer.Opacity = InGBufferD.r;
 	
 	GBuffer.Roughness = InGBufferB.b;
@@ -439,9 +443,9 @@ FGBufferData DecodeGBufferData(
 
 	if(GBuffer.ShadingModelID == SHADINGMODELID_SUBSURFACE_PROFILE)
 	{
-		// we apply the BaseColor later, this is why we need StoredBaseColor
+		// we apply the BaseColor and Specular later, this is why we need StoredBaseColor
 		GBuffer.BaseColor = bChecker;
-		GBuffer.Specular *= !bChecker;
+		GBuffer.Specular = !bChecker;
 	}
 
 	// derived from BaseColor, Metalness, Specular
diff --git a/Engine/Shaders/MaterialTemplate.usf b/Engine/Shaders/MaterialTemplate.usf
--- a/Engine/Shaders/MaterialTemplate.usf
+++ b/Engine/Shaders/MaterialTemplate.usf
@@ -1088,6 +1088,8 @@ float4 SceneTextureLookup(float2 UV, int SceneTextureIndex, bool bFiltered)
 	else if (SceneTextureIndex == 22)	return float4(GetShadingModelColor(ScreenSpaceData.GBuffer), 1);
 	else if (SceneTextureIndex == 23)	return ScreenSpaceData.AmbientOcclusion;
 	else if (SceneTextureIndex == 24)	return ScreenSpaceData.GBuffer.CustomStencil;
+	else if (SceneTextureIndex == 25)	return ScreenSpaceData.GBuffer.StoredBaseColor;
+	else if (SceneTextureIndex == 26)	return ScreenSpaceData.GBuffer.StoredSpecular;
 #else
 	switch(SceneTextureIndex)
 	{
@@ -1147,6 +1149,10 @@ float4 SceneTextureLookup(float2 UV, int SceneTextureIndex, bool bFiltered)
 		case 23: return ScreenSpaceData.AmbientOcclusion;
 		// PPI_CustomStencil
 		case 24: return ScreenSpaceData.GBuffer.CustomStencil;
+		// PPI_StoredBaseColor
+		case 25: return float4(ScreenSpaceData.GBuffer.StoredBaseColor, 0);
+		// PPI_StoredSpecular
+		case 26: return float4(ScreenSpaceData.GBuffer.StoredSpecular, 0);
 		default:
 			return float4(0, 0, 0, 0);
 	}
diff --git a/Engine/Shaders/PostProcessSubsurface.usf b/Engine/Shaders/PostProcessSubsurface.usf
--- a/Engine/Shaders/PostProcessSubsurface.usf
+++ b/Engine/Shaders/PostProcessSubsurface.usf
@@ -237,18 +237,41 @@ SDiffuseAndSpecular ReconstructDiffuse(float2 UVSceneColor)
 
 	float3 Quant0 = Texture2DSample(PostprocessInput0, PostprocessInput0Sampler, UVSceneColor).rgb;
 
-	// todo: expose as scalability setting
-	// 0:fast but can appear pattern can appear, especially without TemporalAA / 1:high quality
+	// todo: expose as scalability setting (can be evaluate best without TemporalAA)
+	// 0:fast but pattern can appear, 1:better, 2: even better, 3: best but expensive
+
+	// With TemporalAA they all look the same but 0 has an offset which we can fixed by alternating the 2 (L R) or 4 directions (L R U D)
 	#define QUALITY_LEVEL 1
 
 #if QUALITY_LEVEL == 0
 	float3 Quant1 = Texture2DSample(PostprocessInput0, PostprocessInput0Sampler, UVSceneColor + float2(1, 0) * View.BufferSizeAndInvSize.zw).rgb;
-#else
+#elif QUALITY_LEVEL == 1
+	float3 Quant1 = 0.5f * (
+		Texture2DSample(PostprocessInput0, PostprocessInput0Sampler, UVSceneColor + float2( 1,  0) * View.BufferSizeAndInvSize.zw).rgb +
+		Texture2DSample(PostprocessInput0, PostprocessInput0Sampler, UVSceneColor + float2(-1,  0) * View.BufferSizeAndInvSize.zw).rgb);
+#elif QUALITY_LEVEL == 2
 	float3 Quant1 = 0.25f * (
 		Texture2DSample(PostprocessInput0, PostprocessInput0Sampler, UVSceneColor + float2( 1,  0) * View.BufferSizeAndInvSize.zw).rgb +
 		Texture2DSample(PostprocessInput0, PostprocessInput0Sampler, UVSceneColor + float2( 0,  1) * View.BufferSizeAndInvSize.zw).rgb +
 		Texture2DSample(PostprocessInput0, PostprocessInput0Sampler, UVSceneColor + float2(-1,  0) * View.BufferSizeAndInvSize.zw).rgb +
 		Texture2DSample(PostprocessInput0, PostprocessInput0Sampler, UVSceneColor + float2( 0, -1) * View.BufferSizeAndInvSize.zw).rgb);
+#else
+	float3 A = Texture2DSample(PostprocessInput0, PostprocessInput0Sampler, UVSceneColor + float2( 1,  0) * View.BufferSizeAndInvSize.zw).rgb;
+	float3 B = Texture2DSample(PostprocessInput0, PostprocessInput0Sampler, UVSceneColor + float2(-1,  0) * View.BufferSizeAndInvSize.zw).rgb;
+	float3 C = Texture2DSample(PostprocessInput0, PostprocessInput0Sampler, UVSceneColor + float2( 0,  1) * View.BufferSizeAndInvSize.zw).rgb;
+	float3 D = Texture2DSample(PostprocessInput0, PostprocessInput0Sampler, UVSceneColor + float2( 0, -1) * View.BufferSizeAndInvSize.zw).rgb;
+
+	// Luminance could be green channel only
+	float a = Luminance(A);
+	float b = Luminance(B);
+	float c = Luminance(C);
+	float d = Luminance(D);
+
+	float ab = abs(a - b);
+	float cd = abs(c - d);
+
+	// take the average in the direction that avoids dither pattern
+	float3 Quant1 = 0.5f * lerp(A + B, C + D, ab > cd);
 #endif
 
 	Ret.Diffuse = lerp(Quant1, Quant0, bChecker);
@@ -412,9 +435,13 @@ void SubsurfaceRecombinePS(float4 UVAndScreenPos : TEXCOORD0, out float4 OutColo
 	float4 SceneColor4 = Texture2DSample(PostprocessInput0, PostprocessInput0Sampler, UVSceneColor);
 
 
+	// we multiply the base color later in to get more crips human skin textures (scanned data always has Subsurface included)
+	float3 StoredBaseColor = ScreenSpaceData.GBuffer.StoredBaseColor;
+	float3 StoredSpecular = ScreenSpaceData.GBuffer.StoredSpecular;
+
 	SDiffuseAndSpecular DiffuseAndSpecular = ReconstructDiffuse(UVSceneColor);
 	
-	float3 ExtractedNonSubsurface = lerp(SceneColor4.rgb, DiffuseAndSpecular.Specular, BlendFactor);
+	float3 ExtractedNonSubsurface = lerp(SceneColor4.rgb, DiffuseAndSpecular.Specular * StoredSpecular, BlendFactor);
 
 	// asset specific color
 	float3 SubsurfaceColor = GetProfileSubsurfaceColor(ScreenSpaceData.GBuffer);
@@ -426,10 +453,7 @@ void SubsurfaceRecombinePS(float4 UVAndScreenPos : TEXCOORD0, out float4 OutColo
 	// combine potentially half res with full res
 	float3 SubsurfaceLighting = lerp(DiffuseAndSpecular.Diffuse, SSSColor, FadedSubsurfaceColor);
 
-	// we multiply the base color later in to get more crips human skin textures (scanned data always has Subsurface included)
-	float3 StoredBaseColor = ScreenSpaceData.GBuffer.StoredBaseColor;
 
 	OutColor = float4(SubsurfaceLighting * StoredBaseColor + ExtractedNonSubsurface, 0); 
-
 }
 
diff --git a/Engine/Source/Runtime/Engine/Classes/Materials/MaterialExpressionSceneTexture.h b/Engine/Source/Runtime/Engine/Classes/Materials/MaterialExpressionSceneTexture.h
--- a/Engine/Source/Runtime/Engine/Classes/Materials/MaterialExpressionSceneTexture.h
+++ b/Engine/Source/Runtime/Engine/Classes/Materials/MaterialExpressionSceneTexture.h
@@ -13,16 +13,16 @@ enum ESceneTextureId
 	PPI_SceneColor UMETA(DisplayName="SceneColor"),
 	/** Scene depth, single channel, contains the linear depth of the opaque objects */
 	PPI_SceneDepth UMETA(DisplayName="SceneDepth"),
-	/** Material diffuse, RGB color (GBuffer) */
+	/** Material diffuse, RGB color (computed from GBuffer) */
 	PPI_DiffuseColor UMETA(DisplayName="DiffuseColor"),
-	/** Material specular, RGB color (GBuffer) */
+	/** Material specular, RGB color (computed from GBuffer) */
 	PPI_SpecularColor UMETA(DisplayName="SpecularColor"),
-	/** Material subsurface, RGB color (GBuffer) */
+	/** Material subsurface, RGB color (GBuffer, only for some ShadingModels) */
 	PPI_SubsurfaceColor UMETA(DisplayName="SubsurfaceColor"),
-	/** Material base, RGB color (GBuffer) */
-	PPI_BaseColor UMETA(DisplayName="BaseColor"),
-	/** Material specular, single channel (GBuffer) */
-	PPI_Specular UMETA(DisplayName="Specular"),
+	/** Material base, RGB color (GBuffer), can be modified on read by the ShadingModel, consider StoredBasedColor */
+	PPI_BaseColor UMETA(DisplayName="BaseColor (for lighting)"),
+	/** Material specular, single channel (GBuffer), can be modified on read by the ShadingModel, consider StoredSpecular */
+	PPI_Specular UMETA(DisplayName="Specular (for lighting)"),
 	/** Material metallic, single channel (GBuffer) */
 	PPI_Metallic UMETA(DisplayName="Metallic"),
 	/** Normal, RGB in -1..1 range, not normalized (GBuffer) */
@@ -59,6 +59,10 @@ enum ESceneTextureId
 	PPI_AmbientOcclusion UMETA(DisplayName="Ambient Occlusion"),
 	/** Scene stencil, contains CustomStencil mesh property of the opaque objects rendered with CustomDepth */
 	PPI_CustomStencil UMETA(DisplayName="CustomStencil"),
+	/** Material base, RGB color (GBuffer) */
+	PPI_StoredBaseColor UMETA(DisplayName="BaseColor (as stored in GBuffer)"),
+	/** Material specular, single channel (GBuffer) */
+	PPI_StoredSpecular UMETA(DisplayName="Specular (as stored in GBuffer)"),
 };
 
 UCLASS(collapsecategories, hidecategories=Object)
diff --git a/Engine/Source/Runtime/Engine/Private/Materials/HLSLMaterialTranslator.h b/Engine/Source/Runtime/Engine/Private/Materials/HLSLMaterialTranslator.h
--- a/Engine/Source/Runtime/Engine/Private/Materials/HLSLMaterialTranslator.h
+++ b/Engine/Source/Runtime/Engine/Private/Materials/HLSLMaterialTranslator.h
@@ -2668,7 +2668,9 @@ protected:
 			|| SceneTextureId == PPI_Roughness
 			|| SceneTextureId == PPI_MaterialAO
 			|| SceneTextureId == PPI_DecalMask
-			|| SceneTextureId == PPI_ShadingModel;
+			|| SceneTextureId == PPI_ShadingModel
+			|| SceneTextureId == PPI_StoredBaseColor
+			|| SceneTextureId == PPI_StoredSpecular;
 
 		// not yet tracked:
 		//   PPI_SeparateTranslucency, PPI_CustomDepth, PPI_AmbientOcclusion
