very minor quality improvement to SubsurfaceScatterProfile when used with TemporalAA (alternate Checkerboard pattern, assumes TemporalAA sample count in div by 2)

diff --git a/Engine/Shaders/DeferredShadingCommon.usf b/Engine/Shaders/DeferredShadingCommon.usf
--- a/Engine/Shaders/DeferredShadingCommon.usf
+++ b/Engine/Shaders/DeferredShadingCommon.usf
@@ -371,10 +371,14 @@ void EncodeGBuffer(
 // @return true/false, todo: profile if float 0/1 would be better (need to make sure it's 100% the same)
 bool CheckerFromPixelPos(uint2 PixelPos)
 {
+	// todo: Index is float and by staying float we can optimize this 
+	// We alternate the pattern to get 2x supersampling on the lower res data to get more near to full res
+	uint TemporalAASampleIndex = View.TemporalAAParams.x;
+
 #if FEATURE_LEVEL >= FEATURE_LEVEL_SM4
-	return (PixelPos.x + PixelPos.y) % 2;
+	return (PixelPos.x + PixelPos.y + TemporalAASampleIndex) % 2;
 #else
-	return (uint)(fmod(PixelPos.x + PixelPos.y, 2)) != 0;
+	return (uint)(fmod(PixelPos.x + PixelPos.y + TemporalAASampleIndex, 2)) != 0;
 #endif
 }
 
