Fix shader compile failures

diff --git a/Engine/Shaders/DeferredShadingCommon.usf b/Engine/Shaders/DeferredShadingCommon.usf
--- a/Engine/Shaders/DeferredShadingCommon.usf
+++ b/Engine/Shaders/DeferredShadingCommon.usf
@@ -428,7 +428,7 @@ FGBufferData DecodeGBufferData(
 
 	GBuffer.Metallic = InGBufferB.r;
 	GBuffer.StoredSpecular = InGBufferB.g;
-	GBuffer.Specular = GBuffer.StoredSpecular;
+	GBuffer.Specular = GBuffer.StoredSpecular.g;
 	GBuffer.Opacity = InGBufferD.r;
 	
 	GBuffer.Roughness = InGBufferB.b;
diff --git a/Engine/Shaders/MaterialTemplate.usf b/Engine/Shaders/MaterialTemplate.usf
--- a/Engine/Shaders/MaterialTemplate.usf
+++ b/Engine/Shaders/MaterialTemplate.usf
@@ -1088,8 +1088,8 @@ float4 SceneTextureLookup(float2 UV, int SceneTextureIndex, bool bFiltered)
 	else if (SceneTextureIndex == 22)	return float4(GetShadingModelColor(ScreenSpaceData.GBuffer), 1);
 	else if (SceneTextureIndex == 23)	return ScreenSpaceData.AmbientOcclusion;
 	else if (SceneTextureIndex == 24)	return ScreenSpaceData.GBuffer.CustomStencil;
-	else if (SceneTextureIndex == 25)	return ScreenSpaceData.GBuffer.StoredBaseColor;
-	else if (SceneTextureIndex == 26)	return ScreenSpaceData.GBuffer.StoredSpecular;
+	else if (SceneTextureIndex == 25)	return float4(ScreenSpaceData.GBuffer.StoredBaseColor, 0);
+	else if (SceneTextureIndex == 26)	return float4(ScreenSpaceData.GBuffer.StoredSpecular, 0);
 #else
 	switch(SceneTextureIndex)
 	{
