fixed HalfRes outline artifact in SubsurfaceScatteringProfile

diff --git a/Engine/Shaders/PostProcessSubsurface.usf b/Engine/Shaders/PostProcessSubsurface.usf
--- a/Engine/Shaders/PostProcessSubsurface.usf
+++ b/Engine/Shaders/PostProcessSubsurface.usf
@@ -445,7 +445,7 @@ void SubsurfaceRecombinePS(float4 UVAndScreenPos : TEXCOORD0, out float4 OutColo
 
 	// asset specific color
 	float3 SubsurfaceColor = GetProfileSubsurfaceColor(ScreenSpaceData.GBuffer);
-	float3 FadedSubsurfaceColor = GetProfileSubsurfaceColor(ScreenSpaceData.GBuffer) * LerpFactor;
+	float3 FadedSubsurfaceColor = GetProfileSubsurfaceColor(ScreenSpaceData.GBuffer) * (LerpFactor * BlendFactor);
 
 	// hack to debug if the fade is happening at the same time the VisualizeSSS shows it
 //	SSSColor = float3(0,1,0);
