Fixed: Mobile starter content has rendering corruption on Nexus 4 and Nexus 7 (2013 version)

diff --git a/Engine/Shaders/BasePassForForwardShadingPixelShader.usf b/Engine/Shaders/BasePassForForwardShadingPixelShader.usf
--- a/Engine/Shaders/BasePassForForwardShadingPixelShader.usf
+++ b/Engine/Shaders/BasePassForForwardShadingPixelShader.usf
@@ -270,11 +270,13 @@ void Main(
 		}
 		#endif /* MOVABLE_DIRECTIONAL_LIGHT_CSM */
 
-		half NoL = max(0, dot(MaterialParameters.WorldNormal, View.DirectionalLightDirection));
-		half RoL = max(0, dot(MaterialParameters.ReflectionVector, View.DirectionalLightDirection));
+		float NoL = max(0, dot(MaterialParameters.WorldNormal, View.DirectionalLightDirection));
+		float RoL = max(0, dot(MaterialParameters.ReflectionVector, View.DirectionalLightDirection));
 
-		Color += (Shadow * NoL) * View.DirectionalLightColor.rgb * (DiffuseColor + SpecularColor * PhongApprox(Roughness, RoL));
-		#if !MATERIAL_FULLY_ROUGH
+		#if MATERIAL_FULLY_ROUGH
+			Color += (Shadow * NoL) * View.DirectionalLightColor.rgb * DiffuseColor;
+		#else
+			Color += (Shadow * NoL) * View.DirectionalLightColor.rgb * (DiffuseColor + SpecularColor * PhongApprox(Roughness, RoL));
 			half3 SpecularIBL = GetImageBasedReflectionLighting(MaterialParameters, Roughness);
 			// Environment map has been prenormalized, scale by lightmap luminance
 			Color += SpecularIBL * IndirectIrradiance * SpecularColor;
