missing change for CL 2642461
  Vignette for GaussianDOF

diff --git a/Engine/Shaders/PostProcessDOF.usf b/Engine/Shaders/PostProcessDOF.usf
--- a/Engine/Shaders/PostProcessDOF.usf
+++ b/Engine/Shaders/PostProcessDOF.usf
@@ -25,7 +25,7 @@ float ComputeDOFFarFocalMask(float SceneDepth)
 	return saturate((SceneDepth - FarFocalPlane) / View.DepthOfFieldFarTransitionRegion);
 }
 
-// .x:far, .y:near
+// @return .x:far, .y:near
 float2 ComputeDOFFocalMask(float SceneDepth, float SkyWithoutHorizonMask)
 {
 	float2 Ret = float2(ComputeDOFFarFocalMask(SceneDepth), ComputeDOFNearFocalMask(SceneDepth));
@@ -82,33 +82,48 @@ void SetupPS(
 	// can be optimized, needed to not blur the skybox
 	float3 ScreenVector = normalize(mul(float4(ScreenSpacePos, 1, 0), View.ScreenToWorld).xyz);
 	float SkyWithoutHorizonMask = saturate(ScreenVector.z * 3.0f);
-	
+
+	// 0:see though..1:Near blurred
+	float VignetteMask = 0;
+
+#if ENABLE_DOF_VIGNETTE
+	{
+		float DepthOfFieldVignetteMul = DepthOfFieldParams[0].y;
+		float DepthOfFieldVignetteAdd = DepthOfFieldParams[0].z;
+
+		// todo: we could optimize that multiplication away
+		float InvAspectRatio = View.ViewSizeAndInvSize.y * View.ViewSizeAndInvSize.z;
+
+		float CenterDist = length(ScreenSpacePos * float2(1, InvAspectRatio));
+
+		// We prepare the constants on CPU side and use MAD (Multiply and Add) as this is faster
+		VignetteMask = saturate(CenterDist * DepthOfFieldVignetteMul + DepthOfFieldVignetteAdd);
+	}
+#endif
+
 	Mask = ComputeDOFFocalMask(DepthQuad.x, SkyWithoutHorizonMask);
 	Sample = float4(Texture2DSampleLevel(PostprocessInput0, PostprocessInput0Sampler, UV + Offset * float2(-1, 1), 0).rgb, 1);
 	OutColor0 += Sample * Mask.x;
-	OutColor1 += Sample * Mask.y;
+	OutColor1 += Sample * lerp(Mask.y, 1, VignetteMask);
 
 	Mask = ComputeDOFFocalMask(DepthQuad.y, SkyWithoutHorizonMask);
 	Sample = float4(Texture2DSample(PostprocessInput0, PostprocessInput0Sampler, UV + Offset * float2(1, 1)).rgb, 1);
 	OutColor0 += Sample * Mask.x;
-	OutColor1 += Sample * Mask.y;
+	OutColor1 += Sample * lerp(Mask.y, 1, VignetteMask);
 
 	Mask = ComputeDOFFocalMask(DepthQuad.z, SkyWithoutHorizonMask);
 	Sample = float4(Texture2DSample(PostprocessInput0, PostprocessInput0Sampler, UV + Offset * float2(1, -1)).rgb, 1);
 	OutColor0 += Sample * Mask.x;
-	OutColor1 += Sample * Mask.y;
+	OutColor1 += Sample * lerp(Mask.y, 1, VignetteMask);
 
 	Mask = ComputeDOFFocalMask(DepthQuad.w, SkyWithoutHorizonMask);
 	Sample = float4(Texture2DSample(PostprocessInput0, PostprocessInput0Sampler, UV + Offset * float2(-1, -1)).rgb, 1);
 	OutColor0 += Sample * Mask.x;
-	OutColor1 += Sample * Mask.y;
+	OutColor1 += Sample * lerp(Mask.y, 1, VignetteMask);
 
 	// we average 4 samples
 	OutColor0 /= 4;
 	OutColor1 /= 4;
-
-//	OutColor0.rgb *= float3(1,0,0);
-//	OutColor1.rgb *= float3(0,1,0);
 }
 
 
@@ -194,549 +209,3 @@ void MainRecombinePS(
 
 
 
-//
-// PROTOTYPE CIRCLE DOF : WORK IN PROGRESS
-//
-
-// {radius, depth blur amp, depth blur radius, ViewPortScale}
-float4 CircleDofParams;
-
-// CIRCLE DOF: Compute circle of confusion size in pixels.
-float DepthToCoc(float Depth) 
-{
-	float Focus = View.DepthOfFieldFocalDistance;
-	float Radius = CircleDofParams.x;
-	float CocRadius = ((Depth - Focus) / Depth) * Radius;
-	float DepthBlurRadius = (1.0 - exp2(-Depth * CircleDofParams.y)) * CircleDofParams.z;
-	float ReturnCoc = max(abs(CocRadius), DepthBlurRadius);
-	if(CocRadius < 0.0) 
-	{
-		// near CoC is using negative values
-		ReturnCoc = -ReturnCoc;
-	}
-	return ReturnCoc;
-}
-
-// pixel shader entry point
-void CircleSetupPS(float4 UVAndScreenPos : TEXCOORD0, out float4 OutColor0 : SV_Target0)
-{
-	float2 UV = UVAndScreenPos.xy;
-
-	float4 DepthQuad = GatherSceneDepth(UV, PostprocessInput1Size.zw);
-
-	UV = UVAndScreenPos.xy - 0.5*PostprocessInput0Size.zw;
-
-	float4 CW = PostprocessInput0.SampleLevel(PostprocessInput0Sampler, UV.xy, 0);
-	float4 CZ = PostprocessInput0.SampleLevel(PostprocessInput0Sampler, UV.xy, 0, int2(1,0));
-	float4 CX = PostprocessInput0.SampleLevel(PostprocessInput0Sampler, UV.xy, 0, int2(0,1));
-	float4 CY = PostprocessInput0.SampleLevel(PostprocessInput0Sampler, UV.xy, 0, int2(1,1));
-
-	float4 CocQuad = float4(DepthToCoc(DepthQuad.x), DepthToCoc(DepthQuad.y), DepthToCoc(DepthQuad.z), DepthToCoc(DepthQuad.w));
-
-	// Doing a max depth reduction (erode the foreground). Less correct, but less artifacts.
-	// Perhaps need to re-open this in the future.
-
-	#if 0
-		// Stuff max radius in alpha.
-		OutColor0.a = max(max(CocQuad.x,CocQuad.y),max(CocQuad.z,CocQuad.w));
-	#else
-		// This in theory is better but causes bleeding artifacts with temporal AA..
-		// This is important otherwise near thin objects disappear (leaves clamping artifacts in recombined pass).
-		OutColor0.a = CocQuad.x;
-		if(abs(OutColor0.a) > CocQuad.y) OutColor0.a = CocQuad.y;
-		if(abs(OutColor0.a) > CocQuad.z) OutColor0.a = CocQuad.z;
-		if(abs(OutColor0.a) > CocQuad.w) OutColor0.a = CocQuad.w;
-	#endif
-
-	// Remove samples which are outside the size.
-	// TODO: Tune the ScaleFactor.
-	float ScaleFactor = 64.0;
-	float4 W = float4(
-		1.0-saturate(abs(OutColor0.a - CocQuad.x) * ScaleFactor),
-		1.0-saturate(abs(OutColor0.a - CocQuad.y) * ScaleFactor),
-		1.0-saturate(abs(OutColor0.a - CocQuad.z) * ScaleFactor),
-		1.0-saturate(abs(OutColor0.a - CocQuad.w) * ScaleFactor));
-
-	OutColor0.rgb = (1.0/(W.x+W.y+W.z+W.w)) * (CX.rgb*W.x + CY.rgb*W.y + CZ.rgb*W.z + CW.rgb*W.w);
-}
-
-
-
-// {0 to 1} output.
-float NoizNorm(float2 N, float X)
-{
-	N+=X;
-	return frac(sin(dot(N.xy,float2(12.9898, 78.233)))*43758.5453);
-}
-
-// {-1 to 1} output.
-float NoizSnorm(float2 N, float X)
-{
-	return NoizNorm(N,X)*2.0-1.0;
-}
-
-float2 RotVec(float Radius, float Radians)
-{
-	return Radius * float2(cos(Radians), sin(Radians));
-}
-
-float2 RandomOffset;
-
-
-
-
-float Min4(float4 A)
-{
-	return min(min(A.x,A.y),min(A.z,A.w));
-}
-
-float Min16(float4 A, float4 B, float4 C, float4 D)
-{
-	return min(min(Min4(A),Min4(B)),min(Min4(C),Min4(D)));
-}
-
-// This does a 2x2:1 reduction with a 4x4:1 dilation.
-void CircleDilatePS(float4 UVAndScreenPos : TEXCOORD0, out float OutColor : SV_Target0)
-{
-	// Sampling pattern (each gather4)
-    //   d g
-	//   j M  (M={0,0} point)
-
-	#if COMPILER_GLSL || COMPILER_GLSL_ES2
-		float2 UV = UVAndScreenPos.xy + 0.5*PostprocessInput0Size.zw;
-		// This leverages nearest sampling (bilinear won't work).
-		// Probably not the best way to do this.
-		float4 Sd, Sg, Sj, Sm;
-		Sd.x = PostprocessInput0.SampleLevel(PostprocessInput0Sampler, UV.xy, 0, int2(-2,-2)).a;
-		Sd.y = PostprocessInput0.SampleLevel(PostprocessInput0Sampler, UV.xy, 0, int2(-1,-2)).a;
-		Sd.z = PostprocessInput0.SampleLevel(PostprocessInput0Sampler, UV.xy, 0, int2(-2,-1)).a;
-		Sd.w = PostprocessInput0.SampleLevel(PostprocessInput0Sampler, UV.xy, 0, int2(-1,-1)).a;
-		Sg.x = PostprocessInput0.SampleLevel(PostprocessInput0Sampler, UV.xy, 0, int2(0,-2)).a;
-		Sg.y = PostprocessInput0.SampleLevel(PostprocessInput0Sampler, UV.xy, 0, int2(1,-2)).a;
-		Sg.z = PostprocessInput0.SampleLevel(PostprocessInput0Sampler, UV.xy, 0, int2(0,-1)).a;
-		Sg.w = PostprocessInput0.SampleLevel(PostprocessInput0Sampler, UV.xy, 0, int2(1,-1)).a;
-		Sj.x = PostprocessInput0.SampleLevel(PostprocessInput0Sampler, UV.xy, 0, int2(-2,0)).a;
-		Sj.y = PostprocessInput0.SampleLevel(PostprocessInput0Sampler, UV.xy, 0, int2(-1,0)).a;
-		Sj.z = PostprocessInput0.SampleLevel(PostprocessInput0Sampler, UV.xy, 0, int2(-2,1)).a;
-		Sj.w = PostprocessInput0.SampleLevel(PostprocessInput0Sampler, UV.xy, 0, int2(-1,1)).a;
-		Sm.x = PostprocessInput0.SampleLevel(PostprocessInput0Sampler, UV.xy, 0, int2(0,0)).a;
-		Sm.y = PostprocessInput0.SampleLevel(PostprocessInput0Sampler, UV.xy, 0, int2(1,0)).a;
-		Sm.z = PostprocessInput0.SampleLevel(PostprocessInput0Sampler, UV.xy, 0, int2(0,1)).a;
-		Sm.w = PostprocessInput0.SampleLevel(PostprocessInput0Sampler, UV.xy, 0, int2(1,1)).a;
-	#else
-		float2 UV = UVAndScreenPos.xy + 1.0*PostprocessInput0Size.zw;
-		float4 Sd = PostprocessInput0.GatherAlpha(PostprocessInput0Sampler, UV, int2(-2,-2));
-		float4 Sg = PostprocessInput0.GatherAlpha(PostprocessInput0Sampler, UV, int2(0,-2));
-		float4 Sj = PostprocessInput0.GatherAlpha(PostprocessInput0Sampler, UV, int2(-2,0));
-		float4 Sm = PostprocessInput0.GatherAlpha(PostprocessInput0Sampler, UV, int2(0,0));
-	#endif // COMPILER_GLSL
-	
-	// Make sure near is only near blur.
-	OutColor = min(0.0, Min16(Sd, Sg, Sj, Sm));
-}
-
-
-// pixel shader entry point
-void CirclePS(float4 UVAndScreenPos : TEXCOORD0, out float4 OutColor0 : SV_Target0)
-{
-	float2 UV = UVAndScreenPos.xy;
-
-	//
-	// Pass 0
-	// Dilate near minimum CoC (near CoC is negative values).
-	//
-
-	// Fixed maximum search size (in terms of Circle of Confusion radius).
-	// Higher than 8 is too noizy for 4 samples.
-	float Coc = 8.0 * CircleDofParams.w;
-
-	// Get base semi-random direction and dither along radius.
-	// Reused throughout the rest of the algorithm.
-	float TwoPi = 2.0 * 3.14159;
-	float RadianBase = NoizSnorm(UVAndScreenPos.xy, 0.010 * RandomOffset.x) * TwoPi;
-	float RadiusBase = NoizNorm(UVAndScreenPos.xy, 0.013 * RandomOffset.x);
-
-	// Radius 
-	float RadiusBase2 = RadiusBase * (1.0/4.0);
-	float R1 = sqrt(RadiusBase2 + 3.0/4.0) * Coc;
-	float R2 = sqrt(RadiusBase2 + 2.0/4.0) * Coc;
-	float R3 = sqrt(RadiusBase2 + 1.0/4.0) * Coc;
-	float R4 = sqrt(RadiusBase2 + 0.0/4.0) * Coc;
-
-	float2 UV1 = RotVec(R1, RadianBase + TwoPi * 0.0/4.0);
-	float2 UV2 = RotVec(R2, RadianBase + TwoPi * 2.0/4.0);
-	float2 UV3 = RotVec(R3, RadianBase + TwoPi * 1.0/4.0);
-	float2 UV4 = RotVec(R4, RadianBase + TwoPi * 3.0/4.0);
-
-	UV1 = UVAndScreenPos.xy + UV1 * PostprocessInput0Size.zw;
-	UV2 = UVAndScreenPos.xy + UV2 * PostprocessInput0Size.zw;
-	UV3 = UVAndScreenPos.xy + UV3 * PostprocessInput0Size.zw;
-	UV4 = UVAndScreenPos.xy + UV4 * PostprocessInput0Size.zw;
-
-	float D1 = PostprocessInput1.SampleLevel(PostprocessInput1Sampler, UV1, 0).x;
-	float D2 = PostprocessInput1.SampleLevel(PostprocessInput1Sampler, UV2, 0).x;
-	float D3 = PostprocessInput1.SampleLevel(PostprocessInput1Sampler, UV3, 0).x;
-	float D4 = PostprocessInput1.SampleLevel(PostprocessInput1Sampler, UV4, 0).x;
-
-	float NearCoc = 65536.0;
-	float Feather = 2.0;
-	if(abs(D1)+Feather > R1) NearCoc = min(NearCoc, D1);
-	if(abs(D2)+Feather > R2) NearCoc = min(NearCoc, D2);
-	if(abs(D3)+Feather > R3) NearCoc = min(NearCoc, D3);
-	if(abs(D4)+Feather > R4) NearCoc = min(NearCoc, D4);
-
-
-	//
-	// Pass 1
-	//
-
-	// Going to grab sets of 4 samples per pass.
-	// Each set of 4 samples can be a smaller circle of confusion
-	// (aka can be in-front of the larger background).
-
-	// Setup for 12 samples (3 passes of 4 samples).
-	RadiusBase *= (1.0/11.5);
-
-	// Grab circle of confusion for the pixel and pixel color.
-	OutColor0 = Texture2DSampleLevel(PostprocessInput0, PostprocessInput0Sampler, UV, 0);
-	float FarCoc = OutColor0.a;
-
-	// Fix in case no near exists.
-	NearCoc = min(NearCoc, FarCoc);
-
-	// Used for sample pattern.
-	Coc = max(abs(FarCoc),abs(NearCoc));
-	
-	// Bring out to the smaller radius of sample sets.
-	// This has the highest chance of seeing a smaller overlapping CoC.
-	R1 = (RadiusBase+9.0/11.5) * Coc;
-	R2 = (RadiusBase+3.0/11.5) * Coc;
-	R3 = (RadiusBase+6.0/11.5) * Coc;
-	R4 = (RadiusBase+0.0/11.5) * Coc;
-
-	// Ensure at least getting different sample than center pixel.
-	float R1a = max(1.0,R1);
-	float R2a = max(1.0,R2);
-	float R3a = max(1.0,R3);
-	float R4a = max(1.0,R4);
-
-	UV1 = RotVec(R1a, RadianBase + TwoPi * 0.0/12.0);
-	UV2 = RotVec(R2a, RadianBase + TwoPi * 3.0/12.0);
-	UV3 = RotVec(R3a, RadianBase + TwoPi * 6.0/12.0);
-	UV4 = RotVec(R4a, RadianBase + TwoPi * 9.0/12.0);
-
-	UV1 = UVAndScreenPos.xy + UV1 * PostprocessInput0Size.zw;
-	UV2 = UVAndScreenPos.xy + UV2 * PostprocessInput0Size.zw;
-	UV3 = UVAndScreenPos.xy + UV3 * PostprocessInput0Size.zw;
-	UV4 = UVAndScreenPos.xy + UV4 * PostprocessInput0Size.zw;
-
-	float4 C1 = PostprocessInput0.SampleLevel(PostprocessInput0Sampler, UV1, 0);
-	float4 C2 = PostprocessInput0.SampleLevel(PostprocessInput0Sampler, UV2, 0);
-	float4 C3 = PostprocessInput0.SampleLevel(PostprocessInput0Sampler, UV3, 0);
-	float4 C4 = PostprocessInput0.SampleLevel(PostprocessInput0Sampler, UV4, 0);
-
-	// Base weight works around the max(1.0,radius) constraint.
-	// Base weight also shapes to weight higher on the outside radius.
-	float W0 = 1.0 - saturate(Coc);
-	float W1 = R1;
-	float W2 = R2;
-	float W3 = R3;
-	float W4 = R4;
-
-	// Intersection weight: 0=sample does not intersect pixel, to 1=sample intersects.
-	// TODO: Tune feather factors.
-	float IFeather0 = 1.0/4.0;
-	float I1 = saturate((abs(C1.a) - R1) * IFeather0);
-	float I2 = saturate((abs(C2.a) - R2) * IFeather0);
-	float I3 = saturate((abs(C3.a) - R3) * IFeather0);
-	float I4 = saturate((abs(C4.a) - R4) * IFeather0);
-
-	// Check if have a more near intersecting Coc for next pass.
-	float FarCoc2 = FarCoc;
-	if(I1*W1 > 0.0) FarCoc2 = min(FarCoc2, C1.a);
-	if(I2*W2 > 0.0) FarCoc2 = min(FarCoc2, C2.a);
-	if(I3*W3 > 0.0) FarCoc2 = min(FarCoc2, C3.a);
-	if(I4*W4 > 0.0) FarCoc2 = min(FarCoc2, C4.a);
-
-	// Fully ignore intersection weight when in nearfield blur
-	// and sample average CoC is 50% between near and far CoC neighborhood.
-	float AvgCoc = (FarCoc + C1.a + C2.a + C3.a + C4.a) * (1.0/5.0);
-	// Get dilated far.
-	FarCoc = max(FarCoc, max(max(C1.a, C2.a),max(C3.a, C4.a)));
-	// Controls the transition between states.
-	float IFeather1 = 1.0;
-	float IFeather2 = 2.0;
-	float Ignore = saturate(-NearCoc * IFeather1) * saturate(((AvgCoc - FarCoc) / (NearCoc - FarCoc)) * IFeather2);
-
-	W1 *= lerp(I1, 1.0, Ignore);
-	W2 *= lerp(I2, 1.0, Ignore);
-	W3 *= lerp(I3, 1.0, Ignore);
-	W4 *= lerp(I4, 1.0, Ignore);
-
-	// Make sure at least something is not zero.
-	W0 += 1.0/65536.0;
-
-	// Start weighted accumulation.
-	OutColor0.rgb = OutColor0.rgb * W0 + C1.rgb * W1 + C2.rgb * W2 + C3.rgb * W3 + C4.rgb * W4;
-	float Weight = W0+W1+W2+W3+W4;
-
-	// Set current result as possible background.
-	float3 Background = OutColor0.rgb * (1.0/Weight);
-
-
-	//
-	// Pass 2
-	//
-
-	// Drop weight of existing pass if Coc changes too much.
-	float Coc2 = max(abs(FarCoc2),abs(NearCoc));
-	float Drop = (1.0/65536.0) + 1.0 - saturate(abs(Coc - Coc2));
-	OutColor0.rgb *= Drop;
-	Weight *= Drop;
-
-	R1 = (RadiusBase+10.0/11.5) * Coc2;
-	R2 = (RadiusBase+ 4.0/11.5) * Coc2;
-	R3 = (RadiusBase+ 7.0/11.5) * Coc2;
-	R4 = (RadiusBase+ 1.0/11.5) * Coc2;
-
-	R1a = max(1.0,R1);
-	R2a = max(1.0,R2);
-	R3a = max(1.0,R3);
-	R4a = max(1.0,R4);
-
-	UV1 = RotVec(R1a, RadianBase + TwoPi *  8.0/12.0);
-	UV2 = RotVec(R2a, RadianBase + TwoPi * 11.0/12.0);
-	UV3 = RotVec(R3a, RadianBase + TwoPi *  2.0/12.0);
-	UV4 = RotVec(R4a, RadianBase + TwoPi *  5.0/12.0);
-
-	UV1 = UVAndScreenPos.xy + UV1 * PostprocessInput0Size.zw;
-	UV2 = UVAndScreenPos.xy + UV2 * PostprocessInput0Size.zw;
-	UV3 = UVAndScreenPos.xy + UV3 * PostprocessInput0Size.zw;
-	UV4 = UVAndScreenPos.xy + UV4 * PostprocessInput0Size.zw;
-
-	C1 = PostprocessInput0.SampleLevel(PostprocessInput0Sampler, UV1, 0);
-	C2 = PostprocessInput0.SampleLevel(PostprocessInput0Sampler, UV2, 0);
-	C3 = PostprocessInput0.SampleLevel(PostprocessInput0Sampler, UV3, 0);
-	C4 = PostprocessInput0.SampleLevel(PostprocessInput0Sampler, UV4, 0);
-
-	// Lerp to background if outside possibly smaller CoC.
-	C1.rgb = lerp(C1.rgb, Background.rgb, saturate(abs(C1.a) - Coc2));
-	C2.rgb = lerp(C2.rgb, Background.rgb, saturate(abs(C2.a) - Coc2));
-	C3.rgb = lerp(C3.rgb, Background.rgb, saturate(abs(C3.a) - Coc2));
-	C4.rgb = lerp(C4.rgb, Background.rgb, saturate(abs(C4.a) - Coc2));
-
-	W1 = R1;
-	W2 = R2;
-	W3 = R3;
-	W4 = R4;
-
-	// Intersection weight: 0=sample does not intersect pixel, to 1=sample intersects.
-	I1 = saturate((abs(C1.a) - R1) * IFeather0);
-	I2 = saturate((abs(C2.a) - R2) * IFeather0);
-	I3 = saturate((abs(C3.a) - R3) * IFeather0);
-	I4 = saturate((abs(C4.a) - R4) * IFeather0);
-
-	// Check if have a more near intersecting Coc for next pass.
-	float FarCoc3 = FarCoc2;
-	if(I1*W1 > 0.0) FarCoc3 = min(FarCoc3, C1.a);
-	if(I2*W2 > 0.0) FarCoc3 = min(FarCoc3, C2.a);
-	if(I3*W3 > 0.0) FarCoc3 = min(FarCoc3, C3.a);
-	if(I4*W4 > 0.0) FarCoc3 = min(FarCoc3, C4.a);
-
-	W1 *= lerp(I1, 1.0, Ignore);
-	W2 *= lerp(I2, 1.0, Ignore);
-	W3 *= lerp(I3, 1.0, Ignore);
-	W4 *= lerp(I4, 1.0, Ignore);
-
-	OutColor0.rgb += C1.rgb * W1 + C2.rgb * W2 + C3.rgb * W3 + C4.rgb * W4;
-	Weight += W1+W2+W3+W4;
-
-
-	//
-	// Pass 3
-	//
-
-	// Drop weight of existing pass if Coc changes too much.
-	float Coc3 = max(abs(FarCoc3),abs(NearCoc));
-	Drop = (1.0/65536.0) + 1.0 - saturate(abs(Coc2 - Coc3));
-	OutColor0.rgb *= Drop;
-	Weight *= Drop;
-
-	// Send near most CoC back to recombine pass.
-	OutColor0.a = min(FarCoc3, NearCoc);
-
-	R1 = (RadiusBase+11.0/11.5) * Coc3;
-	R2 = (RadiusBase+ 5.0/11.5) * Coc3;
-	R3 = (RadiusBase+ 8.0/11.5) * Coc3;
-	R4 = (RadiusBase+ 2.0/11.5) * Coc3;
-
-	R1a = max(1.0,R1);
-	R2a = max(1.0,R2);
-	R3a = max(1.0,R3);
-	R4a = max(1.0,R4);
-
-	UV1 = RotVec(R1a, RadianBase + TwoPi *  4.0/12.0);
-	UV2 = RotVec(R2a, RadianBase + TwoPi *  7.0/12.0);
-	UV3 = RotVec(R3a, RadianBase + TwoPi * 10.0/12.0);
-	UV4 = RotVec(R4a, RadianBase + TwoPi *  1.0/12.0);
-
-	UV1 = UVAndScreenPos.xy + UV1 * PostprocessInput0Size.zw;
-	UV2 = UVAndScreenPos.xy + UV2 * PostprocessInput0Size.zw;
-	UV3 = UVAndScreenPos.xy + UV3 * PostprocessInput0Size.zw;
-	UV4 = UVAndScreenPos.xy + UV4 * PostprocessInput0Size.zw;
-
-	C1 = PostprocessInput0.SampleLevel(PostprocessInput0Sampler, UV1, 0);
-	C2 = PostprocessInput0.SampleLevel(PostprocessInput0Sampler, UV2, 0);
-	C3 = PostprocessInput0.SampleLevel(PostprocessInput0Sampler, UV3, 0);
-	C4 = PostprocessInput0.SampleLevel(PostprocessInput0Sampler, UV4, 0);
-
-	C1.rgb = lerp(C1.rgb, Background.rgb, saturate(abs(C1.a) - Coc3));
-	C2.rgb = lerp(C2.rgb, Background.rgb, saturate(abs(C2.a) - Coc3));
-	C3.rgb = lerp(C3.rgb, Background.rgb, saturate(abs(C3.a) - Coc3));
-	C4.rgb = lerp(C4.rgb, Background.rgb, saturate(abs(C4.a) - Coc3));
-
-	W1 = R1;
-	W2 = R2;
-	W3 = R3;
-	W4 = R4;
-
-	I1 = saturate((abs(C1.a) - R1) * IFeather0);
-	I2 = saturate((abs(C2.a) - R2) * IFeather0);
-	I3 = saturate((abs(C3.a) - R3) * IFeather0);
-	I4 = saturate((abs(C4.a) - R4) * IFeather0);
-
-	W1 *= lerp(I1, 1.0, Ignore);
-	W2 *= lerp(I2, 1.0, Ignore);
-	W3 *= lerp(I3, 1.0, Ignore);
-	W4 *= lerp(I4, 1.0, Ignore);
-
-	OutColor0.rgb += C1.rgb * W1 + C2.rgb * W2 + C3.rgb * W3 + C4.rgb * W4;
-	Weight += W1+W2+W3+W4;
-
-	OutColor0.rgb *= (1.0/Weight);
-
-}
-
-
-// pixel shader to combine the full res scene and the blurred images behind and in front of the the focal plane
-void MainCircleRecombinePS(in float4 UVAndScreenPos : TEXCOORD0, out float4 OutColor : SV_Target0)
-{
-	// Circle of confusion size for the pixel.
-	float PixDepth = CalcSceneDepth(UVAndScreenPos.xy);
-	float PixCoc = DepthToCoc(PixDepth);
-
-	// Grab nearest Coc.
-	PixCoc = min(PixCoc, PostprocessInput1.SampleLevel(PostprocessInput1Sampler, UVAndScreenPos.xy, 0).a);
-
-	// Transform into sample pattern.
-	PixCoc = abs(PixCoc) * 2.0; // 2x because full instead of half resolution.
-
-	#if 1
-		// Fetch 2 samples mirrored around the pixel 
-		// which is stochastically distributed to fill out the circle of confusion.
-		// TODO: Fix the "random values".
-		float2 UV = UVAndScreenPos.xy * PostprocessInput0Size.xy;
-		float RadianBase = NoizNorm(UVAndScreenPos.xy, 0.010 * RandomOffset.x) * 3.14159 * 2.0;
-		float RadiusJitter = NoizNorm(UVAndScreenPos.xy, 0.013 * RandomOffset.x);
-		float ICoc = PixCoc*sqrt(RadiusJitter);
-		float2 VP = RotVec(float2(ICoc, 0.0), RadianBase) * PostprocessInput0Size.zw;
-
-		// These two samples will still have jitter induced artifacts (very limited utility).
-		// These two samples will also have bleeding artifacts.
-		float4 CA = PostprocessInput0.SampleLevel(PostprocessInput0Sampler, UVAndScreenPos.xy + VP, 0);
-		float4 CB = PostprocessInput0.SampleLevel(PostprocessInput0Sampler, UVAndScreenPos.xy - VP, 0);
-
-		// Weight the two samples to avoid forground into background bleed.
-		float IFeather0 = 1.0/4.0;
-		ICoc *= 0.5; // Coc is half res units.
-		float I1 = saturate((abs(CA.a) - ICoc) * IFeather0);
-		float I2 = saturate((abs(CB.a) - ICoc) * IFeather0);
-		I1 += 1.0/65536.0;
-		I2 += 1.0/65536.0;
-	
-		OutColor = (CA * I1 + CB * I2) * (1.0/(I1+I2));
-	#else
-	    // Possibly higher quality option in the future.
-
-		// Fetch 4 samples in filled disc pattern
-		// which is stochastically distributed to fill out the circle of confusion.
-		float2 UV = UVAndScreenPos.xy * PostprocessInput0Size.xy;
-		float RadianBase = NoizNorm(UVAndScreenPos.xy, 0.010 * RandomOffset.x) * 3.14159 * 2.0;
-		float RadiusBase = NoizNorm(UVAndScreenPos.xy, 0.013 * RandomOffset.x);
-
-		float RadiusBase2 = RadiusBase * (1.0/4.0);
-		float R1 = sqrt(RadiusBase2 + 3.0/4.0) * PixCoc;
-		float R2 = sqrt(RadiusBase2 + 2.0/4.0) * PixCoc;
-		float R3 = sqrt(RadiusBase2 + 1.0/4.0) * PixCoc;
-		float R4 = sqrt(RadiusBase2 + 0.0/4.0) * PixCoc;
-
-		float TwoPi = 3.14159 * 2.0;
-		float2 UV1 = RotVec(R1, RadianBase + TwoPi * 0.0/4.0);
-		float2 UV2 = RotVec(R2, RadianBase + TwoPi * 2.0/4.0);
-		float2 UV3 = RotVec(R3, RadianBase + TwoPi * 1.0/4.0);
-		float2 UV4 = RotVec(R4, RadianBase + TwoPi * 3.0/4.0);
-
-		UV1 = UVAndScreenPos.xy + UV1 * PostprocessInput0Size.zw;
-		UV2 = UVAndScreenPos.xy + UV2 * PostprocessInput0Size.zw;
-		UV3 = UVAndScreenPos.xy + UV3 * PostprocessInput0Size.zw;
-		UV4 = UVAndScreenPos.xy + UV4 * PostprocessInput0Size.zw;
-
-		float4 CA = PostprocessInput0.SampleLevel(PostprocessInput0Sampler, UV1, 0);
-		float4 CB = PostprocessInput0.SampleLevel(PostprocessInput0Sampler, UV2, 0);
-		float4 CC = PostprocessInput0.SampleLevel(PostprocessInput0Sampler, UV3, 0);
-		float4 CD = PostprocessInput0.SampleLevel(PostprocessInput0Sampler, UV4, 0);
-
-		// Weight the two samples to avoid forground into background bleed.
-		float IFeather0 = 1.0/4.0;
-		float ICoc = PixCoc * 0.5; // Coc is half res units.
-		float I1 = saturate((abs(CA.a) - ICoc) * IFeather0);
-		float I2 = saturate((abs(CB.a) - ICoc) * IFeather0);
-		float I3 = saturate((abs(CC.a) - ICoc) * IFeather0);
-		float I4 = saturate((abs(CD.a) - ICoc) * IFeather0);
-
-		// Make sure something is non-zero.
-		I1 += 1.0/65536.0;
-		I2 += 1.0/65536.0;
-		I3 += 1.0/65536.0;
-		I4 += 1.0/65536.0;
-
-		OutColor = (CA * I1 + CB * I2 + CC * I3 + CD * I4) * (1.0/(I1+I2+I3+I4));
-	#endif
-
-	// Grab the half resolution neighborhood to remove the artifacts from the full resolution output.
-	// Nearest location.
-	#if 1
-		// This has higher in-focus contrast, but possibly lower noise reduction later.
-		float2 HUVBase = UVAndScreenPos.xy * PostprocessInput1Size.xy - 0.5;
-		float2 HUVFrac = frac(HUVBase);
-		float2 HUV = (trunc(HUVBase) + 0.5) * PostprocessInput1Size.zw;
-	#else
-		// This makes the mostly in-focus transition bad (too blurry).
-		float2 HUV = UVAndScreenPos.xy - 0.5 * PostprocessInput1Size.zw;
-	#endif
-
-	// Load four nearest samples.
-	float4 H0 = PostprocessInput1.SampleLevel(PostprocessInput1Sampler, HUV, 0);
-	float4 H1 = PostprocessInput1.SampleLevel(PostprocessInput1Sampler, HUV, 0, int2(1,0));
-	float4 H2 = PostprocessInput1.SampleLevel(PostprocessInput1Sampler, HUV, 0, int2(0,1));
-	float4 H3 = PostprocessInput1.SampleLevel(PostprocessInput1Sampler, HUV, 0, int2(1,1));
-
-	// TODO: This would work a lot better in YUV style colorspace?
-	// Limit the full resolution to remove jitter artifacts.
-	float4 HMax = max(max(H0,H1),max(H2,H3));
-	float4 HMin = min(min(H0,H1),min(H2,H3));
-
-	#if 1
-		// Increase constrast of limit a little to workaround to strong denoise at near-in-focus.
-		float4 HD = HMin / 8.0;
-		float Small = 1.0 - saturate(PixCoc*PixCoc*(1.0/64.0));
-		HMax += HD * Small;
-		HMin -= HD * Small;
-	#endif
-
-	// Blend in the limited version quickly to remove HDR jitter artifacts and noise.
-	float4 OutLimited = min(max(OutColor,HMin),HMax);
-	OutColor = lerp(OutColor, OutLimited, saturate(PixCoc*PixCoc*4.0));
-}
-
-
