UE4 - Enable GPU skin caching on PS4
- Minor optimization: Avoid extra UB search on a check()

diff --git a/Engine/Shaders/GpuSkinCacheComputeShader.usf b/Engine/Shaders/GpuSkinCacheComputeShader.usf
--- a/Engine/Shaders/GpuSkinCacheComputeShader.usf
+++ b/Engine/Shaders/GpuSkinCacheComputeShader.usf
@@ -9,6 +9,9 @@
 #include "VertexFactoryCommon.usf"
 #include "GpuSkinCommon.usf"
 
+// Set to 1 to enable passing through values from vertex buffer (ie do not skin)
+#define GPUSKIN_USE_PASSTHROUGH	0
+
 #define FBoneMatrix float3x4
 
 Buffer<float4> BoneMatrices;
@@ -128,7 +131,6 @@ void SkinCacheUpdateBatchCS(uint3 GroupID : SV_GroupID,
  uint3 DispatchThreadID : SV_DispatchThreadID,
  uint3 GroupThreadID : SV_GroupThreadID)
 {
-#if !PS4_PROFILE	//@todo ps4: shader compiler issue, but this isn't currently used so just skip it on ps4 for now
 	uint ThreadId = DispatchThreadID.x;
 
 	if (ThreadId > SkinCacheVertexCount)
@@ -182,9 +184,11 @@ void SkinCacheUpdateBatchCS(uint3 GroupID : SV_GroupID,
 		SkinCacheBuffer[OutputOffset + i] = SkinStreamInputBuffer[InputOffset + i];
 	}
 
-	//for (uint i = 0; i < FloatCount; i++)
-	//{
-	//	SkinCacheBuffer[OutputOffset + i] = SkinStreamInputBuffer[InputOffset + i];
-	//}
+	// Passthrough debug code
+#if GPUSKIN_USE_PASSTHROUGH
+	for (uint i = 0; i < FloatCount; i++)
+	{
+		SkinCacheBuffer[OutputOffset + i] = SkinStreamInputBuffer[InputOffset + i];
+	}
 #endif
 }
diff --git a/Engine/Source/Runtime/Renderer/Public/MaterialShader.h b/Engine/Source/Runtime/Renderer/Public/MaterialShader.h
--- a/Engine/Source/Runtime/Renderer/Public/MaterialShader.h
+++ b/Engine/Source/Runtime/Renderer/Public/MaterialShader.h
@@ -108,9 +108,10 @@ public:
 	template<typename ShaderRHIParamRef>
 	void SetParameters(FRHICommandList& RHICmdList, const ShaderRHIParamRef ShaderRHI,const FSceneView& View)
 	{
-		check(GetUniformBufferParameter<FViewUniformShaderParameters>().IsInitialized());
+		const auto& UBParameter = GetUniformBufferParameter<FViewUniformShaderParameters>();
+		check(UBParameter.IsInitialized());
 		CheckShaderIsValid();
-		SetUniformBufferParameter(RHICmdList, ShaderRHI,GetUniformBufferParameter<FViewUniformShaderParameters>(),View.UniformBuffer);
+		SetUniformBufferParameter(RHICmdList, ShaderRHI, UBParameter,View.UniformBuffer);
 	}
 
 	static void ModifyCompilationEnvironment(EShaderPlatform Platform, FShaderCompilerEnvironment& OutEnvironment)
