Bug UE-20116: Fixes Screen Space Subsurface Scatering pass' regression brought by CL 2640696, outputing the gbuffer's diffuse color additionaly to the scene color, creating global ilumination changes when material using SSS are coming in and out of the camera frutsum.

diff --git a/Engine/Shaders/PostProcessSubsurface.usf b/Engine/Shaders/PostProcessSubsurface.usf
--- a/Engine/Shaders/PostProcessSubsurface.usf
+++ b/Engine/Shaders/PostProcessSubsurface.usf
@@ -401,6 +401,12 @@ void SubsurfaceRecombinePS(float4 UVAndScreenPos : TEXCOORD0, out float4 OutColo
 	float2 UVSceneColor = (PixelPos + 0.5f) * View.BufferSizeAndInvSize.zw;
 
 	FScreenSpaceData ScreenSpaceData = GetScreenSpaceData(UVSceneColor);
+	
+	if (ScreenSpaceData.GBuffer.ShadingModelID != SHADINGMODELID_SUBSURFACE_PROFILE)
+	{
+		OutColor = Texture2DSample(PostprocessInput0, PostprocessInput0Sampler, UVSceneColor);
+		return;
+	}
 
 	// bilinear filtering
 	float3 SSSColor;
@@ -411,9 +417,6 @@ void SubsurfaceRecombinePS(float4 UVAndScreenPos : TEXCOORD0, out float4 OutColo
 		SSSColor = SSSColorWithAlpha.rgb / max(SSSColorWithAlpha.a, 0.00001f);
 	}
 
-	// 0: no SSS,  1: full SSS
-	float BlendFactor = (ScreenSpaceData.GBuffer.ShadingModelID == SHADINGMODELID_SUBSURFACE_PROFILE);
-
 	float LerpFactor = 1;
 
 #if HALF_RES
@@ -432,20 +435,17 @@ void SubsurfaceRecombinePS(float4 UVAndScreenPos : TEXCOORD0, out float4 OutColo
 	LerpFactor = 0;
 #endif
 
-	float4 SceneColor4 = Texture2DSample(PostprocessInput0, PostprocessInput0Sampler, UVSceneColor);
-
-
 	// we multiply the base color later in to get more crips human skin textures (scanned data always has Subsurface included)
 	float3 StoredBaseColor = ScreenSpaceData.GBuffer.StoredBaseColor;
 	float3 StoredSpecular = ScreenSpaceData.GBuffer.StoredSpecular;
 
 	SDiffuseAndSpecular DiffuseAndSpecular = ReconstructDiffuse(UVSceneColor);
 	
-	float3 ExtractedNonSubsurface = lerp(SceneColor4.rgb, DiffuseAndSpecular.Specular * StoredSpecular, BlendFactor);
+	float3 ExtractedNonSubsurface = DiffuseAndSpecular.Specular * StoredSpecular;
 
 	// asset specific color
 	float3 SubsurfaceColor = GetProfileSubsurfaceColor(ScreenSpaceData.GBuffer);
-	float3 FadedSubsurfaceColor = GetProfileSubsurfaceColor(ScreenSpaceData.GBuffer) * (LerpFactor * BlendFactor);
+	float3 FadedSubsurfaceColor = GetProfileSubsurfaceColor(ScreenSpaceData.GBuffer) * LerpFactor;
 
 	// hack to debug if the fade is happening at the same time the VisualizeSSS shows it
 //	SSSColor = float3(0,1,0);
diff --git a/Engine/Source/Runtime/Renderer/Private/PostProcess/PostProcessSubsurface.cpp b/Engine/Source/Runtime/Renderer/Private/PostProcess/PostProcessSubsurface.cpp
--- a/Engine/Source/Runtime/Renderer/Private/PostProcess/PostProcessSubsurface.cpp
+++ b/Engine/Source/Runtime/Renderer/Private/PostProcess/PostProcessSubsurface.cpp
@@ -849,7 +849,7 @@ void FRCPassPostProcessSubsurfaceRecombine::Process(FRenderingCompositePassConte
 
 	TShaderMapRef<FPostProcessVS> VertexShader(Context.GetShaderMap());
 
-	if(GetInput(ePId_Input1))
+	if(GetInput(ePId_Input1)->IsValid())
 	{
 		if(bHalfRes)
 		{
