Implements the Decal Derivative material expression.

It outputs decal's texture coordinate DDX/DDY to replace Decal Mipmap Level material expression fixing the decal's 2x2 artifacts on meshes edges, but with anisotropic support.

diff --git a/Engine/Documentation/Source/Engine/Actors/DecalActor/DecalUserGuide.INT.udn b/Engine/Documentation/Source/Engine/Actors/DecalActor/DecalUserGuide.INT.udn
--- a/Engine/Documentation/Source/Engine/Actors/DecalActor/DecalUserGuide.INT.udn
+++ b/Engine/Documentation/Source/Engine/Actors/DecalActor/DecalUserGuide.INT.udn
@@ -143,26 +143,14 @@ Decals may have 2x2 pixels block artifacts on edges as shown in the screenshot b
 
 ![](2x2_artifacts.bmp)
 
-This is where the node "Decal Mipmap Level" comes in. This node must be used with  has a very big performance impact.
+This is where the node "Decal Derivative" comes in. This node must be used carefully as it has a very big performance impact. It returns the derivatives on the X and Y axis of the decal's default texture coordinates necessary for anisotropic texture filtering, but computed differently than using the hardware's default and DDX/DDY nodes, to avoid this 2x2 pixel block artifact.
 
-![](decal_mipmap_material_1.png)
+![](Decal_Derivative.png)
 
 And using it fixes the artifacts on the decal that benefits from it.
 
 ![](2x2_no_artifacts.bmp)
 
-### Missing SampleTexture UV
-
-Now this node is so expensive from a computing point of view, that it doesn't have UV input as the SampleTexture. But in case of a tilling, you can still do so by patching up the mipmap level manualy. Here is an example to have a 2^4 = 16 tiling repetition.
-
-![](decal_mipmap_material_2.png)
-
 ### Current limitations
 
-* DecalMipmapLevel doesn't support UV
-* DecalMipmapLevel doesn't work with anysotropic filtering
-
-
-
-
-
+* DecalMipmapLevel doesn't support custom UV, but you can patch'up it's output.
diff --git a/Engine/Documentation/Source/Shared/GraphNodes/Material/Material.INT.udn b/Engine/Documentation/Source/Shared/GraphNodes/Material/Material.INT.udn
--- a/Engine/Documentation/Source/Shared/GraphNodes/Material/Material.INT.udn
+++ b/Engine/Documentation/Source/Shared/GraphNodes/Material/Material.INT.udn
@@ -320,8 +320,16 @@ Engine/Rendering/Materials/ExpressionReference/Custom#customtexture
 [/VAR]
 [/EXCERPT:UMaterialExpressionCustomTexture]
 
+[EXCERPT:UMaterialExpressionDecalDerivative]
+The DecalDerivative is used to explicitly compute a decal's texture coordinate's derivative assuming there is no custom uv, to avoid the 2x2 pixel mipmap artefacts on the edges where the decal is project to, but still with support of anisotropic filtering.
+
+[VAR:ToolTipFullLink]
+Engine/Actors/DecalActor/DecalUserGuide
+[/VAR]
+[/EXCERPT:UMaterialExpressionDecalDerivative]
+
 [EXCERPT:UMaterialExpressionDecalMipmapLevel]
-The DecalMipmapLevel is used to explicitly compute a decal's texture's mipmap level assuming there is no custom uv, to avoid the 2x2 pixel mipmap artefacts on the edges where the decal is project to.
+Deprecated. Use DecalDerivative instead.
 
 [VAR:ToolTipFullLink]
 Engine/Actors/DecalActor/DecalUserGuide
diff --git a/Engine/Shaders/MaterialTemplate.usf b/Engine/Shaders/MaterialTemplate.usf
--- a/Engine/Shaders/MaterialTemplate.usf
+++ b/Engine/Shaders/MaterialTemplate.usf
@@ -355,7 +355,7 @@ float3 GetTranslatedWorldPosition_NoMaterialOffsets(FMaterialPixelParameters Par
 		return DecalVector.xy;
 	}
 
-	float ComputeDecalMipmapLevel(FMaterialPixelParameters Parameters, float2 TextureSize)
+	float2 ComputeDecalDDX(FMaterialPixelParameters Parameters)
 	{
 		/*
 		 * Assuming where in a pixel shader invocation, then we compute manualy compute two d(uv)/d(x)
@@ -369,19 +369,39 @@ float3 GetTranslatedWorldPosition_NoMaterialOffsets(FMaterialPixelParameters Par
 		 * So we have two diff on the X axis, we want the one that has the smallest length
 		 * to avoid the 2x2 pixels mipmap artefacts on the edges. 
 		 */
-		float2 UvPixelDiffX = (dot(UvDiffX0, UvDiffX0) < dot(UvDiffX1, UvDiffX1) ? UvDiffX0 : UvDiffX1) * TextureSize;
-	
+		return dot(UvDiffX0, UvDiffX0) < dot(UvDiffX1, UvDiffX1) ? UvDiffX0 : UvDiffX1;
+	}
+
+	float2 ComputeDecalDDY(FMaterialPixelParameters Parameters)
+	{
 		// do same for the Y axis
 		float4 ScreenDeltaY = ddy(Parameters.ScreenPosition);
 		float2 UvDiffY0 = Parameters.TexCoords[0] - ComputeDecalUVFromScreenPosition(Parameters.ScreenPosition - ScreenDeltaY);
 		float2 UvDiffY1 = ComputeDecalUVFromScreenPosition(Parameters.ScreenPosition + ScreenDeltaY) - Parameters.TexCoords[0];
-		float2 UvPixelDiffY = (dot(UvDiffY0, UvDiffY0) < dot(UvDiffY1, UvDiffY1) ? UvDiffY0 : UvDiffY1) * TextureSize;
+
+		return dot(UvDiffY0, UvDiffY0) < dot(UvDiffY1, UvDiffY1) ? UvDiffY0 : UvDiffY1;
+	}
+
+	float ComputeDecalMipmapLevel(FMaterialPixelParameters Parameters, float2 TextureSize)
+	{
+		float2 UvPixelDiffX = ComputeDecalDDX(Parameters) * TextureSize;
+		float2 UvPixelDiffY = ComputeDecalDDY(Parameters) * TextureSize;
 
 		// Computes the mipmap level
 		float MaxDiff = max(dot(UvPixelDiffX, UvPixelDiffX), dot(UvPixelDiffY, UvPixelDiffY));
 		return 0.5 * log2(MaxDiff);
 	}
 #else
+	float2 ComputeDecalDDX(FMaterialPixelParameters Parameters)
+	{
+		return 0.0f;
+	}
+	
+	float2 ComputeDecalDDY(FMaterialPixelParameters Parameters)
+	{
+		return 0.0f;
+	}
+
 	float ComputeDecalMipmapLevel(FMaterialPixelParameters Parameters, float2 TextureSize)
 	{
 		return 0.0f;
diff --git a/Engine/Source/Runtime/Engine/Classes/Materials/MaterialExpressionDecalDerivative.h b/Engine/Source/Runtime/Engine/Classes/Materials/MaterialExpressionDecalDerivative.h
new file mode 100644
--- /dev/null
+++ b/Engine/Source/Runtime/Engine/Classes/Materials/MaterialExpressionDecalDerivative.h
@@ -0,0 +1,19 @@
+// Copyright 1998-2015 Epic Games, Inc. All Rights Reserved.
+
+
+#pragma once
+#include "Materials/MaterialExpression.h"
+#include "MaterialExpressionDecalDerivative.generated.h"
+
+
+UCLASS(collapsecategories, hidecategories = Object)
+class UMaterialExpressionDecalDerivative : public UMaterialExpression
+{
+	GENERATED_UCLASS_BODY()
+
+
+	// Begin UMaterialExpression Interface
+	virtual int32 Compile(class FMaterialCompiler* Compiler, int32 OutputIndex, int32 MultiplexIndex) override;
+	virtual void GetCaption(TArray<FString>& OutCaptions) const override;
+	// End UMaterialExpression Interface
+};
diff --git a/Engine/Source/Runtime/Engine/Private/Materials/HLSLMaterialTranslator.h b/Engine/Source/Runtime/Engine/Private/Materials/HLSLMaterialTranslator.h
--- a/Engine/Source/Runtime/Engine/Private/Materials/HLSLMaterialTranslator.h
+++ b/Engine/Source/Runtime/Engine/Private/Materials/HLSLMaterialTranslator.h
@@ -2457,6 +2457,11 @@ protected:
 
 	virtual int32 TextureDecalMipmapLevel(int32 TextureSizeInput) override
 	{
+		if (Material->GetMaterialDomain() != MD_DeferredDecal)
+		{
+			return Errorf(TEXT("Decal mipmap level only available in the decal material domain."));
+		}
+
 		EMaterialValueType TextureSizeType = GetParameterType(TextureSizeInput);
 
 		if (TextureSizeType != MCT_Float2)
@@ -2474,6 +2479,19 @@ protected:
 			);
 	}
 
+	virtual int32 TextureDecalDerivative(bool bDDY) override
+	{
+		if (Material->GetMaterialDomain() != MD_DeferredDecal)
+		{
+			return Errorf(TEXT("Decal derivatives only available in the decal material domain."));
+		}
+
+		return AddCodeChunk(
+			MCT_Float2,
+			bDDY ? TEXT("ComputeDecalDDY(Parameters)") : TEXT("ComputeDecalDDX(Parameters)")
+			);
+	}
+
 	virtual int32 PixelDepth() override
 	{
 		if (ShaderFrequency != SF_Pixel && ShaderFrequency != SF_Compute)
diff --git a/Engine/Source/Runtime/Engine/Private/Materials/MaterialExpressions.cpp b/Engine/Source/Runtime/Engine/Private/Materials/MaterialExpressions.cpp
--- a/Engine/Source/Runtime/Engine/Private/Materials/MaterialExpressions.cpp
+++ b/Engine/Source/Runtime/Engine/Private/Materials/MaterialExpressions.cpp
@@ -31,6 +31,7 @@
 #include "Materials/MaterialExpressionCustom.h"
 #include "Materials/MaterialExpressionDDX.h"
 #include "Materials/MaterialExpressionDDY.h"
+#include "Materials/MaterialExpressionDecalDerivative.h"
 #include "Materials/MaterialExpressionDecalMipmapLevel.h"
 #include "Materials/MaterialExpressionDepthFade.h"
 #include "Materials/MaterialExpressionDepthOfFieldFunction.h"
@@ -8821,6 +8822,43 @@ void UMaterialExpressionAntialiasedTextureMask::SetDefaultTexture()
 	Texture = LoadObject<UTexture2D>(NULL, TEXT("/Engine/EngineResources/DefaultTexture.DefaultTexture"), NULL, LOAD_None, NULL);
 }
 
+//
+//	UMaterialExpressionDecalDerivative
+//
+UMaterialExpressionDecalDerivative::UMaterialExpressionDecalDerivative(const FObjectInitializer& ObjectInitializer)
+	: Super(ObjectInitializer)
+{
+	// Structure to hold one-time initialization
+	struct FConstructorStatics
+	{
+		FText NAME_Vectors;
+		FConstructorStatics()
+			: NAME_Vectors(LOCTEXT("Utils", "Utils"))
+		{
+		}
+	};
+	static FConstructorStatics ConstructorStatics;
+
+	MenuCategories.Add(ConstructorStatics.NAME_Vectors);
+	//bCollapsed = true;
+	bShaderInputData = true;
+	bShowOutputNameOnPin = true;
+	
+	Outputs.Reset();
+	Outputs.Add(FExpressionOutput(TEXT("DDX")));
+	Outputs.Add(FExpressionOutput(TEXT("DDY")));
+}
+
+int32 UMaterialExpressionDecalDerivative::Compile(class FMaterialCompiler* Compiler, int32 OutputIndex, int32 MultiplexIndex)
+{
+	return Compiler->TextureDecalDerivative(OutputIndex == 1);
+}
+
+void UMaterialExpressionDecalDerivative::GetCaption(TArray<FString>& OutCaptions) const
+{
+	OutCaptions.Add(TEXT("Decal Derivative"));
+}
+
 //
 //	UMaterialExpressionDecalMipmapLevel
 //
diff --git a/Engine/Source/Runtime/Engine/Public/MaterialCompiler.h b/Engine/Source/Runtime/Engine/Public/MaterialCompiler.h
--- a/Engine/Source/Runtime/Engine/Public/MaterialCompiler.h
+++ b/Engine/Source/Runtime/Engine/Public/MaterialCompiler.h
@@ -103,6 +103,7 @@ public:
 	virtual int32 TextureSample(int32 Texture,int32 Coordinate,enum EMaterialSamplerType SamplerType,int32 MipValue0Index=INDEX_NONE,int32 MipValue1Index=INDEX_NONE,ETextureMipValueMode MipValueMode=TMVM_None,ESamplerSourceMode SamplerSource=SSM_FromTextureAsset) = 0;
 
 	virtual int32 TextureDecalMipmapLevel(int32 TextureSizeInput) = 0;
+	virtual int32 TextureDecalDerivative(bool bDDY) = 0;
 
 	virtual int32 Texture(UTexture* Texture,ESamplerSourceMode SamplerSource=SSM_FromTextureAsset) = 0;
 	virtual int32 TextureParameter(FName ParameterName,UTexture* DefaultTexture,ESamplerSourceMode SamplerSource=SSM_FromTextureAsset) = 0;
@@ -265,6 +266,7 @@ public:
 	virtual int32 TextureCoordinate(uint32 CoordinateIndex, bool UnMirrorU, bool UnMirrorV) override { return Compiler->TextureCoordinate(CoordinateIndex, UnMirrorU, UnMirrorV); }
 
 	virtual int32 TextureDecalMipmapLevel(int32 TextureSizeInput) override { return Compiler->TextureDecalMipmapLevel(TextureSizeInput); }
+	virtual int32 TextureDecalDerivative(bool bDDY) override { return Compiler->TextureDecalDerivative(bDDY); }
 
 	virtual int32 Texture(UTexture* InTexture,ESamplerSourceMode SamplerSource=SSM_FromTextureAsset) override { return Compiler->Texture(InTexture,SamplerSource); }
 	virtual int32 TextureParameter(FName ParameterName,UTexture* DefaultValue,ESamplerSourceMode SamplerSource=SSM_FromTextureAsset) override { return Compiler->TextureParameter(ParameterName,DefaultValue,SamplerSource); }
