fixed shader warning

diff --git a/Engine/Shaders/PostProcessNoiseBlur.usf b/Engine/Shaders/PostProcessNoiseBlur.usf
--- a/Engine/Shaders/PostProcessNoiseBlur.usf
+++ b/Engine/Shaders/PostProcessNoiseBlur.usf
@@ -33,6 +33,8 @@ void MainPS(float4 UVAndScreenPos : TEXCOORD0, out float4 OutColor : SV_Target0)
 
 	const int Count = 16;
 
+	OutColor = 0;
+
 	for(int i = 0; i < Count; ++i)
 	{
 		// more samples in the center, can be adjusted
