Force shader recompile to ensure changes in 2665840 take effect.

diff --git a/Engine/Shaders/ShaderVersion.usf b/Engine/Shaders/ShaderVersion.usf
--- a/Engine/Shaders/ShaderVersion.usf
+++ b/Engine/Shaders/ShaderVersion.usf
@@ -2,4 +2,4 @@
 // Each time the console command is executed it generates a new GUID. As this file is included
 // in common.usf (which should be included in any shader) it allows to invalidate the shader DDC.
 // 
-// GUID = AD9819114E9A382DD5C6E2B82A720CC3
\ No newline at end of file
+// GUID = 4C4BD3C9571F4F3A93CE647EC1EAACF6
\ No newline at end of file
