Merging general rendering & shader changes necessary to support desktop Metal.

diff --git a/Engine/Shaders/BasePassForForwardShadingVertexShader.usf b/Engine/Shaders/BasePassForForwardShadingVertexShader.usf
--- a/Engine/Shaders/BasePassForForwardShadingVertexShader.usf
+++ b/Engine/Shaders/BasePassForForwardShadingVertexShader.usf
@@ -7,7 +7,7 @@
 #define IS_FORWARD_BASEPASS_VERTEX_SHADER 1
 
 // @todo urban: branch was: #define FORCE_FLOATS 1
-#define FORCE_FLOATS METAL_PROFILE
+#define FORCE_FLOATS (COMPILER_METAL)
 
 #include "Common.usf"
 #include "BasePassForForwardShadingCommon.usf"
diff --git a/Engine/Shaders/Common.usf b/Engine/Shaders/Common.usf
--- a/Engine/Shaders/Common.usf
+++ b/Engine/Shaders/Common.usf
@@ -12,7 +12,7 @@
 	#include "PS4/PS4Common.usf"
 #endif
 
-#if METAL_PROFILE || METAL_MRT_PROFILE || METAL_SM5_PROFILE
+#if COMPILER_METAL
 	// Helps with iteration when changing Metal shader code generation backend.
 	#include "MetalCommon.usf"
 #endif
@@ -63,7 +63,7 @@
 	#define MaterialFloat4x3 float4x3 
 #endif
 
-#if COMPILER_GLSL_ES2 || METAL_PROFILE || METAL_MRT_PROFILE || METAL_SM5_PROFILE
+#if COMPILER_GLSL_ES2 || COMPILER_METAL
 
 	float determinant(float3x3 M)
 	{    
@@ -134,7 +134,19 @@
 	// @todo metal: remove this and make sure all uses handle METAL_MRT_PROFILE
 	#undef SM4_PROFILE
 	#define SM4_PROFILE 1
-	#define FCOLOR_COMPONENT_SWIZZLE .bgra  
+	#define FCOLOR_COMPONENT_SWIZZLE .bgra
+#elif METAL_SM4_PROFILE
+	#define FEATURE_LEVEL FEATURE_LEVEL_SM4
+	// @todo metal: remove this and make sure all uses handle METAL_SM4_PROFILE
+	#undef SM4_PROFILE
+	#define SM4_PROFILE 1
+	#define FCOLOR_COMPONENT_SWIZZLE .bgra
+#elif METAL_SM5_PROFILE
+	#define FEATURE_LEVEL FEATURE_LEVEL_SM5
+	// @todo metal: remove this and make sure all uses handle FEATURE_LEVEL_SM5
+	#undef SM5_PROFILE
+	#define SM5_PROFILE 1
+	#define FCOLOR_COMPONENT_SWIZZLE .bgra
 #elif ES2_PROFILE || ES3_1_PROFILE
 	#if ES3_1_PROFILE
 		#define FEATURE_LEVEL FEATURE_LEVEL_ES3_1
@@ -179,7 +191,7 @@
 // When enabled, different texture lookups can share samplers to allow more artist samplers in the base pass
 // Ideally this would just be enabled for all SM4 and above feature level platforms
 // @todo metal mrt: No reason this can't work with Metal, once cross compiler is fixed
-#define SUPPORTS_INDEPENDENT_SAMPLERS ((PS4_PROFILE || SM5_PROFILE || SM4_PROFILE) && !METAL_MRT_PROFILE)
+#define SUPPORTS_INDEPENDENT_SAMPLERS (PS4_PROFILE || SM5_PROFILE || SM4_PROFILE || METAL_MRT_PROFILE || METAL_SM4_PROFILE || METAL_SM5_PROFILE)
 
 // This would need to be a #define in GLSL to ignore the SamplerState, however, it is currently a function call in HLSL
 // for type checking of the parameters - ironically the type checking is really only needed in GLSL!
@@ -1260,7 +1272,7 @@ float DecodePackedHeight(float2 PackedHeight)
 
 uint ReverseBits32( uint bits )
 {
-#if SM5_PROFILE && !METAL_MRT_PROFILE && !METAL_PROFILE && !METAL_SM5_PROFILE
+#if SM5_PROFILE && !METAL_SM5_PROFILE
 	return reversebits( bits );
 #else
 	bits = ( bits << 16) | ( bits >> 16);
diff --git a/Engine/Shaders/DistanceFieldGlobalIllumination.usf b/Engine/Shaders/DistanceFieldGlobalIllumination.usf
--- a/Engine/Shaders/DistanceFieldGlobalIllumination.usf
+++ b/Engine/Shaders/DistanceFieldGlobalIllumination.usf
@@ -211,7 +211,7 @@ void VPLPlacementCS(
 
 				float3 DiffuseColor = .5f;
 				float CellExtent = VPLPlacementCameraRadius * InvPlacementGridSize.x;
-				float CellRadius = sqrt(2) * CellExtent;
+				float CellRadius = sqrt(2.0) * CellExtent;
 
 				float3 Flux = DiffuseColor * LightColor.rgb * max(dot(IntersectionNormal, -LightDirection), 0) * PI * CellRadius * CellRadius;
 
diff --git a/Engine/Shaders/DistanceFieldScreenGridLighting.usf b/Engine/Shaders/DistanceFieldScreenGridLighting.usf
--- a/Engine/Shaders/DistanceFieldScreenGridLighting.usf
+++ b/Engine/Shaders/DistanceFieldScreenGridLighting.usf
@@ -16,7 +16,7 @@
 
 #define SHARED_VISIBILITY_DIM (CONE_TRACE_DISPATCH_SIZEX * CONE_TRACE_DISPATCH_SIZEX * NUM_CONE_DIRECTIONS)
 groupshared uint SharedConeVisibility[SHARED_VISIBILITY_DIM];
-groupshared uint SharedConeRawVisibility[CONE_TRACE_DISPATCH_SIZEX * CONE_TRACE_DISPATCH_SIZEX * NUM_CONE_DIRECTIONS][NUM_VISIBILITY_STEPS];
+groupshared uint SharedConeRawVisibility[(CONE_TRACE_DISPATCH_SIZEX * CONE_TRACE_DISPATCH_SIZEX * NUM_CONE_DIRECTIONS) * NUM_VISIBILITY_STEPS];
 
 Buffer<float4> TileConeDepthRanges;
 float TanConeHalfAngle;
@@ -182,7 +182,7 @@ void HemisphereConeTraceAgainstTileCulledObjects(uint2 OutputCoordinate, uint Sh
 						UNROLL
 						for (uint i = 0; i < NUM_VISIBILITY_STEPS; i++)
 						{
-							InterlockedMin(SharedConeRawVisibility[SharedVisibilityBaseIndex + ConeIndex][i], asuint(MinRawVisibility[i]));
+							InterlockedMin(SharedConeRawVisibility[(SharedVisibilityBaseIndex + ConeIndex) * NUM_VISIBILITY_STEPS + i], asuint(MinRawVisibility[i]));
 						}
 					#endif
 				}
@@ -213,7 +213,7 @@ void ConeTraceObjectOcclusionCS(
 			UNROLL
 			for (uint i = 0; i < NUM_VISIBILITY_STEPS; i++)
 			{
-				SharedConeRawVisibility[Index][i] = FloatOneEncoded;
+				SharedConeRawVisibility[Index * NUM_VISIBILITY_STEPS + i] = FloatOneEncoded;
 			}
 		#endif
 	}
@@ -265,7 +265,7 @@ void ConeTraceObjectOcclusionCS(
 				UNROLL
 				for (uint StepIndex = 0; StepIndex < NUM_VISIBILITY_STEPS; StepIndex++)
 				{
-					float StepVisibility = asfloat(SharedConeRawVisibility[SharedVisibilityBaseIndex + ConeIndex][StepIndex]);
+					float StepVisibility = asfloat(SharedConeRawVisibility[(SharedVisibilityBaseIndex + ConeIndex) * NUM_VISIBILITY_STEPS + StepIndex]);
 					// Propagate min visibility down the cone
 					MinStepVisibility = min(MinStepVisibility, StepVisibility);
 					RWConeDepthVisibilityFunction[(OutputBaseIndex + ConeIndex) * NUM_VISIBILITY_STEPS + StepIndex] = MinStepVisibility;
@@ -299,7 +299,7 @@ void HemisphereConeTraceAgainstGlobalDistanceFieldClipmap(
 		{
 			float3 WorldSamplePosition = WorldShadingPosition + RotatedConeDirection * WorldStepOffset;
 			float3 StepVolumeUV = ComputeGlobalUV(WorldSamplePosition, ClipmapIndex);
-			float DistanceToOccluder = Texture3DSampleLevel(GlobalDistanceFieldTexture[ClipmapIndex], GlobalDistanceFieldSampler[ClipmapIndex], StepVolumeUV, 0).x;
+			float DistanceToOccluder = SampleGlobalDistanceField(ClipmapIndex, StepVolumeUV).x;
 			float SphereRadius = WorldStepOffset * TanConeHalfAngle;
 			//@todo - have to bias away from surface further for this to work
 			float ShadingSphereRadius = SphereRadius * 1.0f;
@@ -332,7 +332,7 @@ void HemisphereConeTraceAgainstGlobalDistanceField(uint OutputBaseIndex, float3
 	BRANCH
 	if (DistanceFromClipmap > AOGlobalMaxOcclusionDistance)
 	{
-		HemisphereConeTraceAgainstGlobalDistanceFieldClipmap(0, OutputBaseIndex, WorldShadingPosition, SceneDepth, WorldNormal, TangentX, TangentY);
+		HemisphereConeTraceAgainstGlobalDistanceFieldClipmap((uint)0, OutputBaseIndex, WorldShadingPosition, SceneDepth, WorldNormal, TangentX, TangentY);
 	}
 	else
 	{
@@ -341,7 +341,7 @@ void HemisphereConeTraceAgainstGlobalDistanceField(uint OutputBaseIndex, float3
 		BRANCH
 		if (DistanceFromClipmap > AOGlobalMaxOcclusionDistance)
 		{
-			HemisphereConeTraceAgainstGlobalDistanceFieldClipmap(1, OutputBaseIndex, WorldShadingPosition, SceneDepth, WorldNormal, TangentX, TangentY);
+			HemisphereConeTraceAgainstGlobalDistanceFieldClipmap((uint)1, OutputBaseIndex, WorldShadingPosition, SceneDepth, WorldNormal, TangentX, TangentY);
 		}
 		else
 		{
@@ -351,17 +351,17 @@ void HemisphereConeTraceAgainstGlobalDistanceField(uint OutputBaseIndex, float3
 			BRANCH
 			if (DistanceFromClipmap > AOGlobalMaxOcclusionDistance)
 			{
-				HemisphereConeTraceAgainstGlobalDistanceFieldClipmap(2, OutputBaseIndex, WorldShadingPosition, SceneDepth, WorldNormal, TangentX, TangentY);
+				HemisphereConeTraceAgainstGlobalDistanceFieldClipmap((uint)2, OutputBaseIndex, WorldShadingPosition, SceneDepth, WorldNormal, TangentX, TangentY);
 			}
 			else if (DistanceFromLastClipmap > AOGlobalMaxOcclusionDistance)
 			{
-				HemisphereConeTraceAgainstGlobalDistanceFieldClipmap(3, OutputBaseIndex, WorldShadingPosition, SceneDepth, WorldNormal, TangentX, TangentY);
+				HemisphereConeTraceAgainstGlobalDistanceFieldClipmap((uint)3, OutputBaseIndex, WorldShadingPosition, SceneDepth, WorldNormal, TangentX, TangentY);
 			}
 		}
 	}
 #else
 
-	HemisphereConeTraceAgainstGlobalDistanceFieldClipmap(0, SharedVisibilityBaseIndex, ThreadIndex, ThreadStride, WorldShadingPosition, SceneDepth, WorldNormal, TangentX, TangentY);
+	HemisphereConeTraceAgainstGlobalDistanceFieldClipmap((uint)0, SharedVisibilityBaseIndex, ThreadIndex, ThreadStride, WorldShadingPosition, SceneDepth, WorldNormal, TangentX, TangentY);
 
 #endif
 }
diff --git a/Engine/Shaders/DistanceFieldShadowing.usf b/Engine/Shaders/DistanceFieldShadowing.usf
--- a/Engine/Shaders/DistanceFieldShadowing.usf
+++ b/Engine/Shaders/DistanceFieldShadowing.usf
@@ -449,7 +449,7 @@ void CullObjectsToTileWithGather(
 				uint ListIndex;
 				InterlockedAdd(TileNumObjects0, 1U, ListIndex);
 				// Don't overwrite on overflow
-				ListIndex = min(ListIndex, MAX_INTERSECTING_OBJECTS - 1);
+				ListIndex = min(ListIndex, (uint)(MAX_INTERSECTING_OBJECTS - 1));
 				IntersectingObjectIndices[MAX_INTERSECTING_OBJECTS * 0 + ListIndex] = ObjectIndex; 
 			}
 
@@ -459,7 +459,7 @@ void CullObjectsToTileWithGather(
 				uint ListIndex;
 				InterlockedAdd(TileNumObjects1, 1U, ListIndex);
 				// Don't write out of bounds on overflow
-				ListIndex = min(ListIndex, MAX_INTERSECTING_OBJECTS - 1);
+				ListIndex = min(ListIndex, (uint)(MAX_INTERSECTING_OBJECTS - 1));
 				IntersectingObjectIndices[MAX_INTERSECTING_OBJECTS * 1 + ListIndex] = ObjectIndex; 
 			}
 		}
@@ -468,7 +468,7 @@ void CullObjectsToTileWithGather(
 	GroupMemoryBarrierWithGroupSync();
 
 	GroupIndex = SceneDepth > MaxTileZ2 ? 1 : 0;
-	NumIntersectingObjects = min(GroupIndex == 0 ? TileNumObjects0 : TileNumObjects1, MAX_INTERSECTING_OBJECTS);
+	NumIntersectingObjects = min(GroupIndex == 0 ? TileNumObjects0 : TileNumObjects1, (uint)MAX_INTERSECTING_OBJECTS);
 }
 
 uint DownsampleFactor;
diff --git a/Engine/Shaders/DistanceFieldSurfaceCacheLighting.usf b/Engine/Shaders/DistanceFieldSurfaceCacheLighting.usf
--- a/Engine/Shaders/DistanceFieldSurfaceCacheLighting.usf
+++ b/Engine/Shaders/DistanceFieldSurfaceCacheLighting.usf
@@ -373,7 +373,7 @@ void RayTraceThroughTileCulledDistanceFields(float3 WorldRayStart, float3 WorldR
 	TotalStepsTaken = 0;
 
 	LOOP
-	for (uint ListObjectIndex = 0; ListObjectIndex < min(NumIntersectingObjects, MAX_INTERSECTING_OBJECTS); ListObjectIndex++)
+	for (uint ListObjectIndex = 0; ListObjectIndex < min(NumIntersectingObjects, (uint)MAX_INTERSECTING_OBJECTS); ListObjectIndex++)
 	{
 		uint ObjectIndex = IntersectingObjectIndices[ListObjectIndex];
 		float4 SphereCenterAndRadius = LoadObjectPositionAndRadius(ObjectIndex);
@@ -434,7 +434,7 @@ void RayTraceThroughTileCulledDistanceFields(float3 WorldRayStart, float3 WorldR
 }
 
 void RayTraceThroughGlobalDistanceField(
-	uniform uint ClipmapIndex, 
+	uniform uint ClipmapIndex,
 	float3 WorldRayStart, 
 	float3 WorldRayEnd, 
 	float RayLength, 
@@ -472,7 +472,7 @@ void RayTraceThroughGlobalDistanceField(
 		{
 			float3 SampleVolumePosition = VolumeRayStart + VolumeRayDirection * SampleRayTime;
 			float3 VolumeUV = ComputeGlobalUV(SampleVolumePosition + GlobalVolumeCenter, ClipmapIndex);
-			float DistanceField = Texture3DSampleLevel(GlobalDistanceFieldTexture[ClipmapIndex], GlobalDistanceFieldSampler[ClipmapIndex], VolumeUV, 0).x;
+			float DistanceField = SampleGlobalDistanceField(ClipmapIndex, VolumeUV).x;
 			MinDistance = min(MinDistance, DistanceField);
 
 			float MinStepSize = GlobalVolumeExtent * 2 / 8000.0f;
@@ -527,26 +527,26 @@ void VisualizeMeshDistanceFieldCS(
 	float MaxRayTime0;
 	float IntersectRayTime;
 	float StepsTaken;
-	RayTraceThroughGlobalDistanceField(0, WorldRayStart, WorldRayEnd, TraceDistance, 0, MaxRayTime0, IntersectRayTime, StepsTaken);
+	RayTraceThroughGlobalDistanceField((uint)0, WorldRayStart, WorldRayEnd, TraceDistance, 0, MaxRayTime0, IntersectRayTime, StepsTaken);
 
 	TotalStepsTaken += StepsTaken;
 
 	if (IntersectRayTime >= TraceDistance)
 	{
 		float MaxRayTime1;
-		RayTraceThroughGlobalDistanceField(1, WorldRayStart, WorldRayEnd, TraceDistance, MaxRayTime0, MaxRayTime1, IntersectRayTime, StepsTaken);
+		RayTraceThroughGlobalDistanceField((uint)1, WorldRayStart, WorldRayEnd, TraceDistance, MaxRayTime0, MaxRayTime1, IntersectRayTime, StepsTaken);
 		TotalStepsTaken += StepsTaken;
 
 		if (IntersectRayTime >= TraceDistance)
 		{
 			float MaxRayTime2;
-			RayTraceThroughGlobalDistanceField(2, WorldRayStart, WorldRayEnd, TraceDistance, MaxRayTime1, MaxRayTime2, IntersectRayTime, StepsTaken);
+			RayTraceThroughGlobalDistanceField((uint)2, WorldRayStart, WorldRayEnd, TraceDistance, MaxRayTime1, MaxRayTime2, IntersectRayTime, StepsTaken);
 			TotalStepsTaken += StepsTaken;
 
 			if (IntersectRayTime >= TraceDistance)
 			{
 				float MaxRayTime3;
-				RayTraceThroughGlobalDistanceField(3, WorldRayStart, WorldRayEnd, TraceDistance, MaxRayTime2, MaxRayTime3, IntersectRayTime, StepsTaken);
+				RayTraceThroughGlobalDistanceField((uint)3, WorldRayStart, WorldRayEnd, TraceDistance, MaxRayTime2, MaxRayTime3, IntersectRayTime, StepsTaken);
 				TotalStepsTaken += StepsTaken;
 			}
 		}
@@ -620,4 +620,4 @@ void VisualizeDistanceFieldUpsamplePS(in float4 UVAndScreenPos : TEXCOORD0, out
 	float3 Value = Texture2DSampleLevel(VisualizeDistanceFieldTexture, VisualizeDistanceFieldSampler, DistanceFieldUVs, 0).xyz;
 
 	OutColor = float4(Value, 1);
-}
\ No newline at end of file
+}
diff --git a/Engine/Shaders/DistanceFieldSurfaceCacheLightingCompute.usf b/Engine/Shaders/DistanceFieldSurfaceCacheLightingCompute.usf
--- a/Engine/Shaders/DistanceFieldSurfaceCacheLightingCompute.usf
+++ b/Engine/Shaders/DistanceFieldSurfaceCacheLightingCompute.usf
@@ -741,7 +741,7 @@ RWBuffer<float4> RWDebugBuffer;
 #define THREADS_PER_OBJECT (FINAL_GATHER_THREADGROUP_SIZE / SIMULTANEOUSLY_TRACED_OBJECTS)
 
 groupshared uint SharedConeVisibility[NUM_CONE_DIRECTIONS];
-groupshared uint SharedConeRawVisibility[NUM_CONE_DIRECTIONS][NUM_VISIBILITY_STEPS];
+groupshared uint SharedConeRawVisibility[NUM_CONE_DIRECTIONS * NUM_VISIBILITY_STEPS];
 groupshared uint SharedMinOcclusionDistance;
 
 #define MAX_RECORD_CULLED_OBJECTS 128
@@ -768,7 +768,7 @@ void ConeTraceOcclusionCS(
 				UNROLL
 				for (uint i = 0; i < NUM_VISIBILITY_STEPS; i++)
 				{
-					SharedConeRawVisibility[ConeIndex][i] = asuint(1.0f);
+					SharedConeRawVisibility[(ConeIndex * NUM_VISIBILITY_STEPS) + i] = asuint(1.0f);
 				}
 			#endif
 		}
@@ -963,7 +963,7 @@ void ConeTraceOcclusionCS(
 		UNROLL
 		for (uint i = 0; i < NUM_VISIBILITY_STEPS; i++)
 		{
-			InterlockedMin(SharedConeRawVisibility[ConeIndex][i], asuint(MinRawVisibility[i]));
+			InterlockedMin(SharedConeRawVisibility[(ConeIndex * NUM_VISIBILITY_STEPS) + i], asuint(MinRawVisibility[i]));
 		}
 
 		GroupMemoryBarrierWithGroupSync();
@@ -978,7 +978,7 @@ void ConeTraceOcclusionCS(
 			UNROLL
 			for (uint i = 0; i < NUM_VISIBILITY_STEPS; i++)
 			{
-				float StepVisibility = asfloat(SharedConeRawVisibility[ConeIndex][i]);
+				float StepVisibility = asfloat(SharedConeRawVisibility[(ConeIndex * NUM_VISIBILITY_STEPS) + i]);
 				// Propagate min visibility down the cone
 				MinStepVisibility = min(MinStepVisibility, StepVisibility);
 				RWRecordConeData[RecordConeDataIndex + i] = MinStepVisibility;
@@ -1075,7 +1075,7 @@ void SetupCopyIndirectArgumentsCS(
 	uint NumRecords = (1 - TrimFraction) * DrawParameters[1];
 
 	// Make sure we spawn at least one group so RWScatterDrawParameters gets written to in the next pass
-	RWDispatchParameters[0] = max((NumRecords + COPY_THREADGROUP_SIZE - 1) / COPY_THREADGROUP_SIZE, 1);
+	RWDispatchParameters[0] = max((NumRecords + COPY_THREADGROUP_SIZE - 1) / COPY_THREADGROUP_SIZE, (uint)1);
 
 	RWDispatchParameters[1] = 1;
 	RWDispatchParameters[2] = 1;
diff --git a/Engine/Shaders/GlobalDistanceField.usf b/Engine/Shaders/GlobalDistanceField.usf
--- a/Engine/Shaders/GlobalDistanceField.usf
+++ b/Engine/Shaders/GlobalDistanceField.usf
@@ -121,7 +121,7 @@ uint3 CullGridDimension;
 RWBuffer<uint> RWCulledObjectGrid;
 
 #ifndef CULL_GRID_TILE_SIZE
-#define CULL_GRID_TILE_SIZE 0
+#define CULL_GRID_TILE_SIZE 1
 #endif
 
 #ifndef MAX_GRID_CULLED_DF_OBJECTS
@@ -223,7 +223,7 @@ void CompositeObjectDistanceFieldsCS(
 
 #if USE_CULL_GRID
 	uint GridBaseIndex = ComputeGridBaseIndex(DispatchThreadId / CULL_GRID_TILE_SIZE);
-	uint NumGridCulledObjects = min(CulledObjectGrid[GridBaseIndex], MAX_GRID_CULLED_DF_OBJECTS);
+	uint NumGridCulledObjects = min(CulledObjectGrid[GridBaseIndex], (uint)MAX_GRID_CULLED_DF_OBJECTS);
 	NumCulledObjects = NumGridCulledObjects;
 #endif
 
@@ -252,7 +252,7 @@ void CompositeObjectDistanceFieldsCS(
 
 	GroupMemoryBarrierWithGroupSync();
 
-	NumCulledObjects = min(NumTileCulledObjects, MAX_CULLED_DF_OBJECTS);
+	NumCulledObjects = min(NumTileCulledObjects, (uint)MAX_CULLED_DF_OBJECTS);
 
 #endif
 
diff --git a/Engine/Shaders/GlobalDistanceFieldShared.usf b/Engine/Shaders/GlobalDistanceFieldShared.usf
--- a/Engine/Shaders/GlobalDistanceFieldShared.usf
+++ b/Engine/Shaders/GlobalDistanceFieldShared.usf
@@ -14,13 +14,39 @@ float MaxGlobalDistance;
 // Must match C++
 #define MAX_GLOBAL_DF_CLIPMAPS 4
 
-Texture3D GlobalDistanceFieldTexture[MAX_GLOBAL_DF_CLIPMAPS];
-SamplerState GlobalDistanceFieldSampler[MAX_GLOBAL_DF_CLIPMAPS];
+Texture3D GlobalDistanceFieldTexture0;
+Texture3D GlobalDistanceFieldTexture1;
+Texture3D GlobalDistanceFieldTexture2;
+Texture3D GlobalDistanceFieldTexture3;
+SamplerState GlobalDistanceFieldSampler0;
+SamplerState GlobalDistanceFieldSampler1;
+SamplerState GlobalDistanceFieldSampler2;
+SamplerState GlobalDistanceFieldSampler3;
 float4 GlobalVolumeCenterAndExtent[MAX_GLOBAL_DF_CLIPMAPS];
 float4 GlobalVolumeWorldToUVAddAndMul[MAX_GLOBAL_DF_CLIPMAPS];
 float GlobalVolumeDimension;
 float GlobalVolumeTexelSize;
 
+float4 SampleGlobalDistanceField(int ClipmapIndex, float3 UV)
+{
+	if (ClipmapIndex == 0)
+	{
+		return Texture3DSampleLevel(GlobalDistanceFieldTexture0, GlobalDistanceFieldSampler0, UV, 0);
+	}
+	else if (ClipmapIndex == 1)
+	{
+		return Texture3DSampleLevel(GlobalDistanceFieldTexture1, GlobalDistanceFieldSampler1, UV, 0);
+	}
+	else if (ClipmapIndex == 2)
+	{
+		return Texture3DSampleLevel(GlobalDistanceFieldTexture2, GlobalDistanceFieldSampler2, UV, 0);
+	}
+	else
+	{
+		return Texture3DSampleLevel(GlobalDistanceFieldTexture3, GlobalDistanceFieldSampler3, UV, 0);
+	}
+}
+
 float3 ComputeGlobalUV(float3 WorldPosition, uint ClipmapIndex)
 {
 	//return ((WorldPosition - GlobalVolumeCenterAndExtent[ClipmapIndex].xyz + GlobalVolumeScollOffset[ClipmapIndex].xyz) / (GlobalVolumeCenterAndExtent[ClipmapIndex].w * 2) + .5f);
@@ -31,7 +57,23 @@ float3 ComputeGlobalUV(float3 WorldPosition, uint ClipmapIndex)
 float GetDistanceToNearestSurfaceGlobalClipmap(float3 WorldPosition, uint ClipmapIndex)
 {
 	float3 GlobalUV = ComputeGlobalUV(WorldPosition, ClipmapIndex);
-	float DistanceToSurface = Texture3DSampleLevel(GlobalDistanceFieldTexture[ClipmapIndex], GlobalDistanceFieldSampler[ClipmapIndex], GlobalUV, 0).x;
+	float DistanceToSurface = 0;
+	if (ClipmapIndex == 0)
+	{
+		DistanceToSurface = Texture3DSampleLevel(GlobalDistanceFieldTexture0, GlobalDistanceFieldSampler0, GlobalUV, 0).x;
+	}
+	else if (ClipmapIndex == 1)
+	{
+		DistanceToSurface = Texture3DSampleLevel(GlobalDistanceFieldTexture1, GlobalDistanceFieldSampler1, GlobalUV, 0).x;
+	}
+	else if (ClipmapIndex == 2)
+	{
+		DistanceToSurface = Texture3DSampleLevel(GlobalDistanceFieldTexture2, GlobalDistanceFieldSampler2, GlobalUV, 0).x;
+	}
+	else if (ClipmapIndex == 3)
+	{
+		DistanceToSurface = Texture3DSampleLevel(GlobalDistanceFieldTexture3, GlobalDistanceFieldSampler3, GlobalUV, 0).x;
+	}
 	return DistanceToSurface;
 }
 
@@ -78,12 +120,50 @@ float GetDistanceToNearestSurfaceGlobal(float3 WorldPosition)
 float3 GetDistanceFieldGradientGlobalClipmap(float3 WorldPosition, uint ClipmapIndex)
 {
 	float3 GlobalUV = ComputeGlobalUV(WorldPosition, ClipmapIndex);
-	float R = Texture3DSampleLevel(GlobalDistanceFieldTexture[ClipmapIndex], GlobalDistanceFieldSampler[ClipmapIndex], float3(GlobalUV.x + GlobalVolumeTexelSize, GlobalUV.y, GlobalUV.z), 0).x;
-	float L = Texture3DSampleLevel(GlobalDistanceFieldTexture[ClipmapIndex], GlobalDistanceFieldSampler[ClipmapIndex], float3(GlobalUV.x - GlobalVolumeTexelSize, GlobalUV.y, GlobalUV.z), 0).x;
-	float F = Texture3DSampleLevel(GlobalDistanceFieldTexture[ClipmapIndex], GlobalDistanceFieldSampler[ClipmapIndex], float3(GlobalUV.x, GlobalUV.y + GlobalVolumeTexelSize, GlobalUV.z), 0).x;
-	float B = Texture3DSampleLevel(GlobalDistanceFieldTexture[ClipmapIndex], GlobalDistanceFieldSampler[ClipmapIndex], float3(GlobalUV.x, GlobalUV.y - GlobalVolumeTexelSize, GlobalUV.z), 0).x;
-	float U = Texture3DSampleLevel(GlobalDistanceFieldTexture[ClipmapIndex], GlobalDistanceFieldSampler[ClipmapIndex], float3(GlobalUV.x, GlobalUV.y, GlobalUV.z + GlobalVolumeTexelSize), 0).x;
-	float D = Texture3DSampleLevel(GlobalDistanceFieldTexture[ClipmapIndex], GlobalDistanceFieldSampler[ClipmapIndex], float3(GlobalUV.x, GlobalUV.y, GlobalUV.z - GlobalVolumeTexelSize), 0).x;
+
+	float R = 0;
+	float L = 0;
+	float F = 0;
+	float B = 0;
+	float U = 0;
+	float D = 0;
+
+	if (ClipmapIndex == 0)
+	{
+		R = Texture3DSampleLevel(GlobalDistanceFieldTexture0, GlobalDistanceFieldSampler0, float3(GlobalUV.x + GlobalVolumeTexelSize, GlobalUV.y, GlobalUV.z), 0).x;
+		L = Texture3DSampleLevel(GlobalDistanceFieldTexture0, GlobalDistanceFieldSampler0, float3(GlobalUV.x - GlobalVolumeTexelSize, GlobalUV.y, GlobalUV.z), 0).x;
+		F = Texture3DSampleLevel(GlobalDistanceFieldTexture0, GlobalDistanceFieldSampler0, float3(GlobalUV.x, GlobalUV.y + GlobalVolumeTexelSize, GlobalUV.z), 0).x;
+		B = Texture3DSampleLevel(GlobalDistanceFieldTexture0, GlobalDistanceFieldSampler0, float3(GlobalUV.x, GlobalUV.y - GlobalVolumeTexelSize, GlobalUV.z), 0).x;
+		U = Texture3DSampleLevel(GlobalDistanceFieldTexture0, GlobalDistanceFieldSampler0, float3(GlobalUV.x, GlobalUV.y, GlobalUV.z + GlobalVolumeTexelSize), 0).x;
+		D = Texture3DSampleLevel(GlobalDistanceFieldTexture0, GlobalDistanceFieldSampler0, float3(GlobalUV.x, GlobalUV.y, GlobalUV.z - GlobalVolumeTexelSize), 0).x;
+	}
+	else if (ClipmapIndex == 1)
+	{
+		R = Texture3DSampleLevel(GlobalDistanceFieldTexture1, GlobalDistanceFieldSampler1, float3(GlobalUV.x + GlobalVolumeTexelSize, GlobalUV.y, GlobalUV.z), 0).x;
+		L = Texture3DSampleLevel(GlobalDistanceFieldTexture1, GlobalDistanceFieldSampler1, float3(GlobalUV.x - GlobalVolumeTexelSize, GlobalUV.y, GlobalUV.z), 0).x;
+		F = Texture3DSampleLevel(GlobalDistanceFieldTexture1, GlobalDistanceFieldSampler1, float3(GlobalUV.x, GlobalUV.y + GlobalVolumeTexelSize, GlobalUV.z), 0).x;
+		B = Texture3DSampleLevel(GlobalDistanceFieldTexture1, GlobalDistanceFieldSampler1, float3(GlobalUV.x, GlobalUV.y - GlobalVolumeTexelSize, GlobalUV.z), 0).x;
+		U = Texture3DSampleLevel(GlobalDistanceFieldTexture1, GlobalDistanceFieldSampler1, float3(GlobalUV.x, GlobalUV.y, GlobalUV.z + GlobalVolumeTexelSize), 0).x;
+		D = Texture3DSampleLevel(GlobalDistanceFieldTexture1, GlobalDistanceFieldSampler1, float3(GlobalUV.x, GlobalUV.y, GlobalUV.z - GlobalVolumeTexelSize), 0).x;
+	}
+	else if (ClipmapIndex == 2)
+	{
+		R = Texture3DSampleLevel(GlobalDistanceFieldTexture2, GlobalDistanceFieldSampler2, float3(GlobalUV.x + GlobalVolumeTexelSize, GlobalUV.y, GlobalUV.z), 0).x;
+		L = Texture3DSampleLevel(GlobalDistanceFieldTexture2, GlobalDistanceFieldSampler2, float3(GlobalUV.x - GlobalVolumeTexelSize, GlobalUV.y, GlobalUV.z), 0).x;
+		F = Texture3DSampleLevel(GlobalDistanceFieldTexture2, GlobalDistanceFieldSampler2, float3(GlobalUV.x, GlobalUV.y + GlobalVolumeTexelSize, GlobalUV.z), 0).x;
+		B = Texture3DSampleLevel(GlobalDistanceFieldTexture2, GlobalDistanceFieldSampler2, float3(GlobalUV.x, GlobalUV.y - GlobalVolumeTexelSize, GlobalUV.z), 0).x;
+		U = Texture3DSampleLevel(GlobalDistanceFieldTexture2, GlobalDistanceFieldSampler2, float3(GlobalUV.x, GlobalUV.y, GlobalUV.z + GlobalVolumeTexelSize), 0).x;
+		D = Texture3DSampleLevel(GlobalDistanceFieldTexture2, GlobalDistanceFieldSampler2, float3(GlobalUV.x, GlobalUV.y, GlobalUV.z - GlobalVolumeTexelSize), 0).x;
+	}
+	else if (ClipmapIndex == 3)
+	{
+		R = Texture3DSampleLevel(GlobalDistanceFieldTexture3, GlobalDistanceFieldSampler3, float3(GlobalUV.x + GlobalVolumeTexelSize, GlobalUV.y, GlobalUV.z), 0).x;
+		L = Texture3DSampleLevel(GlobalDistanceFieldTexture3, GlobalDistanceFieldSampler3, float3(GlobalUV.x - GlobalVolumeTexelSize, GlobalUV.y, GlobalUV.z), 0).x;
+		F = Texture3DSampleLevel(GlobalDistanceFieldTexture3, GlobalDistanceFieldSampler3, float3(GlobalUV.x, GlobalUV.y + GlobalVolumeTexelSize, GlobalUV.z), 0).x;
+		B = Texture3DSampleLevel(GlobalDistanceFieldTexture3, GlobalDistanceFieldSampler3, float3(GlobalUV.x, GlobalUV.y - GlobalVolumeTexelSize, GlobalUV.z), 0).x;
+		U = Texture3DSampleLevel(GlobalDistanceFieldTexture3, GlobalDistanceFieldSampler3, float3(GlobalUV.x, GlobalUV.y, GlobalUV.z + GlobalVolumeTexelSize), 0).x;
+		D = Texture3DSampleLevel(GlobalDistanceFieldTexture3, GlobalDistanceFieldSampler3, float3(GlobalUV.x, GlobalUV.y, GlobalUV.z - GlobalVolumeTexelSize), 0).x;
+	}
 		
 	float Extent = GlobalVolumeCenterAndExtent[ClipmapIndex].w;
 	float3 Gradient = .5f * float3(R - L, F - B, U - D) / Extent;
@@ -144,4 +224,4 @@ float3 GetDistanceFieldGradientGlobal(float3 WorldPosition)
 }
 
 #endif
-#endif // FEATURE_LEVEL >= FEATURE_LEVEL_SM5
\ No newline at end of file
+#endif // FEATURE_LEVEL >= FEATURE_LEVEL_SM5
diff --git a/Engine/Shaders/LocalVertexFactory.usf b/Engine/Shaders/LocalVertexFactory.usf
--- a/Engine/Shaders/LocalVertexFactory.usf
+++ b/Engine/Shaders/LocalVertexFactory.usf
@@ -375,7 +375,7 @@ FMaterialVertexParameters GetMaterialVertexParameters(FVertexFactoryInput Input,
 		YVec *= UseScale.y;
 
 		// Build overall transform
-#if MAC // @todo Remove once there is a fix for Intel Iris/HD5000 GPUs from rendering corrupt geometry without the normalize calls
+#if MAC && COMPILER_GLSL // @todo Remove once there is a fix for Intel Iris/HD5000 GPUs from rendering corrupt geometry without the normalize calls
 		float3x3 SliceTransform3 = mul(transpose(float3x3(normalize(SplineMeshDir), normalize(SplineMeshX), normalize(SplineMeshY))), float3x3(float3(0,0,0), XVec, YVec));
 #else
 		float3x3 SliceTransform3 = mul(transpose(float3x3(SplineMeshDir, SplineMeshX, SplineMeshY)), float3x3(float3(0,0,0), XVec, YVec));
@@ -415,7 +415,7 @@ FMaterialVertexParameters GetMaterialVertexParameters(FVertexFactoryInput Input,
 		YVec *= sign(UseScale.y);
 
 		// Build rotation transform
-#if MAC // @todo Remove once there is a fix for Intel Iris/HD5000 GPUs from rendering corrupt geometry without the normalize calls
+#if MAC && COMPILER_GLSL // @todo Remove once there is a fix for Intel Iris/HD5000 GPUs from rendering corrupt geometry without the normalize calls
 		half3x3 SliceTransform = mul(transpose(half3x3(normalize(SplineMeshDir), normalize(SplineMeshX), normalize(SplineMeshY))), half3x3(SplineDir, XVec, YVec));
 #else
 		half3x3 SliceTransform = mul(transpose(half3x3(SplineMeshDir, SplineMeshX, SplineMeshY)), half3x3(SplineDir, XVec, YVec));
diff --git a/Engine/Shaders/PostProcessSelectionOutline.usf b/Engine/Shaders/PostProcessSelectionOutline.usf
--- a/Engine/Shaders/PostProcessSelectionOutline.usf
+++ b/Engine/Shaders/PostProcessSelectionOutline.usf
@@ -61,7 +61,7 @@ FPixelInfo GetPixelInfo(int2 PixelPos, int SampleID, int OffsetX, int OffsetY, f
 	float DeviceZ = PostprocessInput1MS.Load(PixelPos, SampleID).r;
 	int Stencil = EditorPrimitivesStencil.Load(PixelPos, SampleID).g;
 
-#if !COMPILER_GLSL && !PS4_PROFILE
+#if !COMPILER_GLSL && !PS4_PROFILE && !COMPILER_METAL
 	// not yet supported on OpenGL, slightly better quality
 	ReconstructionOffset += PostprocessInput1MS.GetSamplePosition(SampleID);
 #endif
@@ -84,7 +84,7 @@ FPixelInfo GetPixelInfo(int2 PixelPos, int SampleID, int OffsetX, int OffsetY, f
 	// outline even between multiple selected objects (best usability)
 	Ret.ObjectId = Stencil;
 
-#if COMPILER_GLSL == 1 && !GL4_PROFILE
+#if (COMPILER_GLSL == 1 && !GL4_PROFILE) || COMPILER_METAL
 	// Stencil read in opengl is not supported on older versions
 	Ret.ObjectId = Ret.bObjectMask ? 2 : 0;
 #endif
diff --git a/Engine/Shaders/ShadowDepthPixelShader.usf b/Engine/Shaders/ShadowDepthPixelShader.usf
--- a/Engine/Shaders/ShadowDepthPixelShader.usf
+++ b/Engine/Shaders/ShadowDepthPixelShader.usf
@@ -21,8 +21,8 @@
 // to increase AO contribution. The AO calculation works best with closed geometry and small-medium scale objects
 #define LPV_GV_BOOST				1
 
-//@todo-rco: Remove this when reenabling for OpenGL
-#if !COMPILER_GLSL
+//@todo-rco: Remove this when reenabling for OpenGL & Metal
+#if !COMPILER_GLSL && !COMPILER_METAL
 #if REFLECTIVE_SHADOW_MAP
 	#include "LPVWriteVplCommon.usf"
 
@@ -87,8 +87,8 @@ void Main(
 	float3 normal = normalize( MaterialParameters.TangentToWorld[2] ); // Smooth normal
 	float3 flux = GetMaterialBounceColor(MaterialParameters);
 
-//@todo-rco: Remove this when reenabling for OpenGL
-#if !COMPILER_GLSL
+//@todo-rco: Remove this when reenabling for OpenGL & Metal
+#if !COMPILER_GLSL && !COMPILER_METAL
 #if MATERIAL_INJECT_EMISSIVE_INTO_LPV
 	float3 emissive = GetMaterialEmissive( MaterialParameters );
 
diff --git a/Engine/Shaders/ShadowProjectionCommon.usf b/Engine/Shaders/ShadowProjectionCommon.usf
--- a/Engine/Shaders/ShadowProjectionCommon.usf
+++ b/Engine/Shaders/ShadowProjectionCommon.usf
@@ -137,7 +137,7 @@ float CubemapHardwarePCF(float3 WorldPosition, float3 LightPosition, float Light
 		// Figure out which cube face we're sampling from
 		float3 AbsLightVector = abs(LightVector);
 		float MaxCoordinate = max(AbsLightVector.x, max(AbsLightVector.y, AbsLightVector.z));
-#if MAC // @todo Remove once there is a fix for the Apple Intel HD driver/GLSL compiler failing to correctly index into the uniform array using the dynamically calculated CubeFaceIndex variable. radr://19190281 "Intel HD 4000 fails to index uniform array with dynamic variable, causes point lights to render incorrectly in UE4"
+#if MAC && COMPILER_GLSL // @todo Remove once there is a fix for the Apple Intel HD driver/GLSL compiler failing to correctly index into the uniform array using the dynamically calculated CubeFaceIndex variable. radr://19190281 "Intel HD 4000 fails to index uniform array with dynamic variable, causes point lights to render incorrectly in UE4"
 		// Transform the position into shadow space
 		float4 ShadowPosition;
 		if (MaxCoordinate == AbsLightVector.x)
diff --git a/Engine/Shaders/TonemapCommon.usf b/Engine/Shaders/TonemapCommon.usf
--- a/Engine/Shaders/TonemapCommon.usf
+++ b/Engine/Shaders/TonemapCommon.usf
@@ -5,7 +5,7 @@
 =============================================================================*/
 
 // #todo-rco: Fix in the cross-compiler
-#if COMPILER_GLSL || COMPILER_GLSL_ES2 || METAL_PROFILE || VULKAN_PROFILE || METAL_MRT_PROFILE || METAL_SM5_PROFILE
+#if COMPILER_GLSL || COMPILER_GLSL_ES2 || COMPILER_METAL || VULKAN_PROFILE
 #define log10(x)  log(x) / log(10.0)
 #endif
 
diff --git a/Engine/Source/Developer/iOS/MetalShaderFormat/Private/MetalShaderCompiler.cpp b/Engine/Source/Developer/iOS/MetalShaderFormat/Private/MetalShaderCompiler.cpp
--- a/Engine/Source/Developer/iOS/MetalShaderFormat/Private/MetalShaderCompiler.cpp
+++ b/Engine/Source/Developer/iOS/MetalShaderFormat/Private/MetalShaderCompiler.cpp
@@ -480,6 +480,7 @@ void CompileShader_Metal(const FShaderCompilerInput& Input,FShaderCompilerOutput
 	}
 
 	AdditionalDefines.SetDefine(TEXT("row_major"), TEXT(""));
+	AdditionalDefines.SetDefine(TEXT("COMPILER_METAL"), 1);
 
 	static FName NAME_SF_METAL(TEXT("SF_METAL"));
 	static FName NAME_SF_METAL_MRT(TEXT("SF_METAL_MRT"));
diff --git a/Engine/Source/Runtime/Engine/Private/Atmosphere/Atmosphere.h b/Engine/Source/Runtime/Engine/Private/Atmosphere/Atmosphere.h
--- a/Engine/Source/Runtime/Engine/Private/Atmosphere/Atmosphere.h
+++ b/Engine/Source/Runtime/Engine/Private/Atmosphere/Atmosphere.h
@@ -150,6 +150,7 @@ public:
 						/*NumSamples=*/ 1,
 						/*Flags=*/ TexCreate_ShaderResource,
 						/*BulkData=*/ CreateInfo );
+					RHIBindDebugLabelName(TextureRHI, TEXT("E_Transmittance"));
 				}
 				break;
 			case E_Irradiance:
@@ -162,7 +163,8 @@ public:
 						/*NumMips=*/ 1,
 						/*NumSamples=*/ 1,
 						/*Flags=*/ TexCreate_ShaderResource,
-						/*BulkData=*/ CreateInfo );
+													/*BulkData=*/ CreateInfo );
+					RHIBindDebugLabelName(TextureRHI, TEXT("E_Irradiance"));
 				}
 				break;
 			case E_Inscatter:
@@ -174,7 +176,8 @@ public:
 						SizeX, SizeY, SizeZ, PF_FloatRGBA,
 						/*NumMips=*/ 1,
 						/*Flags=*/ TexCreate_ShaderResource,
-						/*BulkData=*/ CreateInfo );
+													/*BulkData=*/ CreateInfo );
+					RHIBindDebugLabelName(TextureRHI, TEXT("E_Inscatter"));
 				}
 				break;
 			}
diff --git a/Engine/Source/Runtime/Engine/Private/GPUSkinCache.h b/Engine/Source/Runtime/Engine/Private/GPUSkinCache.h
--- a/Engine/Source/Runtime/Engine/Private/GPUSkinCache.h
+++ b/Engine/Source/Runtime/Engine/Private/GPUSkinCache.h
@@ -19,7 +19,8 @@ class FGPUSkinCache : public FRenderResource
 public:
 	enum SkinCacheInitSettings
 	{
-		MaxBufferSize = 1024 * 1024 * 32,
+//#todo-rco: Make this an ini setting
+		MaxBufferSize = PLATFORM_MAC ? 16 * 1024 * 1024 : 1024 * 1024 * 32,
 		MaxUniformBufferBones = 256,
 		MaxCachedElements = 1024,
 		MaxCachedVertexBufferSRVs = 128,
diff --git a/Engine/Source/Runtime/Engine/Private/UnrealClient.cpp b/Engine/Source/Runtime/Engine/Private/UnrealClient.cpp
--- a/Engine/Source/Runtime/Engine/Private/UnrealClient.cpp
+++ b/Engine/Source/Runtime/Engine/Private/UnrealClient.cpp
@@ -1112,10 +1112,8 @@ const TArray<FColor>& FViewport::GetRawHitProxyData(FIntRect InRect)
 			FViewport*, Viewport, this,
 			{
 			// Set the hit proxy map's render target.
-			SetRenderTarget(RHICmdList, Viewport->HitProxyMap.GetRenderTargetTexture(), FTextureRHIRef());
-
 			// Clear the hit proxy map to white, which is overloaded to mean no hit proxy.
-			RHICmdList.Clear(true, FLinearColor::White, false, 0, false, 0, FIntRect());
+			SetRenderTarget(RHICmdList, Viewport->HitProxyMap.GetRenderTargetTexture(), FTextureRHIRef(), ESimpleRenderTargetMode::EClearColorExistingDepth);
 		});
 
 		// Let the viewport client draw its hit proxies.
@@ -1434,7 +1432,7 @@ void FViewport::FHitProxyMap::Init(uint32 NewSizeX,uint32 NewSizeY)
 
 	// Create a render target to store the hit proxy map.
 	{
-		FRHIResourceCreateInfo CreateInfo;
+		FRHIResourceCreateInfo CreateInfo(FClearValueBinding::White);
 		RHICreateTargetableShaderResource2D(SizeX,SizeY,PF_B8G8R8A8,1,TexCreate_None,TexCreate_RenderTargetable,false,CreateInfo,RenderTargetTextureRHI,HitProxyTexture);
 	}
 	{
diff --git a/Engine/Source/Runtime/Engine/Private/UserInterface/Canvas.cpp b/Engine/Source/Runtime/Engine/Private/UserInterface/Canvas.cpp
--- a/Engine/Source/Runtime/Engine/Private/UserInterface/Canvas.cpp
+++ b/Engine/Source/Runtime/Engine/Private/UserInterface/Canvas.cpp
@@ -654,9 +654,9 @@ void FCanvas::Flush_RenderThread(FRHICommandListImmediate& RHICmdList, bool bFor
 	const FTexture2DRHIRef& RenderTargetTexture = RenderTarget->GetRenderTargetTexture();
 
 	check(IsValidRef(RenderTargetTexture));
-
+	
 	// Set the RHI render target.
-	::SetRenderTarget(RHICmdList, RenderTargetTexture, FTextureRHIRef());
+	::SetRenderTarget(RHICmdList, RenderTargetTexture, FTexture2DRHIRef());
 	// disable depth test & writes
 	RHICmdList.SetDepthStencilState(TStaticDepthStencilState<false, CF_Always>::GetRHI());
 
diff --git a/Engine/Source/Runtime/Engine/Public/SceneManagement.h b/Engine/Source/Runtime/Engine/Public/SceneManagement.h
--- a/Engine/Source/Runtime/Engine/Public/SceneManagement.h
+++ b/Engine/Source/Runtime/Engine/Public/SceneManagement.h
@@ -624,7 +624,7 @@ public:
 inline bool DoesPlatformSupportDistanceFieldShadowing(EShaderPlatform Platform)
 {
 	// Hasn't been tested elsewhere yet
-	return Platform == SP_PCD3D_SM5 || Platform == SP_PS4;
+	return Platform == SP_PCD3D_SM5 || Platform == SP_PS4 || Platform == SP_METAL_SM5;
 }
 
 /** Represents a USkyLightComponent to the rendering thread. */
diff --git a/Engine/Source/Runtime/OpenGLDrv/Private/OpenGLTexture.cpp b/Engine/Source/Runtime/OpenGLDrv/Private/OpenGLTexture.cpp
--- a/Engine/Source/Runtime/OpenGLDrv/Private/OpenGLTexture.cpp
+++ b/Engine/Source/Runtime/OpenGLDrv/Private/OpenGLTexture.cpp
@@ -1680,7 +1680,7 @@ FShaderResourceViewRHIRef FOpenGLDynamicRHI::RHICreateShaderResourceView(FTextur
 		View = new FOpenGLShaderResourceView(this, Texture2D->Resource, Texture2D->Target, MipLevel, false);
 	}
 	
-	FShaderCache::LogSRV(View, Texture2DRHI, MipLevel, Texture2DRHI->GetNumMips(), Texture2DRHI->GetFormat());
+	FShaderCache::LogSRV(View, Texture2DRHI, MipLevel, 1, PF_Unknown);
 
 	return View;
 }
diff --git a/Engine/Source/Runtime/RHI/Private/RHI.cpp b/Engine/Source/Runtime/RHI/Private/RHI.cpp
--- a/Engine/Source/Runtime/RHI/Private/RHI.cpp
+++ b/Engine/Source/Runtime/RHI/Private/RHI.cpp
@@ -211,6 +211,7 @@ bool GRHISupportsParallelRHIExecute = false;
 bool GSupportsHDR32bppEncodeModeIntrinsic = false;
 bool GSupportsParallelOcclusionQueries = false;
 
+bool GRHISupportsMSAADepthSampleAccess = false;
 
 /** Whether we are profiling GPU hitches. */
 bool GTriggerGPUHitchProfile = false;
@@ -304,6 +305,7 @@ static FName NAME_SF_METAL_MRT(TEXT("SF_METAL_MRT"));
 static FName NAME_GLSL_310_ES_EXT(TEXT("GLSL_310_ES_EXT"));
 static FName NAME_SF_METAL_SM5(TEXT("SF_METAL_SM5"));
 static FName NAME_PC_VULKAN_ES2(TEXT("PC_VULKAN_ES2"));
+static FName NAME_SF_METAL_SM4(TEXT("SF_METAL_SM4"));
 
 FName LegacyShaderPlatformToShaderFormat(EShaderPlatform Platform)
 {
@@ -344,6 +346,8 @@ FName LegacyShaderPlatformToShaderFormat(EShaderPlatform Platform)
 		return NAME_SF_METAL;
 	case SP_METAL_MRT:
 		return NAME_SF_METAL_MRT;
+	case SP_METAL_SM4:
+		return NAME_SF_METAL_SM4;
 	case SP_METAL_SM5:
 		return NAME_SF_METAL_SM5;
 	case SP_OPENGL_ES31_EXT:
@@ -379,6 +383,7 @@ EShaderPlatform ShaderFormatToLegacyShaderPlatform(FName ShaderFormat)
 	if (ShaderFormat == NAME_GLSL_310_ES_EXT)	return SP_OPENGL_ES31_EXT;
 	if (ShaderFormat == NAME_SF_METAL_SM5)		return SP_METAL_SM5;
 	if (ShaderFormat == NAME_PC_VULKAN_ES2)		return SP_VULKAN_ES2;
+	if (ShaderFormat == NAME_SF_METAL_SM4)		return SP_METAL_SM4;
 	return SP_NumPlatforms;
 }
 
diff --git a/Engine/Source/Runtime/RHI/Public/RHI.h b/Engine/Source/Runtime/RHI/Public/RHI.h
--- a/Engine/Source/Runtime/RHI/Public/RHI.h
+++ b/Engine/Source/Runtime/RHI/Public/RHI.h
@@ -236,6 +236,9 @@ Requirements:
 ***/
 extern RHI_API bool GRHISupportsParallelRHIExecute;
 
+/** Whether or not the RHI can perform MSAA sample load. */
+extern RHI_API bool GRHISupportsMSAADepthSampleAccess;
+
 /** Called once per frame only from within an RHI. */
 extern RHI_API void RHIPrivateBeginFrame();
 
diff --git a/Engine/Source/Runtime/RHI/Public/RHIDefinitions.h b/Engine/Source/Runtime/RHI/Public/RHIDefinitions.h
--- a/Engine/Source/Runtime/RHI/Public/RHIDefinitions.h
+++ b/Engine/Source/Runtime/RHI/Public/RHIDefinitions.h
@@ -36,22 +36,22 @@ static_assert(SF_NumFrequencies <= (1 << SF_NumBits), "SF_NumFrequencies will no
 /** @warning: update *LegacyShaderPlatform* when the below changes */
 enum EShaderPlatform
 {
-	SP_PCD3D_SM5			= 0,
+	SP_PCD3D_SM5		= 0,
 	SP_OPENGL_SM4		= 1,
 	SP_PS4				= 2,
 	/** Used when running in Feature Level ES2 in OpenGL. */
 	SP_OPENGL_PCES2		= 3,
 	SP_XBOXONE			= 4,
-	SP_PCD3D_SM4			= 5,
+	SP_PCD3D_SM4		= 5,
 	SP_OPENGL_SM5		= 6,
 	/** Used when running in Feature Level ES2 in D3D11. */
-	SP_PCD3D_ES2			= 7,
+	SP_PCD3D_ES2		= 7,
 	SP_OPENGL_ES2		= 8,
 	SP_OPENGL_ES2_WEBGL = 9, 
 	SP_OPENGL_ES2_IOS	= 10,
-	SP_METAL				= 11,
+	SP_METAL			= 11,
 	SP_OPENGL_SM4_MAC	= 12,
-	SP_METAL_MRT			= 13,
+	SP_METAL_MRT		= 13,
 	SP_OPENGL_ES31_EXT	= 14,
 	/** Used when running in Feature Level ES3_1 in D3D11. */
 	SP_PCD3D_ES3_1		= 15,
@@ -59,8 +59,9 @@ enum EShaderPlatform
 	SP_OPENGL_PCES3_1	= 16,
 	SP_METAL_SM5		= 17,
 	SP_VULKAN_ES2		= 18,
+	SP_METAL_SM4		= 19,
 
-	SP_NumPlatforms		= 19,
+	SP_NumPlatforms		= 20,
 	SP_NumBits			= 5,
 };
 static_assert(SP_NumPlatforms <= (1 << SP_NumBits), "SP_NumPlatforms will not fit on SP_NumBits");
@@ -598,7 +599,7 @@ enum class ESimpleRenderTargetMode
 
 inline bool IsPCPlatform(const EShaderPlatform Platform)
 {
-	return Platform == SP_PCD3D_SM5 || Platform == SP_PCD3D_SM4 || Platform == SP_PCD3D_ES2 || Platform == SP_PCD3D_ES3_1 || Platform ==  SP_OPENGL_SM4 || Platform == SP_OPENGL_SM4_MAC || Platform == SP_OPENGL_SM5 || Platform == SP_OPENGL_PCES2 || Platform == SP_OPENGL_PCES3_1 || (PLATFORM_MAC && (Platform == SP_METAL_MRT || Platform == SP_METAL_SM5));
+	return Platform == SP_PCD3D_SM5 || Platform == SP_PCD3D_SM4 || Platform == SP_PCD3D_ES2 || Platform == SP_PCD3D_ES3_1 || Platform ==  SP_OPENGL_SM4 || Platform == SP_OPENGL_SM4_MAC || Platform == SP_OPENGL_SM5 || Platform == SP_OPENGL_PCES2 || Platform == SP_OPENGL_PCES3_1 || Platform == SP_METAL_SM4 || Platform == SP_METAL_SM5;
 }
 
 /** Whether the shader platform corresponds to the ES2 feature level. */
@@ -619,6 +620,11 @@ inline bool IsOpenGLPlatform(const EShaderPlatform Platform)
 		|| Platform == SP_OPENGL_ES2 || Platform == SP_OPENGL_ES2_WEBGL || Platform == SP_OPENGL_ES2_IOS || Platform == SP_OPENGL_ES31_EXT || Platform == SP_VULKAN_ES2;
 }
 
+inline bool IsMetalPlatform(const EShaderPlatform Platform)
+{
+	return Platform == SP_METAL || Platform == SP_METAL_MRT || Platform == SP_METAL_SM4 || Platform == SP_METAL_SM5;
+}
+
 inline bool IsConsolePlatform(const EShaderPlatform Platform)
 {
 	return Platform == SP_PS4 || Platform == SP_XBOXONE;
@@ -639,6 +645,7 @@ inline ERHIFeatureLevel::Type GetMaxSupportedFeatureLevel(EShaderPlatform InShad
 	case SP_OPENGL_SM4:
 	case SP_OPENGL_SM4_MAC:
 	case SP_METAL_MRT:
+	case SP_METAL_SM4:
 		return ERHIFeatureLevel::SM4;
 	case SP_PCD3D_ES2:
 	case SP_OPENGL_PCES2:
@@ -695,6 +702,8 @@ inline bool IsFeatureLevelSupported(EShaderPlatform InShaderPlatform, ERHIFeatur
 		return InFeatureLevel <= ERHIFeatureLevel::ES3_1;
 	case SP_METAL_MRT:
 		return InFeatureLevel <= ERHIFeatureLevel::SM4;
+	case SP_METAL_SM4:
+		return InFeatureLevel <= ERHIFeatureLevel::SM4;
 	case SP_OPENGL_ES31_EXT:
 		return InFeatureLevel <= ERHIFeatureLevel::SM5;
 	case SP_METAL_SM5:
@@ -732,7 +741,7 @@ inline bool RHISupportsInstancing(ERHIFeatureLevel::Type FeatureLevel)
 inline bool RHISupportsSeparateMSAAAndResolveTextures(const EShaderPlatform Platform)
 {
 	// Metal needs to handle MSAA and resolve textures internally (unless RHICreateTexture2D was changed to take an optional resolve target)
-	return Platform != SP_METAL && Platform != SP_METAL_MRT && Platform != SP_METAL_SM5;
+	return !IsMetalPlatform(Platform);
 }
 
 inline bool RHISupportsComputeShaders(const EShaderPlatform Platform)
@@ -741,14 +750,19 @@ inline bool RHISupportsComputeShaders(const EShaderPlatform Platform)
 	return IsFeatureLevelSupported(Platform, ERHIFeatureLevel::SM5);
 }
 
+inline bool RHISupportsPixelShaderUAVs(const EShaderPlatform Platform)
+{
+	return IsFeatureLevelSupported(Platform, ERHIFeatureLevel::SM5) && !IsMetalPlatform(Platform);
+}
+
 inline bool RHISupportsGeometryShaders(const EShaderPlatform Platform)
 {
-	return IsFeatureLevelSupported(Platform, ERHIFeatureLevel::SM4) && Platform != SP_METAL_MRT && Platform != SP_METAL_SM5;
+	return IsFeatureLevelSupported(Platform, ERHIFeatureLevel::SM4) && !IsMetalPlatform(Platform);
 }
 
 inline bool RHIHasTiledGPU(const EShaderPlatform Platform)
 {
-	return (PLATFORM_IOS && Platform == SP_METAL_MRT) || Platform == SP_METAL || Platform == SP_OPENGL_ES2_IOS || Platform == SP_OPENGL_ES2;
+	return (Platform == SP_METAL_MRT) || Platform == SP_METAL || Platform == SP_OPENGL_ES2_IOS || Platform == SP_OPENGL_ES2;
 }
 
 inline uint32 GetFeatureLevelMaxTextureSamplers(ERHIFeatureLevel::Type FeatureLevel)
diff --git a/Engine/Source/Runtime/Renderer/Private/CompositionLighting/PostProcessDeferredDecals.cpp b/Engine/Source/Runtime/Renderer/Private/CompositionLighting/PostProcessDeferredDecals.cpp
--- a/Engine/Source/Runtime/Renderer/Private/CompositionLighting/PostProcessDeferredDecals.cpp
+++ b/Engine/Source/Runtime/Renderer/Private/CompositionLighting/PostProcessDeferredDecals.cpp
@@ -300,9 +300,9 @@ void StencilDecalMask(FRHICommandList& RHICmdList, const FViewInfo& View, bool b
 {
 	SCOPED_DRAW_EVENT(RHICmdList, StencilDecalMask);
 	FSceneRenderTargets& SceneContext = FSceneRenderTargets::Get(RHICmdList);
+	SetRenderTarget(RHICmdList, NULL, SceneContext.GetSceneDepthSurface(), ESimpleRenderTargetMode::EUninitializedColorExistingDepth, FExclusiveDepthStencil::DepthRead_StencilWrite);
 	RHICmdList.SetRasterizerState(TStaticRasterizerState<FM_Solid, CM_None>::GetRHI());
 	RHICmdList.SetBlendState(TStaticBlendState<CW_NONE>::GetRHI());
-	SetRenderTarget(RHICmdList, NULL, SceneContext.GetSceneDepthSurface(), ESimpleRenderTargetMode::EUninitializedColorExistingDepth, FExclusiveDepthStencil::DepthRead_StencilWrite);
 	RHICmdList.SetViewport(View.ViewRect.Min.X, View.ViewRect.Min.Y, 0.0f, View.ViewRect.Max.X, View.ViewRect.Max.Y, 1.0f);
 
 	// Write 1 to highest bit of stencil to areas that should not receive decals
@@ -351,6 +351,8 @@ bool RenderPreStencil(FRenderingCompositePassContext& Context, const FMatrix& Co
 			return false;
 		}
 	}
+	
+	SCOPED_DRAW_EVENT(Context.RHICmdList, RenderPreStencil);
 
 	FDecalRendering::SetVertexShaderOnly(Context.RHICmdList, View, FrustumComponentToClip);
 
diff --git a/Engine/Source/Runtime/Renderer/Private/CompositionLighting/PostProcessLpvIndirect.cpp b/Engine/Source/Runtime/Renderer/Private/CompositionLighting/PostProcessLpvIndirect.cpp
--- a/Engine/Source/Runtime/Renderer/Private/CompositionLighting/PostProcessLpvIndirect.cpp
+++ b/Engine/Source/Runtime/Renderer/Private/CompositionLighting/PostProcessLpvIndirect.cpp
@@ -38,7 +38,7 @@ class FPostProcessLpvIndirectPS : public FGlobalShader
 	DECLARE_SHADER_TYPE(FPostProcessLpvIndirectPS, Global);
 
 	//@todo-rco: Remove this when reenabling for OpenGL
-	static bool ShouldCache( EShaderPlatform Platform )		{ return IsFeatureLevelSupported(Platform, ERHIFeatureLevel::SM5) && !IsOpenGLPlatform(Platform); }
+	static bool ShouldCache( EShaderPlatform Platform )		{ return IsFeatureLevelSupported(Platform, ERHIFeatureLevel::SM5) && !IsOpenGLPlatform(Platform) && !IsMetalPlatform(Platform); }
 
 	/** Default constructor. */
 	FPostProcessLpvIndirectPS() {}
diff --git a/Engine/Source/Runtime/Renderer/Private/DeferredShadingRenderer.cpp b/Engine/Source/Runtime/Renderer/Private/DeferredShadingRenderer.cpp
--- a/Engine/Source/Runtime/Renderer/Private/DeferredShadingRenderer.cpp
+++ b/Engine/Source/Runtime/Renderer/Private/DeferredShadingRenderer.cpp
@@ -1040,8 +1040,7 @@ void FDeferredShadingSceneRenderer::Render(FRHICommandListImmediate& RHICmdList)
 	if (bIsWireframe && FSceneRenderer::ShouldCompositeEditorPrimitives(Views[0]))
 	{
 		// In Editor we want wire frame view modes to be MSAA for better quality. Resolve will be done with EditorPrimitives
-		SetRenderTarget(RHICmdList, SceneContext.GetEditorPrimitivesColor(), SceneContext.GetEditorPrimitivesDepth());
-		RHICmdList.Clear(true, FLinearColor(0, 0, 0, 0), true, (float)ERHIZBuffer::FarPlane, false, 0, FIntRect());
+		SetRenderTarget(RHICmdList, SceneContext.GetEditorPrimitivesColor(), SceneContext.GetEditorPrimitivesDepth(), ESimpleRenderTargetMode::EClearColorAndDepth);
 	}
 	else if (!bIsGBufferCurrent)
 	{
diff --git a/Engine/Source/Runtime/Renderer/Private/DistanceFieldGlobalIllumination.cpp b/Engine/Source/Runtime/Renderer/Private/DistanceFieldGlobalIllumination.cpp
--- a/Engine/Source/Runtime/Renderer/Private/DistanceFieldGlobalIllumination.cpp
+++ b/Engine/Source/Runtime/Renderer/Private/DistanceFieldGlobalIllumination.cpp
@@ -1385,7 +1385,7 @@ public:
 
 	static bool ShouldCache(EShaderPlatform Platform)
 	{
-		return IsFeatureLevelSupported(Platform, ERHIFeatureLevel::SM5) && DoesPlatformSupportDistanceFieldAO(Platform);
+		return IsFeatureLevelSupported(Platform, ERHIFeatureLevel::SM5) && DoesPlatformSupportDistanceFieldGI(Platform);
 	}
 
 	static void ModifyCompilationEnvironment(EShaderPlatform Platform, FShaderCompilerEnvironment& OutEnvironment)
@@ -1482,7 +1482,7 @@ public:
 
 	static bool ShouldCache(EShaderPlatform Platform)
 	{
-		return IsFeatureLevelSupported(Platform, ERHIFeatureLevel::SM5) && DoesPlatformSupportDistanceFieldAO(Platform);
+		return IsFeatureLevelSupported(Platform, ERHIFeatureLevel::SM5) && DoesPlatformSupportDistanceFieldGI(Platform);
 	}
 
 	static void ModifyCompilationEnvironment(EShaderPlatform Platform, FShaderCompilerEnvironment& OutEnvironment)
diff --git a/Engine/Source/Runtime/Renderer/Private/DistanceFieldShadowing.cpp b/Engine/Source/Runtime/Renderer/Private/DistanceFieldShadowing.cpp
--- a/Engine/Source/Runtime/Renderer/Private/DistanceFieldShadowing.cpp
+++ b/Engine/Source/Runtime/Renderer/Private/DistanceFieldShadowing.cpp
@@ -306,7 +306,7 @@ public:
 
 	static bool ShouldCache(EShaderPlatform Platform)
 	{
-		return IsFeatureLevelSupported(Platform, ERHIFeatureLevel::SM5) && DoesPlatformSupportDistanceFieldShadowing(Platform);
+		return IsFeatureLevelSupported(Platform, ERHIFeatureLevel::SM5) && DoesPlatformSupportDistanceFieldShadowing(Platform) && RHISupportsPixelShaderUAVs(Platform);
 	}
 
 	/** Default constructor. */
diff --git a/Engine/Source/Runtime/Renderer/Private/DistanceFieldSurfaceCacheLighting.cpp b/Engine/Source/Runtime/Renderer/Private/DistanceFieldSurfaceCacheLighting.cpp
--- a/Engine/Source/Runtime/Renderer/Private/DistanceFieldSurfaceCacheLighting.cpp
+++ b/Engine/Source/Runtime/Renderer/Private/DistanceFieldSurfaceCacheLighting.cpp
@@ -632,7 +632,7 @@ public:
 
 	static bool ShouldCache(EShaderPlatform Platform)
 	{
-		return IsFeatureLevelSupported(Platform, ERHIFeatureLevel::SM5) && DoesPlatformSupportDistanceFieldAO(Platform);
+		return IsFeatureLevelSupported(Platform, ERHIFeatureLevel::SM5) && DoesPlatformSupportDistanceFieldAO(Platform) && RHISupportsPixelShaderUAVs(Platform);
 	}
 
 	static void ModifyCompilationEnvironment(EShaderPlatform Platform, FShaderCompilerEnvironment& OutEnvironment)
diff --git a/Engine/Source/Runtime/Renderer/Private/DistanceFieldSurfaceCacheLighting.h b/Engine/Source/Runtime/Renderer/Private/DistanceFieldSurfaceCacheLighting.h
--- a/Engine/Source/Runtime/Renderer/Private/DistanceFieldSurfaceCacheLighting.h
+++ b/Engine/Source/Runtime/Renderer/Private/DistanceFieldSurfaceCacheLighting.h
@@ -21,7 +21,7 @@ const float GDefaultDFAOMaxOcclusionDistance = 600.0f;
 
 inline bool DoesPlatformSupportDistanceFieldAO(EShaderPlatform Platform)
 {
-	return Platform == SP_PCD3D_SM5 || Platform == SP_PS4;
+	return Platform == SP_PCD3D_SM5 || Platform == SP_PS4 || Platform == SP_METAL_SM5;
 }
 
 extern FIntPoint GetBufferSizeForAO();
diff --git a/Engine/Source/Runtime/Renderer/Private/LightPropagationVolume.cpp b/Engine/Source/Runtime/Renderer/Private/LightPropagationVolume.cpp
--- a/Engine/Source/Runtime/Renderer/Private/LightPropagationVolume.cpp
+++ b/Engine/Source/Runtime/Renderer/Private/LightPropagationVolume.cpp
@@ -357,7 +357,7 @@ class FLpvClearCS : public FLpvWriteShaderCSBase
 
 public:
 	//@todo-rco: Remove this when reenabling for OpenGL
-	static bool ShouldCache( EShaderPlatform Platform )		{ return IsFeatureLevelSupported(Platform, ERHIFeatureLevel::SM5) && !IsOpenGLPlatform(Platform); }
+	static bool ShouldCache( EShaderPlatform Platform )		{ return IsFeatureLevelSupported(Platform, ERHIFeatureLevel::SM5) && !IsOpenGLPlatform(Platform) && !IsMetalPlatform(Platform); }
 
 	static void ModifyCompilationEnvironment( EShaderPlatform Platform, FShaderCompilerEnvironment& OutEnvironment )
 	{
@@ -382,7 +382,7 @@ class FLpvClearGeometryVolumeCS : public FLpvWriteShaderCSBase
 
 public:
 	//@todo-rco: Remove this when reenabling for OpenGL
-	static bool ShouldCache( EShaderPlatform Platform )		{ return IsFeatureLevelSupported(Platform, ERHIFeatureLevel::SM5) && !IsOpenGLPlatform(Platform); }
+	static bool ShouldCache( EShaderPlatform Platform )		{ return IsFeatureLevelSupported(Platform, ERHIFeatureLevel::SM5) && !IsOpenGLPlatform(Platform) && !IsMetalPlatform(Platform); }
 
 	static void ModifyCompilationEnvironment( EShaderPlatform Platform, FShaderCompilerEnvironment& OutEnvironment )
 	{
@@ -407,7 +407,7 @@ class FLpvClearListsCS : public FLpvWriteShaderCSBase
 
 public:
 	//@todo-rco: Remove this when reenabling for OpenGL
-	static bool ShouldCache( EShaderPlatform Platform )		{ return IsFeatureLevelSupported(Platform, ERHIFeatureLevel::SM5) && !IsOpenGLPlatform(Platform); }
+	static bool ShouldCache( EShaderPlatform Platform )		{ return IsFeatureLevelSupported(Platform, ERHIFeatureLevel::SM5) && !IsOpenGLPlatform(Platform) && !IsMetalPlatform(Platform); }
 
 	static void ModifyCompilationEnvironment( EShaderPlatform Platform, FShaderCompilerEnvironment& OutEnvironment )
 	{
@@ -431,7 +431,7 @@ class FLpvInject_GenerateVplListsCS : public FLpvWriteShaderCSBase
 
 public:
 	//@todo-rco: Remove this when reenabling for OpenGL
-	static bool ShouldCache( EShaderPlatform Platform )		{ return IsFeatureLevelSupported(Platform, ERHIFeatureLevel::SM5) && !IsOpenGLPlatform(Platform); }
+	static bool ShouldCache( EShaderPlatform Platform )		{ return IsFeatureLevelSupported(Platform, ERHIFeatureLevel::SM5) && !IsOpenGLPlatform(Platform) && !IsMetalPlatform(Platform); }
 
 	static void ModifyCompilationEnvironment( EShaderPlatform Platform, FShaderCompilerEnvironment& OutEnvironment )
 	{
@@ -505,7 +505,7 @@ class FLpvInject_AccumulateVplListsCS : public FLpvWriteShaderCSBase
 
 public:
 	//@todo-rco: Remove this when reenabling for OpenGL
-	static bool ShouldCache( EShaderPlatform Platform )		{ return IsFeatureLevelSupported(Platform, ERHIFeatureLevel::SM5) && !IsOpenGLPlatform(Platform); }
+	static bool ShouldCache( EShaderPlatform Platform )		{ return IsFeatureLevelSupported(Platform, ERHIFeatureLevel::SM5) && !IsOpenGLPlatform(Platform) && !IsMetalPlatform(Platform); }
 
 	static void ModifyCompilationEnvironment( EShaderPlatform Platform, FShaderCompilerEnvironment& OutEnvironment )
 	{
@@ -529,7 +529,7 @@ class FLpvDirectionalOcclusionCS : public FLpvWriteShaderCSBase
 
 public:
 	//@todo-rco: Remove this when reenabling for OpenGL
-	static bool ShouldCache( EShaderPlatform Platform )		{ return IsFeatureLevelSupported(Platform, ERHIFeatureLevel::SM5) && !IsOpenGLPlatform(Platform); }
+	static bool ShouldCache( EShaderPlatform Platform )		{ return IsFeatureLevelSupported(Platform, ERHIFeatureLevel::SM5) && !IsOpenGLPlatform(Platform) && !IsMetalPlatform(Platform); }
 
 	static void ModifyCompilationEnvironment( EShaderPlatform Platform, FShaderCompilerEnvironment& OutEnvironment )
 	{
@@ -569,7 +569,7 @@ class FLpvCopyAOVolumeCS : public FLpvWriteShaderCSBase
 
 public:
 	//@todo-rco: Remove this when reenabling for OpenGL
-	static bool ShouldCache( EShaderPlatform Platform )		{ return IsFeatureLevelSupported(Platform, ERHIFeatureLevel::SM5) && !IsOpenGLPlatform(Platform); }
+	static bool ShouldCache( EShaderPlatform Platform )		{ return IsFeatureLevelSupported(Platform, ERHIFeatureLevel::SM5) && !IsOpenGLPlatform(Platform) && !IsMetalPlatform(Platform); }
 
 	static void ModifyCompilationEnvironment( EShaderPlatform Platform, FShaderCompilerEnvironment& OutEnvironment )
 	{
@@ -610,7 +610,7 @@ class FLpvBuildGeometryVolumeCS : public FLpvWriteShaderCSBase
 
 public:
 	//@todo-rco: Remove this when reenabling for OpenGL
-	static bool ShouldCache( EShaderPlatform Platform )		{ return IsFeatureLevelSupported(Platform, ERHIFeatureLevel::SM5) && !IsOpenGLPlatform(Platform); }
+	static bool ShouldCache( EShaderPlatform Platform )		{ return IsFeatureLevelSupported(Platform, ERHIFeatureLevel::SM5) && !IsOpenGLPlatform(Platform) && !IsMetalPlatform(Platform); }
 
 	static void ModifyCompilationEnvironment( EShaderPlatform Platform, FShaderCompilerEnvironment& OutEnvironment )
 	{
@@ -642,7 +642,7 @@ class TLpvPropagateCS : public FLpvWriteShaderCSBase
 
 public:
 	//@todo-rco: Remove this when reenabling for OpenGL
-	static bool ShouldCache( EShaderPlatform Platform )		{ return IsFeatureLevelSupported(Platform, ERHIFeatureLevel::SM5) && !IsOpenGLPlatform(Platform); }
+	static bool ShouldCache( EShaderPlatform Platform )		{ return IsFeatureLevelSupported(Platform, ERHIFeatureLevel::SM5) && !IsOpenGLPlatform(Platform) && !IsMetalPlatform(Platform); }
 
 	static void ModifyCompilationEnvironment( EShaderPlatform Platform, FShaderCompilerEnvironment& OutEnvironment )
 	{
@@ -734,7 +734,7 @@ class TLpvInject_LightCS : public FLpvInjectShader_Base
 
 public:
 	//@todo-rco: Remove this when reenabling for OpenGL
-	static bool ShouldCache( EShaderPlatform Platform )		{ return IsFeatureLevelSupported(Platform, ERHIFeatureLevel::SM5) && !IsOpenGLPlatform(Platform); }
+	static bool ShouldCache( EShaderPlatform Platform )		{ return IsFeatureLevelSupported(Platform, ERHIFeatureLevel::SM5) && !IsOpenGLPlatform(Platform) && !IsMetalPlatform(Platform); }
 
 	static void ModifyCompilationEnvironment( EShaderPlatform Platform, FShaderCompilerEnvironment& OutEnvironment )
 	{
diff --git a/Engine/Source/Runtime/Renderer/Private/LightPropagationVolumeVisualisation.cpp b/Engine/Source/Runtime/Renderer/Private/LightPropagationVolumeVisualisation.cpp
--- a/Engine/Source/Runtime/Renderer/Private/LightPropagationVolumeVisualisation.cpp
+++ b/Engine/Source/Runtime/Renderer/Private/LightPropagationVolumeVisualisation.cpp
@@ -78,7 +78,7 @@ public:
 	virtual bool Serialize( FArchive& Ar ) override																					{ return FLpvVisualiseBase::Serialize( Ar ); }
 
 	//@todo-rco: Remove this when reenabling for OpenGL
-	static bool ShouldCache( EShaderPlatform Platform )		{ return IsFeatureLevelSupported(Platform, ERHIFeatureLevel::SM5) && !IsOpenGLPlatform(Platform); }
+	static bool ShouldCache( EShaderPlatform Platform )		{ return IsFeatureLevelSupported(Platform, ERHIFeatureLevel::SM5) && !IsOpenGLPlatform(Platform) && !IsMetalPlatform(Platform); }
 	static void ModifyCompilationEnvironment( EShaderPlatform Platform, FShaderCompilerEnvironment& OutEnvironment )				{ FLpvVisualiseBase::ModifyCompilationEnvironment( Platform, OutEnvironment ); }
 
 	void SetParameters(
@@ -111,7 +111,7 @@ public:
 	}
 
 	//@todo-rco: Remove this when reenabling for OpenGL
-	static bool ShouldCache( EShaderPlatform Platform )		{ return IsFeatureLevelSupported(Platform, ERHIFeatureLevel::SM5) && !IsOpenGLPlatform(Platform); }
+	static bool ShouldCache( EShaderPlatform Platform )		{ return IsFeatureLevelSupported(Platform, ERHIFeatureLevel::SM5) && !IsOpenGLPlatform(Platform) && !IsMetalPlatform(Platform); }
 	static void ModifyCompilationEnvironment( EShaderPlatform Platform, FShaderCompilerEnvironment& OutEnvironment )				{ FLpvVisualiseBase::ModifyCompilationEnvironment( Platform, OutEnvironment ); }
 
 	void SetParameters(
diff --git a/Engine/Source/Runtime/Renderer/Private/LightRendering.cpp b/Engine/Source/Runtime/Renderer/Private/LightRendering.cpp
--- a/Engine/Source/Runtime/Renderer/Private/LightRendering.cpp
+++ b/Engine/Source/Runtime/Renderer/Private/LightRendering.cpp
@@ -950,7 +950,7 @@ void FDeferredShadingSceneRenderer::RenderLight(FRHICommandList& RHICmdList, con
 			VertexShader->SetParameters(RHICmdList, View, LightSceneInfo);
 
 			// NUse DBT to allow work culling on shadow lights
-			if (bAllowDepthBoundsTest != 0)
+			if (GSupportsDepthBoundsTest && bAllowDepthBoundsTest != 0)
 			{
 				// Can use the depth bounds test to skip work for pixels which won't be touched by the light (i.e outside the depth range)
 				float NearDepth = 1.f;
@@ -979,7 +979,7 @@ void FDeferredShadingSceneRenderer::RenderLight(FRHICommandList& RHICmdList, con
 			}
 
 			// Use DBT to allow work culling on shadow lights
-			if (bAllowDepthBoundsTest != 0)
+			if (GSupportsDepthBoundsTest && bAllowDepthBoundsTest != 0)
 			{
 				// Turn DBT back off
 				RHICmdList.EnableDepthBoundsTest(false, 0, 1);
diff --git a/Engine/Source/Runtime/Renderer/Private/LightShaftRendering.cpp b/Engine/Source/Runtime/Renderer/Private/LightShaftRendering.cpp
--- a/Engine/Source/Runtime/Renderer/Private/LightShaftRendering.cpp
+++ b/Engine/Source/Runtime/Renderer/Private/LightShaftRendering.cpp
@@ -448,11 +448,10 @@ void AllocateOrReuseLightShaftRenderTarget(FRHICommandListImmediate& RHICmdList,
 		EPixelFormat LightShaftFilterBufferFormat = PF_FloatRGB;
 		const FIntPoint BufferSize = FSceneRenderTargets::Get(RHICmdList).GetBufferSizeXY();
 		FIntPoint LightShaftSize(FMath::Max<uint32>(BufferSize.X / GetLightShaftDownsampleFactor(), 1), FMath::Max<uint32>(BufferSize.Y / GetLightShaftDownsampleFactor(), 1));
-		FPooledRenderTargetDesc Desc(FPooledRenderTargetDesc::Create2DDesc(LightShaftSize, LightShaftFilterBufferFormat, FClearValueBinding::None, TexCreate_None, TexCreate_RenderTargetable, false));
+		FPooledRenderTargetDesc Desc(FPooledRenderTargetDesc::Create2DDesc(LightShaftSize, LightShaftFilterBufferFormat, FClearValueBinding::Black, TexCreate_None, TexCreate_RenderTargetable, false));
 		GRenderTargetPool.FindFreeElement(Desc, Target, Name);
 
-		SetRenderTarget(RHICmdList, Target->GetRenderTargetItem().TargetableTexture, FTextureRHIRef());
-		RHICmdList.Clear(true, FLinearColor(0, 0, 0, 0), false, 1.0f, false, 0, FIntRect());
+		SetRenderTarget(RHICmdList, Target->GetRenderTargetItem().TargetableTexture, FTextureRHIRef(), ESimpleRenderTargetMode::EClearColorExistingDepth);
 	}
 }
 
diff --git a/Engine/Source/Runtime/Renderer/Private/PostProcess/PostProcessCombineLUTs.cpp b/Engine/Source/Runtime/Renderer/Private/PostProcess/PostProcessCombineLUTs.cpp
--- a/Engine/Source/Runtime/Renderer/Private/PostProcess/PostProcessCombineLUTs.cpp
+++ b/Engine/Source/Runtime/Renderer/Private/PostProcess/PostProcessCombineLUTs.cpp
@@ -670,7 +670,7 @@ void FRCPassPostProcessCombineLUTs::Process(FRenderingCompositePassContext& Cont
 	const FSceneRenderTargetItem& DestRenderTarget = PassOutputs[0].RequestSurface(Context);
 
 	// Set the view family's render target/viewport.
-	SetRenderTarget(Context.RHICmdList, DestRenderTarget.TargetableTexture, FTextureRHIRef());
+	SetRenderTarget(Context.RHICmdList, DestRenderTarget.TargetableTexture, FTextureRHIRef(), ESimpleRenderTargetMode::EUninitializedColorAndDepth);
 	Context.SetViewportAndCallRHI(0, 0, 0.0f, DestSize.X, DestSize.Y, 1.0f );
 
 	// set the state
@@ -720,4 +720,4 @@ FPooledRenderTargetDesc FRCPassPostProcessCombineLUTs::ComputeOutputDesc(EPassOu
 	Ret.DebugName = TEXT("CombineLUTs");
 
 	return Ret;
-}
\ No newline at end of file
+}
diff --git a/Engine/Source/Runtime/Renderer/Private/PostProcess/PostProcessCompositeEditorPrimitives.cpp b/Engine/Source/Runtime/Renderer/Private/PostProcess/PostProcessCompositeEditorPrimitives.cpp
--- a/Engine/Source/Runtime/Renderer/Private/PostProcess/PostProcessCompositeEditorPrimitives.cpp
+++ b/Engine/Source/Runtime/Renderer/Private/PostProcess/PostProcessCompositeEditorPrimitives.cpp
@@ -210,7 +210,7 @@ void FRCPassPostProcessCompositeEditorPrimitives::Process(FRenderingCompositePas
 	const uint32 MSAASampleCount = SceneContext.EditorPrimitivesColor->GetDesc().NumSamples;
 
 	{
-		SetRenderTarget(Context.RHICmdList, ColorTarget, DepthTarget);
+		SetRenderTarget(Context.RHICmdList, ColorTarget, DepthTarget, ESimpleRenderTargetMode::EExistingColorAndDepth);
 		Context.SetViewportAndCallRHI(DestRect);
 
 		if (bClearIsNeeded)
diff --git a/Engine/Source/Runtime/Renderer/Private/PostProcess/PostProcessDownsample.cpp b/Engine/Source/Runtime/Renderer/Private/PostProcess/PostProcessDownsample.cpp
--- a/Engine/Source/Runtime/Renderer/Private/PostProcess/PostProcessDownsample.cpp
+++ b/Engine/Source/Runtime/Renderer/Private/PostProcess/PostProcessDownsample.cpp
@@ -176,7 +176,7 @@ void FRCPassPostProcessDownsample::Process(FRenderingCompositePassContext& Conte
 	const FSceneRenderTargetItem& DestRenderTarget = PassOutputs[0].RequestSurface(Context);
 
 	// Set the view family's render target/viewport.
-	SetRenderTarget(Context.RHICmdList, DestRenderTarget.TargetableTexture, FTextureRHIRef());
+	SetRenderTarget(Context.RHICmdList, DestRenderTarget.TargetableTexture, FTextureRHIRef(), ESimpleRenderTargetMode::EExistingColorAndDepth);
 
 	Context.SetViewportAndCallRHI(0, 0, 0.0f, DestSize.X, DestSize.Y, 1.0f );
 
diff --git a/Engine/Source/Runtime/Renderer/Private/PostProcess/PostProcessTonemap.cpp b/Engine/Source/Runtime/Renderer/Private/PostProcess/PostProcessTonemap.cpp
--- a/Engine/Source/Runtime/Renderer/Private/PostProcess/PostProcessTonemap.cpp
+++ b/Engine/Source/Runtime/Renderer/Private/PostProcess/PostProcessTonemap.cpp
@@ -1097,7 +1097,7 @@ void FRCPassPostProcessTonemap::Process(FRenderingCompositePassContext& Context)
 
 	
 	// Set the view family's render target/viewport.
-	SetRenderTarget(Context.RHICmdList, DestRenderTarget.TargetableTexture, FTextureRHIParamRef());
+	SetRenderTarget(Context.RHICmdList, DestRenderTarget.TargetableTexture, FTextureRHIParamRef(), ESimpleRenderTargetMode::EUninitializedColorAndDepth);
 
 	if( ViewFamily.RenderTarget->GetRenderTargetTexture() != DestRenderTarget.TargetableTexture )
 	{
diff --git a/Engine/Source/Runtime/Renderer/Private/PostProcess/PostProcessWeightedSampleSum.cpp b/Engine/Source/Runtime/Renderer/Private/PostProcess/PostProcessWeightedSampleSum.cpp
--- a/Engine/Source/Runtime/Renderer/Private/PostProcess/PostProcessWeightedSampleSum.cpp
+++ b/Engine/Source/Runtime/Renderer/Private/PostProcess/PostProcessWeightedSampleSum.cpp
@@ -553,7 +553,7 @@ void FRCPassPostProcessWeightedSampleSum::Process(FRenderingCompositePassContext
 		BlurWeights[i] = TintValue * OffsetAndWeight[i].Y;
 	}
 
-	SetRenderTarget(Context.RHICmdList, DestRenderTarget.TargetableTexture, FTextureRHIRef());
+	SetRenderTarget(Context.RHICmdList, DestRenderTarget.TargetableTexture, FTextureRHIRef(), ESimpleRenderTargetMode::EExistingColorAndDepth);
 
 	Context.SetViewportAndCallRHI(0, 0, 0.0f, DestSize.X, DestSize.Y, 1.0f);
 
diff --git a/Engine/Source/Runtime/Renderer/Private/PostProcess/SceneRenderTargets.cpp b/Engine/Source/Runtime/Renderer/Private/PostProcess/SceneRenderTargets.cpp
--- a/Engine/Source/Runtime/Renderer/Private/PostProcess/SceneRenderTargets.cpp
+++ b/Engine/Source/Runtime/Renderer/Private/PostProcess/SceneRenderTargets.cpp
@@ -1318,7 +1318,7 @@ int32 FSceneRenderTargets::GetEditorMSAACompositingSampleCount() const
 	int32 Value = 1;
 
 	// only supported on SM5 yet (SM4 doesn't have MSAA sample load functionality which makes it harder to implement)
-	if (CurrentFeatureLevel >= ERHIFeatureLevel::SM5)
+	if (CurrentFeatureLevel >= ERHIFeatureLevel::SM5 && GRHISupportsMSAADepthSampleAccess)
 	{
 		static const auto CVar = IConsoleManager::Get().FindTConsoleVariableDataInt(TEXT("r.MSAA.CompositingSampleCount"));
 
diff --git a/Engine/Source/Runtime/Renderer/Private/SceneOcclusion.cpp b/Engine/Source/Runtime/Renderer/Private/SceneOcclusion.cpp
--- a/Engine/Source/Runtime/Renderer/Private/SceneOcclusion.cpp
+++ b/Engine/Source/Runtime/Renderer/Private/SceneOcclusion.cpp
@@ -993,13 +993,12 @@ void BuildHZB( FRHICommandListImmediate& RHICmdList, FViewInfo& View )
 	
 	FSceneRenderTargetItem& HZBRenderTarget = View.HZB->GetRenderTargetItem();
 	
-	RHICmdList.SetBlendState(TStaticBlendState<>::GetRHI());
-	RHICmdList.SetRasterizerState(TStaticRasterizerState<>::GetRHI());
-	RHICmdList.SetDepthStencilState(TStaticDepthStencilState< false, CF_Always >::GetRHI());
-	
 	// Mip 0
 	{
 		SetRenderTarget(RHICmdList, HZBRenderTarget.TargetableTexture, 0, NULL);
+		RHICmdList.SetBlendState(TStaticBlendState<>::GetRHI());
+		RHICmdList.SetRasterizerState(TStaticRasterizerState<>::GetRHI());
+		RHICmdList.SetDepthStencilState(TStaticDepthStencilState< false, CF_Always >::GetRHI());
 
 		TShaderMapRef< FPostProcessVS >	VertexShader(View.ShaderMap);
 		TShaderMapRef< THZBBuildPS<0> >	PixelShader(View.ShaderMap);
@@ -1037,6 +1036,9 @@ void BuildHZB( FRHICommandListImmediate& RHICmdList, FViewInfo& View )
 		DstSize.Y = FMath::Max(DstSize.Y, 1);
 
 		SetRenderTarget(RHICmdList, HZBRenderTarget.TargetableTexture, MipIndex, NULL);
+		RHICmdList.SetBlendState(TStaticBlendState<>::GetRHI());
+		RHICmdList.SetRasterizerState(TStaticRasterizerState<>::GetRHI());
+		RHICmdList.SetDepthStencilState(TStaticDepthStencilState< false, CF_Always >::GetRHI());
 
 		TShaderMapRef< FPostProcessVS >	VertexShader(View.ShaderMap);
 		TShaderMapRef< THZBBuildPS<1> >	PixelShader(View.ShaderMap);
@@ -1178,7 +1180,7 @@ void FDeferredShadingSceneRenderer::BeginOcclusionTests(FRHICommandListImmediate
 
 									// Draw bounding sphere
 									VertexShader->SetParametersWithBoundingSphere(RHICmdList, View, LightBounds);
-									StencilingGeometry::DrawSphere(RHICmdList);
+									StencilingGeometry::DrawVectorSphere(RHICmdList);
 
 									RHICmdList.EndRenderQuery(ShadowOcclusionQuery);
 								}
diff --git a/Engine/Source/Runtime/Renderer/Private/ShadowRendering.cpp b/Engine/Source/Runtime/Renderer/Private/ShadowRendering.cpp
--- a/Engine/Source/Runtime/Renderer/Private/ShadowRendering.cpp
+++ b/Engine/Source/Runtime/Renderer/Private/ShadowRendering.cpp
@@ -595,12 +595,13 @@ IMPLEMENT_SHADOWDEPTHPASS_PIXELSHADER_TYPE(PixelShadowDepth_OnePassPointLight, t
 IMPLEMENT_SHADOWDEPTHPASS_PIXELSHADER_TYPE(PixelShadowDepth_OnePassPointLight, false);
 
 /** The stencil sphere vertex buffer. */
-TGlobalResource<StencilingGeometry::TStencilSphereVertexBuffer<18, 12> > StencilingGeometry::GStencilSphereVertexBuffer;
+TGlobalResource<StencilingGeometry::TStencilSphereVertexBuffer<18, 12, FVector4> > StencilingGeometry::GStencilSphereVertexBuffer;
+TGlobalResource<StencilingGeometry::TStencilSphereVertexBuffer<18, 12, FVector> > StencilingGeometry::GStencilSphereVectorBuffer;
 
 /** The stencil sphere index buffer. */
 TGlobalResource<StencilingGeometry::TStencilSphereIndexBuffer<18, 12> > StencilingGeometry::GStencilSphereIndexBuffer;
 
-TGlobalResource<StencilingGeometry::TStencilSphereVertexBuffer<4, 4> > StencilingGeometry::GLowPolyStencilSphereVertexBuffer;
+TGlobalResource<StencilingGeometry::TStencilSphereVertexBuffer<4, 4, FVector4> > StencilingGeometry::GLowPolyStencilSphereVertexBuffer;
 TGlobalResource<StencilingGeometry::TStencilSphereIndexBuffer<4, 4> > StencilingGeometry::GLowPolyStencilSphereIndexBuffer;
 
 /** The (dummy) stencil cone vertex buffer. */
@@ -1876,6 +1877,14 @@ void StencilingGeometry::DrawSphere(FRHICommandList& RHICmdList)
 		StencilingGeometry::GStencilSphereVertexBuffer.GetVertexCount(), 0, 
 		StencilingGeometry::GStencilSphereIndexBuffer.GetIndexCount() / 3, 1);
 }
+		
+void StencilingGeometry::DrawVectorSphere(FRHICommandList& RHICmdList)
+{
+	RHICmdList.SetStreamSource(0, StencilingGeometry::GStencilSphereVectorBuffer.VertexBufferRHI, sizeof(FVector), 0);
+	RHICmdList.DrawIndexedPrimitive(StencilingGeometry::GStencilSphereIndexBuffer.IndexBufferRHI, PT_TriangleList, 0, 0,
+									StencilingGeometry::GStencilSphereVectorBuffer.GetVertexCount(), 0,
+									StencilingGeometry::GStencilSphereIndexBuffer.GetIndexCount() / 3, 1);
+}
 
 void StencilingGeometry::DrawCone(FRHICommandList& RHICmdList)
 {
@@ -3638,7 +3647,8 @@ bool FDeferredShadingSceneRenderer::RenderCachedPreshadows(FRHICommandListImmedi
 					bool bPerformClear = true;;
 					auto SetShadowRenderTargets = [this, &bPerformClear, ProjectedShadowInfo, &SceneContext](FRHICommandList& InRHICmdList)
 					{
-						SetRenderTarget(InRHICmdList, FTextureRHIRef(), SceneContext.PreShadowCacheDepthZ->GetRenderTargetItem().TargetableTexture);
+						// Must preserve existing contents as the clear will be scissored
+						SetRenderTarget(InRHICmdList, FTextureRHIRef(), SceneContext.PreShadowCacheDepthZ->GetRenderTargetItem().TargetableTexture, ESimpleRenderTargetMode::EExistingColorAndDepth);
 						ProjectedShadowInfo->ClearDepth(InRHICmdList, this, bPerformClear);
 					};					
 					SetShadowRenderTargets(RHICmdList); // run it now, maybe run it later for parallel command lists
diff --git a/Engine/Source/Runtime/Renderer/Private/ShadowRendering.h b/Engine/Source/Runtime/Renderer/Private/ShadowRendering.h
--- a/Engine/Source/Runtime/Renderer/Private/ShadowRendering.h
+++ b/Engine/Source/Runtime/Renderer/Private/ShadowRendering.h
@@ -21,13 +21,17 @@ namespace StencilingGeometry
 	* Note: The sphere will be of unit size unless transformed by the shader. 
 	*/
 	extern void DrawSphere(FRHICommandList& RHICmdList);
+	/**
+	 * Draws exactly the same as above, but uses FVector rather than FVector4 vertex data.
+	 */
+	extern void DrawVectorSphere(FRHICommandList& RHICmdList);
 	/** Renders a cone with a spherical cap, used for rendering spot lights in deferred passes. */
 	extern void DrawCone(FRHICommandList& RHICmdList);
 
 	/** 
 	* Vertex buffer for a sphere of unit size. Used for drawing a sphere as approximate bounding geometry for deferred passes.
 	*/
-	template<int32 NumSphereSides, int32 NumSphereRings>
+	template<int32 NumSphereSides, int32 NumSphereRings, typename VectorType>
 	class TStencilSphereVertexBuffer : public FVertexBuffer
 	{
 	public:
@@ -49,7 +53,7 @@ namespace StencilingGeometry
 			const float RadiansPerRingSegment = PI / (float)NumRings;
 			float Radius = 1;
 
-			TArray<FVector4, TInlineAllocator<NumRings + 1> > ArcVerts;
+			TArray<VectorType, TInlineAllocator<NumRings + 1> > ArcVerts;
 			ArcVerts.Empty(NumRings + 1);
 			// Calculate verts for one arc
 			for (int32 i = 0; i < NumRings + 1; i++)
@@ -58,7 +62,7 @@ namespace StencilingGeometry
 				ArcVerts.Add(FVector(0.0f, FMath::Sin(Angle), FMath::Cos(Angle)));
 			}
 
-			TResourceArray<FVector4, VERTEXBUFFER_ALIGNMENT> Verts;
+			TResourceArray<VectorType, VERTEXBUFFER_ALIGNMENT> Verts;
 			Verts.Empty(NumVerts);
 			// Then rotate this arc NumSides + 1 times.
 			const FVector Center = FVector(0,0,0);
@@ -261,9 +265,10 @@ namespace StencilingGeometry
 		int32 GetVertexCount() const { return NumVerts; }
 	};
 
-	extern TGlobalResource<TStencilSphereVertexBuffer<18, 12> > GStencilSphereVertexBuffer;
+	extern TGlobalResource<TStencilSphereVertexBuffer<18, 12, FVector4> > GStencilSphereVertexBuffer;
+	extern TGlobalResource<TStencilSphereVertexBuffer<18, 12, FVector> > GStencilSphereVectorBuffer;
 	extern TGlobalResource<TStencilSphereIndexBuffer<18, 12> > GStencilSphereIndexBuffer;
-	extern TGlobalResource<TStencilSphereVertexBuffer<4, 4> > GLowPolyStencilSphereVertexBuffer;
+	extern TGlobalResource<TStencilSphereVertexBuffer<4, 4, FVector4> > GLowPolyStencilSphereVertexBuffer;
 	extern TGlobalResource<TStencilSphereIndexBuffer<4, 4> > GLowPolyStencilSphereIndexBuffer;
 	extern TGlobalResource<FStencilConeVertexBuffer> GStencilConeVertexBuffer;
 	extern TGlobalResource<FStencilConeIndexBuffer> GStencilConeIndexBuffer;
diff --git a/Engine/Source/Runtime/Renderer/Private/SystemTextures.cpp b/Engine/Source/Runtime/Renderer/Private/SystemTextures.cpp
--- a/Engine/Source/Runtime/Renderer/Private/SystemTextures.cpp
+++ b/Engine/Source/Runtime/Renderer/Private/SystemTextures.cpp
@@ -28,21 +28,19 @@ void FSystemTextures::InitializeTextures(FRHICommandListImmediate& RHICmdList, E
 	{
 		// Create a WhiteDummy texture
 		{
-			FPooledRenderTargetDesc Desc(FPooledRenderTargetDesc::Create2DDesc(FIntPoint(1, 1), PF_B8G8R8A8, FClearValueBinding::None, TexCreate_HideInVisualizeTexture, TexCreate_RenderTargetable | TexCreate_NoFastClear, false));
+			FPooledRenderTargetDesc Desc(FPooledRenderTargetDesc::Create2DDesc(FIntPoint(1, 1), PF_B8G8R8A8, FClearValueBinding::White, TexCreate_HideInVisualizeTexture, TexCreate_RenderTargetable | TexCreate_NoFastClear, false));
 			GRenderTargetPool.FindFreeElement(Desc, WhiteDummy, TEXT("WhiteDummy"));
 
-			SetRenderTarget(RHICmdList, WhiteDummy->GetRenderTargetItem().TargetableTexture, FTextureRHIRef());
-			RHICmdList.Clear(true, FLinearColor(1, 1, 1, 1), false, 0, false, 0, FIntRect());
+			SetRenderTarget(RHICmdList, WhiteDummy->GetRenderTargetItem().TargetableTexture, FTextureRHIRef(), ESimpleRenderTargetMode::EClearColorExistingDepth);
 			RHICmdList.CopyToResolveTarget(WhiteDummy->GetRenderTargetItem().TargetableTexture, WhiteDummy->GetRenderTargetItem().ShaderResourceTexture, true, FResolveParams());
 		}
 
 		// Create a BlackDummy texture
 		{
-			FPooledRenderTargetDesc Desc(FPooledRenderTargetDesc::Create2DDesc(FIntPoint(1, 1), PF_B8G8R8A8, FClearValueBinding::None, TexCreate_HideInVisualizeTexture, TexCreate_RenderTargetable | TexCreate_NoFastClear, false));
+			FPooledRenderTargetDesc Desc(FPooledRenderTargetDesc::Create2DDesc(FIntPoint(1, 1), PF_B8G8R8A8, FClearValueBinding::Black, TexCreate_HideInVisualizeTexture, TexCreate_RenderTargetable | TexCreate_NoFastClear, false));
 			GRenderTargetPool.FindFreeElement(Desc, BlackDummy, TEXT("BlackDummy"));
 
-			SetRenderTarget(RHICmdList, BlackDummy->GetRenderTargetItem().TargetableTexture, FTextureRHIRef());
-			RHICmdList.Clear(true, FLinearColor(0, 0, 0, 0), false, 0, false, 0, FIntRect());
+			SetRenderTarget(RHICmdList, BlackDummy->GetRenderTargetItem().TargetableTexture, FTextureRHIRef(), ESimpleRenderTargetMode::EClearColorExistingDepth);
 			RHICmdList.CopyToResolveTarget(BlackDummy->GetRenderTargetItem().TargetableTexture, BlackDummy->GetRenderTargetItem().ShaderResourceTexture, true, FResolveParams());
 		}
 		
@@ -110,7 +108,7 @@ void FSystemTextures::InitializeTextures(FRHICommandListImmediate& RHICmdList, E
 	}
 
 	// Create the PerlinNoise3D texture (similar to http://prettyprocs.wordpress.com/2012/10/20/fast-perlin-noise/)
-	if (FeatureLevelInitializedTo < ERHIFeatureLevel::SM5 && InFeatureLevel >= ERHIFeatureLevel::SM5)
+	if (FeatureLevelInitializedTo < ERHIFeatureLevel::SM5 && InFeatureLevel >= ERHIFeatureLevel::SM4)
 	{
 		uint32 Extent = 16;
 
diff --git a/Engine/Source/Runtime/Renderer/Public/GlobalDistanceFieldParameters.h b/Engine/Source/Runtime/Renderer/Public/GlobalDistanceFieldParameters.h
--- a/Engine/Source/Runtime/Renderer/Public/GlobalDistanceFieldParameters.h
+++ b/Engine/Source/Runtime/Renderer/Public/GlobalDistanceFieldParameters.h
@@ -30,8 +30,14 @@ class FGlobalDistanceFieldParameters
 public:
 	void Bind(const FShaderParameterMap& ParameterMap)
 	{
-		GlobalDistanceFieldTexture.Bind(ParameterMap,TEXT("GlobalDistanceFieldTexture"));
-		GlobalDistanceFieldSampler.Bind(ParameterMap,TEXT("GlobalDistanceFieldSampler"));
+		GlobalDistanceFieldTexture0.Bind(ParameterMap, TEXT("GlobalDistanceFieldTexture0"));
+		GlobalDistanceFieldTexture1.Bind(ParameterMap, TEXT("GlobalDistanceFieldTexture1"));
+		GlobalDistanceFieldTexture2.Bind(ParameterMap, TEXT("GlobalDistanceFieldTexture2"));
+		GlobalDistanceFieldTexture3.Bind(ParameterMap, TEXT("GlobalDistanceFieldTexture3"));
+		GlobalDistanceFieldSampler0.Bind(ParameterMap, TEXT("GlobalDistanceFieldSampler0"));
+		GlobalDistanceFieldSampler1.Bind(ParameterMap, TEXT("GlobalDistanceFieldSampler1"));
+		GlobalDistanceFieldSampler2.Bind(ParameterMap, TEXT("GlobalDistanceFieldSampler2"));
+		GlobalDistanceFieldSampler3.Bind(ParameterMap, TEXT("GlobalDistanceFieldSampler3"));
 		GlobalVolumeCenterAndExtent.Bind(ParameterMap,TEXT("GlobalVolumeCenterAndExtent"));
 		GlobalVolumeWorldToUVAddAndMul.Bind(ParameterMap,TEXT("GlobalVolumeWorldToUVAddAndMul"));
 		GlobalVolumeDimension.Bind(ParameterMap,TEXT("GlobalVolumeDimension"));
@@ -41,8 +47,14 @@ public:
 
 	friend FArchive& operator<<(FArchive& Ar,FGlobalDistanceFieldParameters& Parameters)
 	{
-		Ar << Parameters.GlobalDistanceFieldTexture;
-		Ar << Parameters.GlobalDistanceFieldSampler;
+		Ar << Parameters.GlobalDistanceFieldTexture0;
+		Ar << Parameters.GlobalDistanceFieldTexture1;
+		Ar << Parameters.GlobalDistanceFieldTexture2;
+		Ar << Parameters.GlobalDistanceFieldTexture3;
+		Ar << Parameters.GlobalDistanceFieldSampler0;
+		Ar << Parameters.GlobalDistanceFieldSampler1;
+		Ar << Parameters.GlobalDistanceFieldSampler2;
+		Ar << Parameters.GlobalDistanceFieldSampler3;
 		Ar << Parameters.GlobalVolumeCenterAndExtent;
 		Ar << Parameters.GlobalVolumeWorldToUVAddAndMul;
 		Ar << Parameters.GlobalVolumeDimension;
@@ -56,17 +68,10 @@ public:
 	{
 		if (GlobalVolumeCenterAndExtent.IsBound() || GlobalVolumeWorldToUVAddAndMul.IsBound())
 		{
-			for (int32 ClipmapIndex = 0; ClipmapIndex < (int32)GlobalDistanceFieldTexture.GetNumResources(); ClipmapIndex++)
-			{
-				SetTextureParameter(
-					RHICmdList,
-					ShaderRHI,
-					GlobalDistanceFieldTexture,
-					GlobalDistanceFieldSampler,
-					TStaticSamplerState<SF_Bilinear,AM_Wrap,AM_Wrap,AM_Wrap>::GetRHI(),
-					ParameterData.Textures[ClipmapIndex],
-					ClipmapIndex);
-			}
+			SetTextureParameter(RHICmdList, ShaderRHI, GlobalDistanceFieldTexture0, GlobalDistanceFieldSampler0, TStaticSamplerState<SF_Bilinear, AM_Wrap, AM_Wrap, AM_Wrap>::GetRHI(), ParameterData.Textures[0]);
+			SetTextureParameter(RHICmdList, ShaderRHI, GlobalDistanceFieldTexture1, GlobalDistanceFieldSampler1, TStaticSamplerState<SF_Bilinear, AM_Wrap, AM_Wrap, AM_Wrap>::GetRHI(), ParameterData.Textures[1]);
+			SetTextureParameter(RHICmdList, ShaderRHI, GlobalDistanceFieldTexture2, GlobalDistanceFieldSampler2, TStaticSamplerState<SF_Bilinear, AM_Wrap, AM_Wrap, AM_Wrap>::GetRHI(), ParameterData.Textures[2]);
+			SetTextureParameter(RHICmdList, ShaderRHI, GlobalDistanceFieldTexture3, GlobalDistanceFieldSampler3, TStaticSamplerState<SF_Bilinear, AM_Wrap, AM_Wrap, AM_Wrap>::GetRHI(), ParameterData.Textures[3]);
 
 			SetShaderValueArray(RHICmdList, ShaderRHI, GlobalVolumeCenterAndExtent, ParameterData.CenterAndExtent, GMaxGlobalDistanceFieldClipmaps);
 			SetShaderValueArray(RHICmdList, ShaderRHI, GlobalVolumeWorldToUVAddAndMul, ParameterData.WorldToUVAddAndMul, GMaxGlobalDistanceFieldClipmaps);
@@ -77,8 +82,14 @@ public:
 	}
 
 private:
-	FShaderResourceParameter GlobalDistanceFieldTexture;
-	FShaderResourceParameter GlobalDistanceFieldSampler;
+	FShaderResourceParameter GlobalDistanceFieldTexture0;
+	FShaderResourceParameter GlobalDistanceFieldTexture1;
+	FShaderResourceParameter GlobalDistanceFieldTexture2;
+	FShaderResourceParameter GlobalDistanceFieldTexture3;
+	FShaderResourceParameter GlobalDistanceFieldSampler0;
+	FShaderResourceParameter GlobalDistanceFieldSampler1;
+	FShaderResourceParameter GlobalDistanceFieldSampler2;
+	FShaderResourceParameter GlobalDistanceFieldSampler3;
 	FShaderParameter GlobalVolumeCenterAndExtent;
 	FShaderParameter GlobalVolumeWorldToUVAddAndMul;
 	FShaderParameter GlobalVolumeDimension;
diff --git a/Engine/Source/Runtime/SlateRHIRenderer/Private/SlateRHIRenderer.cpp b/Engine/Source/Runtime/SlateRHIRenderer/Private/SlateRHIRenderer.cpp
--- a/Engine/Source/Runtime/SlateRHIRenderer/Private/SlateRHIRenderer.cpp
+++ b/Engine/Source/Runtime/SlateRHIRenderer/Private/SlateRHIRenderer.cpp
@@ -98,7 +98,7 @@ void FSlateRHIRenderer::FViewportInfo::RecreateDepthBuffer_RenderThread()
 	if (bRequiresStencilTest)
 	{		
 		FTexture2DRHIRef ShaderResourceUnused;
-		FRHIResourceCreateInfo CreateInfo;
+		FRHIResourceCreateInfo CreateInfo(FClearValueBinding::DepthZero);
 		RHICreateTargetableShaderResource2D( Width, Height, PF_DepthStencil, 1, TexCreate_None, TexCreate_DepthStencilTargetable, false, CreateInfo, DepthStencil, ShaderResourceUnused );
 		check( IsValidRef(DepthStencil) );
 	}
@@ -438,18 +438,15 @@ void FSlateRHIRenderer::DrawWindow_RenderThread(FRHICommandListImmediate& RHICmd
 			check(IsValidRef( ViewportInfo.DepthStencil ));
 
 			// Reset the backbuffer as our color render target and also set a depth stencil buffer
-			SetRenderTarget(RHICmdList, BackBuffer, ViewportInfo.DepthStencil);
+			FRHIRenderTargetView ColorView(BackBuffer, 0, -1, bClear ? ERenderTargetLoadAction::EClear : ERenderTargetLoadAction::ELoad, ERenderTargetStoreAction::EStore);
+			FRHISetRenderTargetsInfo Info(1, &ColorView, FRHIDepthRenderTargetView(ViewportInfo.DepthStencil, ERenderTargetLoadAction::ELoad, ERenderTargetStoreAction::EStore, ERenderTargetLoadAction::EClear, ERenderTargetStoreAction::EStore));
+
 			// Clear the stencil buffer
-			RHICmdList.Clear(false, FLinearColor::White, false, 0.0f, true, 0x00, FIntRect());
+			RHICmdList.SetRenderTargetsAndClear(Info);
 		}
 		else
 		{
-			SetRenderTarget(RHICmdList, BackBuffer, FTextureRHIRef());
-		}
-
-		if (bClear)
-		{
-			RHICmdList.Clear(true, FLinearColor(ForceInitToZero), false, 0.0f, true, 0x00, FIntRect());
+			SetRenderTarget(RHICmdList, BackBuffer, FTextureRHIRef(), bClear ? ESimpleRenderTargetMode::EClearColorAndDepth : ESimpleRenderTargetMode::EExistingColorAndDepth);
 		}
 
 #if DEBUG_OVERDRAW
diff --git a/Engine/Source/Runtime/Windows/D3D11RHI/Private/D3D11Device.cpp b/Engine/Source/Runtime/Windows/D3D11RHI/Private/D3D11Device.cpp
--- a/Engine/Source/Runtime/Windows/D3D11RHI/Private/D3D11Device.cpp
+++ b/Engine/Source/Runtime/Windows/D3D11RHI/Private/D3D11Device.cpp
@@ -171,6 +171,7 @@ FD3D11DynamicRHI::FD3D11DynamicRHI(IDXGIFactory1* InDXGIFactory1,D3D_FEATURE_LEV
 		GMaxTextureDimensions = D3D11_REQ_TEXTURE2D_U_OR_V_DIMENSION;
 		GMaxCubeTextureDimensions = D3D11_REQ_TEXTURECUBE_DIMENSION;
 		GMaxTextureArrayLayers = D3D11_REQ_TEXTURE2D_ARRAY_AXIS_DIMENSION;
+		GRHISupportsMSAADepthSampleAccess = true;
 	}
 	else if (FeatureLevel >= D3D_FEATURE_LEVEL_10_0)
 	{
