Merging ZebarBranch CL #2670928: Fix Metal SM5 particle sorting - doesn't resolve all particle rendering in SM5 just the sorting. Metal doesn't have Buffer SRVs so data types must match otherwise what is read or written will be entirely wrong.

diff --git a/Engine/Shaders/ParticleSortKeyGen.usf b/Engine/Shaders/ParticleSortKeyGen.usf
--- a/Engine/Shaders/ParticleSortKeyGen.usf
+++ b/Engine/Shaders/ParticleSortKeyGen.usf
@@ -14,14 +14,20 @@
 #include "Common.usf"
 
 /** Input buffer containing particle indices. */
-Buffer<float2> InParticleIndices;
+#if COMPILER_METAL
+#undef half2
+#define vec2 half2
+#else
+#define vec2 float2
+#endif
+Buffer<vec2> InParticleIndices;
 /** Texture containing particle positions. */
 Texture2D PositionTexture;
 
 /** Output key buffer. */
 RWBuffer<uint> OutKeys;
 /** Output indices buffer. */
-RWBuffer<float2> OutParticleIndices;
+RWBuffer<vec2> OutParticleIndices;
 
 [numthreads(THREAD_COUNT,1,1)]
 void GenerateParticleSortKeys(
@@ -55,7 +61,7 @@ void GenerateParticleSortKeys(
 		if ( InputIndex < ParticleKeyGen.KeyCount )
 		{
 			// Read in the particle index and its position.
-			const float2 ParticleIndex = InParticleIndices[InputIndex];
+			const vec2 ParticleIndex = InParticleIndices[InputIndex];
 			int3 ParticleTexel = int3(ParticleIndex.xy * int2(TEXTURE_SIZE_X, TEXTURE_SIZE_Y), 0);
 			const float4 ParticlePosition = PositionTexture.Load(ParticleTexel);
 
