integrated from Orion

CL 2677907
optimized SSAO and exposed more low quality levels

also incldues som minor refarcot in Orion, will integrate into main

diff --git a/Engine/Config/BaseScalability.ini b/Engine/Config/BaseScalability.ini
--- a/Engine/Config/BaseScalability.ini
+++ b/Engine/Config/BaseScalability.ini
@@ -137,7 +137,7 @@ r.Filter.SizeScale=0.6
 r.MotionBlurQuality=3
 r.AmbientOcclusionMipLevelFactor=1.0
 r.AmbientOcclusionMaxQuality=60
-r.AmbientOcclusionLevels=1
+r.AmbientOcclusionLevels=-1
 r.AmbientOcclusionRadiusScale=1.5
 r.DepthOfFieldQuality=1
 r.RenderTargetPoolMin=350
@@ -155,7 +155,7 @@ r.Filter.SizeScale=0.7
 r.MotionBlurQuality=3
 r.AmbientOcclusionMipLevelFactor=0.6
 r.AmbientOcclusionMaxQuality=100
-r.AmbientOcclusionLevels=1
+r.AmbientOcclusionLevels=-1
 r.AmbientOcclusionRadiusScale=1.5
 r.DepthOfFieldQuality=2
 r.RenderTargetPoolMin=400
@@ -173,7 +173,7 @@ r.Filter.SizeScale=0.8
 r.MotionBlurQuality=4
 r.AmbientOcclusionMipLevelFactor=0.4
 r.AmbientOcclusionMaxQuality=100
-r.AmbientOcclusionLevels=2
+r.AmbientOcclusionLevels=-1
 r.AmbientOcclusionRadiusScale=1.0
 r.DepthOfFieldQuality=2
 r.RenderTargetPoolMin=400
diff --git a/Engine/Shaders/PostProcessAmbientOcclusion.usf b/Engine/Shaders/PostProcessAmbientOcclusion.usf
--- a/Engine/Shaders/PostProcessAmbientOcclusion.usf
+++ b/Engine/Shaders/PostProcessAmbientOcclusion.usf
@@ -10,8 +10,8 @@
 
 // set by C++:
 //
-// 0:low / 1: medium / 2:high
-// SAMPLESET_QUALITY
+// 0:low / 1: medium / 2:high / 4:very high
+// SHADER_QUALITY
 //
 // 0:no / 1:yes
 // USE_AO_SETUP_AS_INPUT
@@ -19,12 +19,12 @@
 // 0:no / 1:yes
 // USE_UPSAMPLE
 
-// 0:normal method using depth computed in setup passes (might need a clear outside the viewport, test with stereo/splitscreen) / 1:faster but slightly lower quality
-#define USE_HZB	1
-
 // 0: classic with weighted sample, 1: don't normalize and adjust the formula to be simpler and faster - can look better and is cheaper (Alchemy like?)
 #define OPTIMIZATION_O1 1
 
+// 0:lower quality, 1: higher quality with double sample count
+//#define SSAO_SAMPLE_PAIRS 0
+
 // 1:lowest quality, 2:medium , 3:high, more doesn't give too much (maybe HZB mip computations should be adjusted)
 //#define SAMPLE_STEPS 3
 
@@ -37,18 +37,26 @@
 // AO_SAMPLE_QUALITY = 2 : efficient high frequency 4x4 pattern without jitter for TemporalAA
 // AO_SAMPLE_QUALITY = 3 : efficient high frequency 4x4 pattern with jitter for TemporalAA
 
-#if SAMPLESET_QUALITY == 0
+#if SHADER_QUALITY == 0
 	// low
 	#define USE_SAMPLESET 1
 	#define SAMPLE_STEPS 1
-#elif SAMPLESET_QUALITY == 1
+	#define SSAO_SAMPLE_PAIRS 0
+#elif SHADER_QUALITY == 1
 	// medium
 	#define USE_SAMPLESET 1
-	#define SAMPLE_STEPS 3
-#else // SAMPLESET_QUALITY == 2
+	#define SAMPLE_STEPS 1
+	#define SSAO_SAMPLE_PAIRS 1
+#elif SHADER_QUALITY == 2
 	// high
+	#define USE_SAMPLESET 1
+	#define SAMPLE_STEPS 3
+	#define SSAO_SAMPLE_PAIRS 1
+#else // SHADER_QUALITY == 3
+	// very high
 	#define USE_SAMPLESET 3
 	#define SAMPLE_STEPS 3
+	#define SSAO_SAMPLE_PAIRS 1
 #endif
 
 // 0:4 samples, 1:9 samples (only really noticable with dither usage ??)
@@ -335,7 +343,7 @@ float acosApproxNormAngle(float x)
 }
 
 // @return float2(InvNormAngle, Weight)
-float2 Wedge(float Scale, float2 LocalRandom, float3 FovFix, float3 InvFovFix, float3 ViewSpacePosition, float3 ScaledViewSpaceNormal, float InvHaloSize, float MipLevel)
+float2 WedgeWithNormal(float Scale, float2 LocalRandom, float3 FovFix, float3 InvFovFix, float3 ViewSpacePosition, float3 ScaledViewSpaceNormal, float InvHaloSize, float MipLevel)
 {
 	float AORadiusInShader = ScreenSpaceAOandSSRShaderParams[1].z;		// 40 is default, /400 is constant
 	float ScaleRadiusInWorldSpace = ScreenSpaceAOandSSRShaderParams[2].z;
@@ -344,15 +352,9 @@ float2 Wedge(float Scale, float2 LocalRandom, float3 FovFix, float3 InvFovFix, f
 
 	float2 ScreenSpacePos = ReconstructSSPosFromCS(SamplePosition2 * FovFix);
 
-#if USE_HZB
 	float2 UV = ScreenSpacePos * float2(0.5f, -0.5f) + 0.5f;
 
 	float Abs = ConvertFromDeviceZ(Texture2DSampleLevel( PostprocessInput3, PostprocessInput3Sampler, HzbUvAndStepMipLevelFactor.xy * UV, MipLevel).r);
-#else
-	float2 UV = ScreenPosToUV(ScreenSpacePos, PostprocessInput0Size);
-
-	float Abs = Texture2DSample(PostprocessInput1, PostprocessInput1Sampler, UV).a * Constant_Float16F_Scale;
-#endif
 
 	float3 SamplePosition = ReconstructCSPos(Abs, ScreenSpacePos) * InvFovFix;
 		
@@ -375,7 +377,7 @@ float2 Wedge(float Scale, float2 LocalRandom, float3 FovFix, float3 InvFovFix, f
 
 
 // @return float2(InvNormAngle, Weight)
-float2 Wedge2(float Scale, float2 LocalRandom, float3 FovFix, float3 InvFovFix, float3 ViewSpacePosition, float3 ScaledViewSpaceNormal, float InvHaloSize, float MipLevel)
+float2 WedgeNoNormal(float Scale, float2 LocalRandom, float3 FovFix, float3 InvFovFix, float3 ViewSpacePosition, float3 ScaledViewSpaceNormal, float InvHaloSize, float MipLevel)
 {
 	float3 DeltaLeft;
 	float AbsLeft;
@@ -389,15 +391,9 @@ float2 Wedge2(float Scale, float2 LocalRandom, float3 FovFix, float3 InvFovFix,
 
 		float2 ScreenSpacePos = ReconstructSSPosFromCS(SamplePosition2 * FovFix);
 
-	#if USE_HZB
 		float2 UV = ScreenSpacePos * float2(0.5f, -0.5f) + 0.5f;
 
 		AbsLeft = ConvertFromDeviceZ(Texture2DSampleLevel( PostprocessInput3, PostprocessInput3Sampler, HzbUvAndStepMipLevelFactor.xy * UV, MipLevel).r);
-	#else
-		float2 UV = ScreenPosToUV(ScreenSpacePos, PostprocessInput0Size);
-
-		AbsLeft = Texture2DSample(PostprocessInput1, PostprocessInput1Sampler, UV).a * Constant_Float16F_Scale;
-	#endif
 
 		SamplePositionLeft = ReconstructCSPos(AbsLeft, ScreenSpacePos) * InvFovFix;
 		
@@ -422,15 +418,9 @@ float2 Wedge2(float Scale, float2 LocalRandom, float3 FovFix, float3 InvFovFix,
 
 		float2 ScreenSpacePos = ReconstructSSPosFromCS(SamplePosition2 * FovFix);
 
-	#if USE_HZB
 		float2 UV = ScreenSpacePos * float2(0.5f, -0.5f) + 0.5f;
 
-		AbsLeft = ConvertFromDeviceZ(Texture2DSampleLevel( PostprocessInput3, PostprocessInput3Sampler, HzbUvAndStepMipLevelFactor.xy * UV, MipLevel).r);
-	#else
-		float2 UV = ScreenPosToUV(ScreenSpacePos, PostprocessInput0Size);
-
-		AbsRight = Texture2DSample(PostprocessInput1, PostprocessInput1Sampler, UV).a * Constant_Float16F_Scale;
-	#endif
+		AbsRight = ConvertFromDeviceZ(Texture2DSampleLevel( PostprocessInput3, PostprocessInput3Sampler, HzbUvAndStepMipLevelFactor.xy * UV, MipLevel).r);
 
 		SamplePositionRight = ReconstructCSPos(AbsRight, ScreenSpacePos) * InvFovFix;
 		
@@ -529,8 +519,9 @@ void MainPS(in float4 UVAndScreenPos : TEXCOORD0, out float4 OutColor : SV_Targe
 	{
 #else // AO_SAMPLE_QUALITY == 3
 	// extract one of 16 base vectors (rotation and scale) from a texture that repeats 4x4, changing over time if TemporalAA is enabled
+
 	// todo: can be moved into the vertex shader
-	float2 TemporalOffset = (View.TemporalAAParams.r % 8) * float2(2.48f, 6.72f) / 64.0f;
+	float2 TemporalOffset = (View.TemporalAAParams.r % 8) * float2(2.48f, 7.52f) / 64.0f;
 
 	float4 RandomNorm = Texture2DSample(RandomNormalTexture, RandomNormalTextureSampler, TemporalOffset + UV * ViewportUVToRandomUV);
 	float4 RandomBase = (RandomNorm * 2 - 1) * float4(ActualAORadiusXY, ActualAORadiusXY);
@@ -538,6 +529,15 @@ void MainPS(in float4 UVAndScreenPos : TEXCOORD0, out float4 OutColor : SV_Targe
 	{
 #endif // AO_SAMPLE_QUALITY == 
 
+#if !SSAO_SAMPLE_PAIRS
+		// alternate -1/1 each pixel
+		float mirror = (frac(UV.x * ViewportUVToRandomUV.x * 32.0f) < 0.5f) ? -1 : 1;
+
+		// noise experiment: 0:dither..0.5:randomness
+//		float RandomnessConst = 0.0f;
+//		mirror *= lerp(1.0f, 2 * PseudoRandom(UV * ViewportUVToRandomUV * 32.0f), RandomnessConst);
+#endif
+
 		// .x means for very anisotropic viewports we scale by x
 		float InvHaloSize = 1.0f / (ActualAORadiusXY.x * 2);
 
@@ -568,7 +568,7 @@ void MainPS(in float4 UVAndScreenPos : TEXCOORD0, out float4 OutColor : SV_Targe
 					// constant at run time (higher is better for texture cache / performance, lower is better quality
 					float MipLevel = log2(1.0f + HzbUvAndStepMipLevelFactor.z * step);
 
-					float2 StepSample = Wedge2(Scale, LocalRandom, FovFix, InvFovFix, ViewSpacePosition, ScaledViewSpaceNormal, InvHaloSize, MipLevel);
+					float2 StepSample = WedgeNoNormal(Scale, LocalRandom, FovFix, InvFovFix, ViewSpacePosition, ScaledViewSpaceNormal, InvHaloSize, MipLevel);
 
 					// combine horizon samples
 					LocalAccumulator = lerp(LocalAccumulator, float2(max(LocalAccumulator.x, StepSample.x), 1), StepSample.y);
@@ -579,7 +579,9 @@ void MainPS(in float4 UVAndScreenPos : TEXCOORD0, out float4 OutColor : SV_Targe
 			}
 			else // SSAO_NORMAL_AFFECTED == 1
 			{
+#if SSAO_SAMPLE_PAIRS
 				UNROLL for(int mirror = -1; mirror <= 1; mirror +=2)
+#endif
 				{
 					float2 LocalAccumulator = float2(0, 0);
 
@@ -590,7 +592,7 @@ void MainPS(in float4 UVAndScreenPos : TEXCOORD0, out float4 OutColor : SV_Targe
 						// constant at run time (higher is better for texture cache / performance, lower is better quality
 						float MipLevel = log2(1.0f + HzbUvAndStepMipLevelFactor.z * step);
 
-						float2 StepSample = Wedge(Scale, LocalRandom, FovFix, InvFovFix, ViewSpacePosition, ScaledViewSpaceNormal, InvHaloSize, MipLevel);
+						float2 StepSample = WedgeWithNormal(Scale, LocalRandom, FovFix, InvFovFix, ViewSpacePosition, ScaledViewSpaceNormal, InvHaloSize, MipLevel);
 
 						// combine horizon samples
 						LocalAccumulator = lerp(LocalAccumulator, float2(max(LocalAccumulator.x, StepSample.x), 1), StepSample.y);
diff --git a/Engine/Source/Runtime/Core/Private/HAL/ConsoleManager.cpp b/Engine/Source/Runtime/Core/Private/HAL/ConsoleManager.cpp
--- a/Engine/Source/Runtime/Core/Private/HAL/ConsoleManager.cpp
+++ b/Engine/Source/Runtime/Core/Private/HAL/ConsoleManager.cpp
@@ -1658,13 +1658,6 @@ static TAutoConsoleVariable<int32> CVarSceneColorFringeQuality(
 
 // ---------------------------------------
 
-static TAutoConsoleVariable<int32> CVarAmbientOcclusionLevels(
-	TEXT("r.AmbientOcclusionLevels"),
-	3,
-	TEXT("Defines how many mip levels are using during the ambient occlusion calculation. This is useful when tweaking the algorithm.\n")
-	TEXT(" 0:none, 1:one, 2:two, 3:three(default), 4:four(larger at little cost but can flicker)"),
-	ECVF_Scalability | ECVF_RenderThreadSafe);
-
 static TAutoConsoleVariable<float> CVarAmbientOcclusionRadiusScale(
 	TEXT("r.AmbientOcclusionRadiusScale"),
 	1.0f,
@@ -2116,4 +2109,4 @@ static TAutoConsoleVariable<int32> CVarHLODSystemEnabled(
 	TEXT("r.HLODEnabled"), 
 	1,
 	TEXT("Toggles whether or not the Hierarchical LOD system is enabled."),
-	ECVF_Scalability | ECVF_RenderThreadSafe);
\ No newline at end of file
+	ECVF_Scalability | ECVF_RenderThreadSafe);
diff --git a/Engine/Source/Runtime/Renderer/Private/CompositionLighting/CompositionLighting.cpp b/Engine/Source/Runtime/Renderer/Private/CompositionLighting/CompositionLighting.cpp
--- a/Engine/Source/Runtime/Renderer/Private/CompositionLighting/CompositionLighting.cpp
+++ b/Engine/Source/Runtime/Renderer/Private/CompositionLighting/CompositionLighting.cpp
@@ -40,6 +40,17 @@ static TAutoConsoleVariable<float> CVarSSSScale(
 	TEXT(">1: scale scatter radius up (for testing)"),
 	ECVF_Scalability | ECVF_RenderThreadSafe);
 
+static TAutoConsoleVariable<int32> CVarAmbientOcclusionLevels(
+	TEXT("r.AmbientOcclusionLevels"),
+	-1,
+	TEXT("Defines how many mip levels are using during the ambient occlusion calculation. This is useful when tweaking the algorithm.\n")
+	TEXT("<0: decide based on the quality setting in the postprocess settings/volume and r.AmbientOcclusionMaxQuality (default)\n")
+	TEXT(" 0: none (disable AmbientOcclusion)\n")
+	TEXT(" 1: one\n")
+	TEXT(" 2: two (costs extra performance, soft addition)\n")
+	TEXT(" 3: three (larger radius cost less but can flicker)"),
+	ECVF_Scalability | ECVF_RenderThreadSafe);
+
 static TAutoConsoleVariable<int32> CVarSSSHalfRes(
 	TEXT("r.SSS.HalfRes"),
 	1,
@@ -108,9 +119,10 @@ static bool IsBasePassAmbientOcclusionRequired(FPostprocessContext& Context)
 	return Context.View.FinalPostProcessSettings.AmbientOcclusionStaticFraction >= 1 / 100.0f && !IsSimpleDynamicLightingEnabled();
 }
 
-// @return 0:off, 0..4
+// @return 0:off, 0..3
 static uint32 ComputeAmbientOcclusionPassCount(FPostprocessContext& Context)
 {
+	// 0:off / 1 / 2 / 3
 	uint32 Ret = 0;
 
 	bool bEnabled = true;
@@ -125,8 +137,24 @@ static uint32 ComputeAmbientOcclusionPassCount(FPostprocessContext& Context)
 
 	if(bEnabled)
 	{
-		static const auto ICVar = IConsoleManager::Get().FindTConsoleVariableDataInt(TEXT("r.AmbientOcclusionLevels"));
-		Ret = FMath::Clamp(ICVar->GetValueOnRenderThread(), 0, 4);
+		// usually in the range 0..100
+		float QualityPercent = GetAmbientOcclusionMaxQualityRT(Context.View);
+
+		// don't expose 0 as the lowest quality should still render
+		Ret = 1 +
+			(QualityPercent > 70.0f) +
+			(QualityPercent > 35.0f);
+
+		int32 CVarLevel = CVarAmbientOcclusionLevels.GetValueOnRenderThread();
+
+		if(CVarLevel >= 0)
+		{
+			// cvar can override (for scalability or to profile/test)
+			Ret = CVarLevel;
+		}
+
+		// bring into valid range
+		Ret = FMath::Min<uint32>(Ret, 3);
 	}
 
 	return Ret;
@@ -141,22 +169,23 @@ static void AddPostProcessingAmbientCubemap(FPostprocessContext& Context, FRende
 	Context.FinalOutput = FRenderingCompositeOutputRef(Pass);
 }
 
-// @param Levels 0..4, how many different resolution levels we want to render
+// @param Levels 0..3, how many different resolution levels we want to render
 static FRenderingCompositeOutputRef AddPostProcessingAmbientOcclusion(FRHICommandListImmediate& RHICmdList, FPostprocessContext& Context, uint32 Levels)
 {
-	check(Levels >= 0 && Levels <= 4);
+	check(Levels >= 0 && Levels <= 3);
 
 	FRenderingCompositePass* AmbientOcclusionInMip1 = 0;
 	FRenderingCompositePass* AmbientOcclusionInMip2 = 0;
-	FRenderingCompositePass* AmbientOcclusionInMip3 = 0;
 	FRenderingCompositePass* AmbientOcclusionPassMip1 = 0; 
 	FRenderingCompositePass* AmbientOcclusionPassMip2 = 0; 
-	FRenderingCompositePass* AmbientOcclusionPassMip3 = 0; 
 
 	// generate input in half, quarter, .. resolution
 
-	AmbientOcclusionInMip1 = Context.Graph.RegisterPass(new(FMemStack::Get()) FRCPassPostProcessAmbientOcclusionSetup());
-	AmbientOcclusionInMip1->SetInput(ePId_Input0, Context.SceneDepth);
+	if(Levels >= 2)
+	{
+		AmbientOcclusionInMip1 = Context.Graph.RegisterPass(new(FMemStack::Get()) FRCPassPostProcessAmbientOcclusionSetup());
+		AmbientOcclusionInMip1->SetInput(ePId_Input0, Context.SceneDepth);
+	}
 
 	if(Levels >= 3)
 	{
@@ -164,30 +193,15 @@ static FRenderingCompositeOutputRef AddPostProcessingAmbientOcclusion(FRHIComman
 		AmbientOcclusionInMip2->SetInput(ePId_Input1, FRenderingCompositeOutputRef(AmbientOcclusionInMip1, ePId_Output0));
 	}
 
-	if(Levels >= 4)
-	{
-		AmbientOcclusionInMip3 = Context.Graph.RegisterPass(new(FMemStack::Get()) FRCPassPostProcessAmbientOcclusionSetup());
-		AmbientOcclusionInMip3->SetInput(ePId_Input1, FRenderingCompositeOutputRef(AmbientOcclusionInMip2, ePId_Output0));
-	}
-
 	FRenderingCompositePass* HZBInput = Context.Graph.RegisterPass( new FRCPassPostProcessInput( const_cast< FViewInfo& >( Context.View ).HZB ) );
 
 	// upsample from lower resolution
 
-	if(Levels >= 4)
-	{
-		AmbientOcclusionPassMip3 = Context.Graph.RegisterPass(new(FMemStack::Get()) FRCPassPostProcessAmbientOcclusion());
-		AmbientOcclusionPassMip3->SetInput(ePId_Input0, AmbientOcclusionInMip3);
-		AmbientOcclusionPassMip3->SetInput(ePId_Input1, AmbientOcclusionInMip3);
-		AmbientOcclusionPassMip3->SetInput(ePId_Input3, HZBInput);
-	}
-
 	if(Levels >= 3)
 	{
 		AmbientOcclusionPassMip2 = Context.Graph.RegisterPass(new(FMemStack::Get()) FRCPassPostProcessAmbientOcclusion());
 		AmbientOcclusionPassMip2->SetInput(ePId_Input0, AmbientOcclusionInMip2);
 		AmbientOcclusionPassMip2->SetInput(ePId_Input1, AmbientOcclusionInMip2);
-		AmbientOcclusionPassMip2->SetInput(ePId_Input2, AmbientOcclusionPassMip3);
 		AmbientOcclusionPassMip2->SetInput(ePId_Input3, HZBInput);
 	}
 
@@ -213,7 +227,10 @@ static FRenderingCompositeOutputRef AddPostProcessingAmbientOcclusion(FRHIComman
 	AmbientOcclusionPassMip0->SetInput(ePId_Input3, HZBInput);
 
 	// to make sure this pass is processed as well (before), needed to make process decals before computing AO
-	AmbientOcclusionInMip1->AddDependency(Context.FinalOutput);
+	if(AmbientOcclusionInMip1)
+	{
+		AmbientOcclusionInMip1->AddDependency(Context.FinalOutput);
+	}
 
 	Context.FinalOutput = FRenderingCompositeOutputRef(AmbientOcclusionPassMip0);
 
diff --git a/Engine/Source/Runtime/Renderer/Private/CompositionLighting/PostProcessAmbientOcclusion.cpp b/Engine/Source/Runtime/Renderer/Private/CompositionLighting/PostProcessAmbientOcclusion.cpp
--- a/Engine/Source/Runtime/Renderer/Private/CompositionLighting/PostProcessAmbientOcclusion.cpp
+++ b/Engine/Source/Runtime/Renderer/Private/CompositionLighting/PostProcessAmbientOcclusion.cpp
@@ -11,9 +11,9 @@
 #include "PostProcessAmbientOcclusion.h"
 #include "SceneUtils.h"
 
-static TAutoConsoleVariable<int32> CVarAmbientOcclusionMaxQuality(
+static TAutoConsoleVariable<float> CVarAmbientOcclusionMaxQuality(
 	TEXT("r.AmbientOcclusionMaxQuality"),
-	100,
+	100.0f,
 	TEXT("Defines the max clamping value from the post process volume's quality level for ScreenSpace Ambient Occlusion\n")
 	TEXT(" 100: don't override quality level from the post process volume (default)\n")
 	TEXT(" <100: clamp down quality level from the post process volume to the maximum set by this cvar"),
@@ -126,7 +126,7 @@ FArchive& operator<<(FArchive& Ar, FScreenSpaceAOandSSRShaderParameters& This)
  * @param bAOSetupAsInput true:use AO setup instead of full resolution depth and normal
  * @param bDoUpsample true:we have lower resolution pass data we need to upsample, false otherwise
  */
-template<uint32 bTAOSetupAsInput, uint32 bDoUpsample, uint32 SampleSetQuality>
+template<uint32 bTAOSetupAsInput, uint32 bDoUpsample, uint32 ShaderQuality>
 class FPostProcessAmbientOcclusionPS : public FGlobalShader
 {
 	DECLARE_SHADER_TYPE(FPostProcessAmbientOcclusionPS, Global);
@@ -142,7 +142,7 @@ class FPostProcessAmbientOcclusionPS : public FGlobalShader
 
 		OutEnvironment.SetDefine(TEXT("USE_UPSAMPLE"), bDoUpsample);
 		OutEnvironment.SetDefine(TEXT("USE_AO_SETUP_AS_INPUT"), bTAOSetupAsInput);
-		OutEnvironment.SetDefine(TEXT("SAMPLESET_QUALITY"), SampleSetQuality);
+		OutEnvironment.SetDefine(TEXT("SHADER_QUALITY"), ShaderQuality);
 	}
 
 	/** Default constructor. */
@@ -192,7 +192,7 @@ public:
 			HzbStepMipLevelFactorValue
 			);
 		
-		SetShaderValue(Context.RHICmdList, ShaderRHI, HzbUvAndStepMipLevelFactor, HzbUvAndStepMipLevelFactorValue );
+		SetShaderValue(Context.RHICmdList, ShaderRHI, HzbUvAndStepMipLevelFactor, HzbUvAndStepMipLevelFactorValue);
 	}
 	
 	// FShader interface.
@@ -224,6 +224,7 @@ public:
 	VARIATION0(0)
 	VARIATION0(1)
 	VARIATION0(2)
+	VARIATION0(3)
 	
 #undef VARIATION0
 #undef VARIATION1
@@ -404,11 +405,11 @@ bool FRCPassPostProcessAmbientOcclusionSetup::IsInitialPass() const
 
 
 // ---------------------------------
-template <uint32 bTAOSetupAsInput, uint32 bDoUpsample, uint32 SampleSetQuality>
+template <uint32 bTAOSetupAsInput, uint32 bDoUpsample, uint32 ShaderQuality>
 FShader* FRCPassPostProcessAmbientOcclusion::SetShaderTempl(const FRenderingCompositePassContext& Context)
 {
 	TShaderMapRef<FPostProcessVS> VertexShader(Context.GetShaderMap());
-	TShaderMapRef<FPostProcessAmbientOcclusionPS<bTAOSetupAsInput, bDoUpsample, SampleSetQuality> > PixelShader(Context.GetShaderMap());
+	TShaderMapRef<FPostProcessAmbientOcclusionPS<bTAOSetupAsInput, bDoUpsample, ShaderQuality> > PixelShader(Context.GetShaderMap());
 
 	static FGlobalBoundShaderState BoundShaderState;
 	
@@ -423,6 +424,11 @@ FShader* FRCPassPostProcessAmbientOcclusion::SetShaderTempl(const FRenderingComp
 	return *VertexShader;
 }
 
+float GetAmbientOcclusionMaxQualityRT(const FSceneView& View)
+{
+	return FMath::Min(CVarAmbientOcclusionMaxQuality.GetValueOnRenderThread(), View.FinalPostProcessSettings.AmbientOcclusionQuality);
+}
+
 void FRCPassPostProcessAmbientOcclusion::Process(FRenderingCompositePassContext& Context)
 {
 	const FSceneView& View = Context.View;
@@ -460,36 +466,41 @@ void FRCPassPostProcessAmbientOcclusion::Process(FRenderingCompositePassContext&
 	Context.RHICmdList.SetRasterizerState(TStaticRasterizerState<>::GetRHI());
 	Context.RHICmdList.SetDepthStencilState(TStaticDepthStencilState<false, CF_Always>::GetRHI());
 
-	const float QualityPercent = FMath::Min(float(CVarAmbientOcclusionMaxQuality.GetValueOnRenderThread()), Context.View.FinalPostProcessSettings.AmbientOcclusionQuality);
-	const int32 QualitySet = 
+	float QualityPercent = GetAmbientOcclusionMaxQualityRT(Context.View);
+	
+	// 0:low / 1:med: / 2:high / 4:very high
+	const int32 ShaderQuality = 
 		(QualityPercent > 60.0f) +
-		(QualityPercent > 25.0f);
+		(QualityPercent > 25.0f) +
+		(QualityPercent > 5.0f);
+
 	bool bDoUpsample = (InputDesc2 != 0);
 	
-	SCOPED_DRAW_EVENTF(Context.RHICmdList, AmbientOcclusion, TEXT("AmbientOcclusion %dx%d SetupAsInput=%d Upsample=%d SampleSetQuality=%d"), 
-		ViewRect.Width(), ViewRect.Height(), bAOSetupAsInput, bDoUpsample, QualitySet);
+	SCOPED_DRAW_EVENTF(Context.RHICmdList, AmbientOcclusion, TEXT("AmbientOcclusion %dx%d SetupAsInput=%d Upsample=%d ShaderQuality=%d"), 
+		ViewRect.Width(), ViewRect.Height(), bAOSetupAsInput, bDoUpsample, ShaderQuality);
 
 	FShader* VertexShader = 0;
 
-#define SET_SHADER_CASE(Quality)                                \
-	case Quality:			                                    \
-	if(bAOSetupAsInput)                                         \
-    {                                                           \
-		if(bDoUpsample) VertexShader = SetShaderTempl<1, 1, Quality>(Context); \
-		else VertexShader = SetShaderTempl<1, 0, Quality>(Context);            \
-	}                                                           \
-	else                                                        \
-    {                                                           \
-		if(bDoUpsample) VertexShader = SetShaderTempl<0, 1, Quality>(Context); \
-		else VertexShader = SetShaderTempl<0, 0, Quality>(Context);            \
-	}                                                           \
+#define SET_SHADER_CASE(ShaderQualityCase) \
+	case ShaderQualityCase: \
+	if(bAOSetupAsInput) \
+    { \
+		if(bDoUpsample) VertexShader = SetShaderTempl<1, 1, ShaderQualityCase>(Context); \
+		else VertexShader = SetShaderTempl<1, 0, ShaderQualityCase>(Context); \
+	} \
+	else \
+    { \
+		if(bDoUpsample) VertexShader = SetShaderTempl<0, 1, ShaderQualityCase>(Context); \
+		else VertexShader = SetShaderTempl<0, 0, ShaderQualityCase>(Context); \
+	} \
 	break
 
-	switch(QualitySet)
+	switch(ShaderQuality)
 	{
 		SET_SHADER_CASE(0);
 		SET_SHADER_CASE(1);
 		SET_SHADER_CASE(2);
+		SET_SHADER_CASE(3);
 		default:
 			break;
 	};
diff --git a/Engine/Source/Runtime/Renderer/Private/CompositionLighting/PostProcessAmbientOcclusion.h b/Engine/Source/Runtime/Renderer/Private/CompositionLighting/PostProcessAmbientOcclusion.h
--- a/Engine/Source/Runtime/Renderer/Private/CompositionLighting/PostProcessAmbientOcclusion.h
+++ b/Engine/Source/Runtime/Renderer/Private/CompositionLighting/PostProcessAmbientOcclusion.h
@@ -30,8 +30,8 @@ private:
 	FShader* SetShaderSetupTempl(const FRenderingCompositePassContext& Context);
 };
 
-// ePId_Input0: defines the resolution we compute AO and provides the normal
-// ePId_Input1: setup in same resolution as ePId_Input1 for depth expect when running in full resolution, then it's half
+// ePId_Input0: defines the resolution we compute AO and provides the normal (only needed if bInAOSetupAsInput)
+// ePId_Input1: setup in same resolution as ePId_Input1 for depth expect when running in full resolution, then it's half (only needed if bInAOSetupAsInput)
 // ePId_Input2: optional AO result one lower resolution
 // ePId_Input3: optional HZB
 // derives from TRenderingCompositePassBase<InputCount, OutputCount> 
@@ -156,4 +156,8 @@ BEGIN_UNIFORM_BUFFER_STRUCT(FCameraMotionParameters,)
 	DECLARE_UNIFORM_BUFFER_STRUCT_MEMBER_ARRAY(FVector4, Value, [5])
 END_UNIFORM_BUFFER_STRUCT(FCameraMotionParameters)
 
-TUniformBufferRef<FCameraMotionParameters> CreateCameraMotionParametersUniformBuffer(const FSceneView& View);
\ No newline at end of file
+TUniformBufferRef<FCameraMotionParameters> CreateCameraMotionParametersUniformBuffer(const FSceneView& View);
+
+// for render thread
+// @return usually in 0..100 range but could be outside, combines the view with the cvar setting
+float GetAmbientOcclusionMaxQualityRT(const FSceneView& View);
\ No newline at end of file
diff --git a/Engine/Source/Runtime/Renderer/Private/PostProcess/PostProcessBokehDOFRecombine.cpp b/Engine/Source/Runtime/Renderer/Private/PostProcess/PostProcessBokehDOFRecombine.cpp
--- a/Engine/Source/Runtime/Renderer/Private/PostProcess/PostProcessBokehDOFRecombine.cpp
+++ b/Engine/Source/Runtime/Renderer/Private/PostProcess/PostProcessBokehDOFRecombine.cpp
@@ -132,12 +132,13 @@ void FRCPassPostProcessBokehDOFRecombine::Process(FRenderingCompositePassContext
 		check(GetInput(ePId_Input2)->GetPass());
 	}
 
-	SCOPED_DRAW_EVENTF(Context.RHICmdList, BokehDOFRecombine, TEXT("BokehDOFRecombine#%d"), Method);
+	const FSceneView& View = Context.View;
+
+	SCOPED_DRAW_EVENTF(Context.RHICmdList, BokehDOFRecombine, TEXT("BokehDOFRecombine#%d %dx%d"), Method, View.ViewRect.Width(), View.ViewRect.Height());
 
 	const FPooledRenderTargetDesc* InputDesc0 = GetInputDesc(ePId_Input0);
 	const FPooledRenderTargetDesc* InputDesc1 = GetInputDesc(ePId_Input1);
-	
-	const FSceneView& View = Context.View;
+
 
 	FIntPoint TexSize = InputDesc1 ? InputDesc1->Extent : InputDesc0->Extent;
 
@@ -162,7 +163,7 @@ void FRCPassPostProcessBokehDOFRecombine::Process(FRenderingCompositePassContext
 	Context.RHICmdList.SetDepthStencilState(TStaticDepthStencilState<false, CF_Always>::GetRHI());
 
 	switch(Method)
-		{
+	{
 		case 1: SetShader<1>(Context); break;
 		case 2: SetShader<2>(Context); break;
 		case 3: SetShader<3>(Context); break;
diff --git a/Engine/Source/Runtime/Renderer/Private/PostProcess/PostProcessMotionBlur.cpp b/Engine/Source/Runtime/Renderer/Private/PostProcess/PostProcessMotionBlur.cpp
--- a/Engine/Source/Runtime/Renderer/Private/PostProcess/PostProcessMotionBlur.cpp
+++ b/Engine/Source/Runtime/Renderer/Private/PostProcess/PostProcessMotionBlur.cpp
@@ -457,8 +457,6 @@ FRCPassPostProcessMotionBlur::FRCPassPostProcessMotionBlur(uint32 InQuality)
 
 void FRCPassPostProcessMotionBlur::Process(FRenderingCompositePassContext& Context)
 {
-	SCOPED_DRAW_EVENT(Context.RHICmdList, MotionBlur);
-
 	const FPooledRenderTargetDesc* InputDesc = GetInputDesc(ePId_Input0);
 
 	if(!InputDesc)
@@ -480,6 +478,9 @@ void FRCPassPostProcessMotionBlur::Process(FRenderingCompositePassContext& Conte
 	uint32 ScaleFactor = FSceneRenderTargets::Get(Context.RHICmdList).GetBufferSizeXY().X / SrcSize.X;
 
 	FIntRect SrcRect = FIntRect::DivideAndRoundUp(View.ViewRect, ScaleFactor);
+
+	SCOPED_DRAW_EVENTF(Context.RHICmdList, MotionBlur, TEXT("MotionBlur(Old) %dx%d"), SrcRect.Width(), SrcRect.Height());
+
 	FIntRect DestRect = SrcRect;
 
 	const FSceneRenderTargetItem& DestRenderTarget = PassOutputs[0].RequestSurface(Context);
@@ -1360,7 +1361,7 @@ void FRCPassPostProcessMotionBlurNew::Process(FRenderingCompositePassContext& Co
 	FIntRect SrcRect = View.ViewRect / ScaleFactor;
 	FIntRect DestRect = SrcRect;
 
-	SCOPED_DRAW_EVENTF(Context.RHICmdList, MotionBlurNew, TEXT("MotionBlurNew %dx%d"), SrcRect.Width(), SrcRect.Height());
+	SCOPED_DRAW_EVENTF(Context.RHICmdList, MotionBlurNew, TEXT("MotionBlur(New) %dx%d"), SrcRect.Width(), SrcRect.Height());
 
 	const FSceneRenderTargetItem& DestRenderTarget = PassOutputs[0].RequestSurface(Context);
 
