Fixes shader compilation failure in ParticleGPUSpriteVertexFactory.usf

diff --git a/Engine/Shaders/ParticleGPUSpriteVertexFactory.usf b/Engine/Shaders/ParticleGPUSpriteVertexFactory.usf
--- a/Engine/Shaders/ParticleGPUSpriteVertexFactory.usf
+++ b/Engine/Shaders/ParticleGPUSpriteVertexFactory.usf
@@ -278,7 +278,7 @@ FVertexFactoryIntermediates GetVertexFactoryIntermediates( FVertexFactoryInput I
 
 	// Optional camera motion blur.
 #if USES_PARTICLE_MOTION_BLUR
-	float3 MblurVector = -View.WorldCameraOriginDelta;
+	float3 MblurVector = -View.WorldCameraMovementSinceLastFrame;
 	float3 MblurDirection = normalize(MblurVector + 0.0001f);
 	//float MblurFactor = saturate(dot(MblurDirection,normalize(VertexWorldOffset)) * 100000.0f);
 	float MblurFactor = clamp(dot(MblurDirection,normalize(VertexWorldOffset)) * 100000.0f,-1.0f,1.0f);
