Integrated from Orion

CL 2679575
	added quality option for SSAO(QuadMessagePassing blur, normal aware full res 3 sample blur)

CL 2679687
	added ability to override SSAO quality by cvar

diff --git a/Engine/Shaders/PostProcessAmbientOcclusion.usf b/Engine/Shaders/PostProcessAmbientOcclusion.usf
--- a/Engine/Shaders/PostProcessAmbientOcclusion.usf
+++ b/Engine/Shaders/PostProcessAmbientOcclusion.usf
@@ -31,6 +31,10 @@
 // 0:does not use normals (slightly faster, lower quality, WIP = not finished), 1: uses normals (slower but higher quality)
 #define SSAO_NORMAL_AFFECTED 1
 
+// useful to remove high frequency dither pattern, not that needed with more sample
+// 0:off (fast but dither pattern with low sample count), 1:non normal aware (half res look), 2:normal aware (slower), 3:normal and depth aware (slowest, doesn't add much)
+//#define QUAD_MESSAGE_PASSING_BLUR 2
+
 // ambient occlusion
 // AO_SAMPLE_QUALITY = 0 : no AO sampling, only upsampling
 // AO_SAMPLE_QUALITY = 1 : no dither/per pixel randomization
@@ -38,25 +42,49 @@
 // AO_SAMPLE_QUALITY = 3 : efficient high frequency 4x4 pattern with jitter for TemporalAA
 
 #if SHADER_QUALITY == 0
-	// low
+	// very low
 	#define USE_SAMPLESET 1
 	#define SAMPLE_STEPS 1
 	#define SSAO_SAMPLE_PAIRS 0
+	#define QUAD_MESSAGE_PASSING_BLUR 0
 #elif SHADER_QUALITY == 1
+	// low
+	#define USE_SAMPLESET 1
+	#define SAMPLE_STEPS 1
+	#define SSAO_SAMPLE_PAIRS 0
+	#define QUAD_MESSAGE_PASSING_BLUR 2
+#elif SHADER_QUALITY == 2
 	// medium
 	#define USE_SAMPLESET 1
 	#define SAMPLE_STEPS 1
 	#define SSAO_SAMPLE_PAIRS 1
-#elif SHADER_QUALITY == 2
+	#define QUAD_MESSAGE_PASSING_BLUR 0
+#elif SHADER_QUALITY == 3
 	// high
 	#define USE_SAMPLESET 1
 	#define SAMPLE_STEPS 3
 	#define SSAO_SAMPLE_PAIRS 1
-#else // SHADER_QUALITY == 3
+	#define QUAD_MESSAGE_PASSING_BLUR 0
+#else // SHADER_QUALITY == 4
 	// very high
 	#define USE_SAMPLESET 3
 	#define SAMPLE_STEPS 3
 	#define SSAO_SAMPLE_PAIRS 1
+	#define QUAD_MESSAGE_PASSING_BLUR 0
+#endif
+
+#if QUAD_MESSAGE_PASSING_BLUR == 0
+	#define QUAD_MESSAGE_PASSING_NORMAL 0
+	#define QUAD_MESSAGE_PASSING_DEPTH 0
+#elif QUAD_MESSAGE_PASSING_BLUR == 1
+	#define QUAD_MESSAGE_PASSING_NORMAL 0
+	#define QUAD_MESSAGE_PASSING_DEPTH 0
+#elif QUAD_MESSAGE_PASSING_BLUR == 2
+	#define QUAD_MESSAGE_PASSING_NORMAL 1
+	#define QUAD_MESSAGE_PASSING_DEPTH 0
+#elif QUAD_MESSAGE_PASSING_BLUR == 3
+	#define QUAD_MESSAGE_PASSING_NORMAL 1
+	#define QUAD_MESSAGE_PASSING_DEPTH 1
 #endif
 
 // 0:4 samples, 1:9 samples (only really noticable with dither usage ??)
@@ -78,7 +106,6 @@
 // 1: 3 points
 // 2: more evenly spread (5 points - slightly faster, stronger effect, better with multiple levels?)
 // 3: near the surface very large, softly fading out (6 points)
-//#define USE_SAMPLESET 2
 #if USE_SAMPLESET == 0
 	#define SAMPLESET_ARRAY_SIZE 1
 	static const float2 OcclusionSamplesOffsets[SAMPLESET_ARRAY_SIZE]=
@@ -467,6 +494,11 @@ float2 WedgeNoNormal(float Scale, float2 LocalRandom, float3 FovFix, float3 InvF
 	return float2((1-NormAngle) / (Weight + 0.001f), Weight);
 }
 
+float3 ReconstructNormal(float2 In)
+{
+	return float3(In, sqrt(1 - dot(In, In)));
+}
+
 // the main pixel shader that computes ambient occlusion
 void MainPS(in float4 UVAndScreenPos : TEXCOORD0, out float4 OutColor : SV_Target0)
 {
@@ -634,6 +666,53 @@ void MainPS(in float4 UVAndScreenPos : TEXCOORD0, out float4 OutColor : SV_Targe
 	// we output in a single alpha channel
 	OutColor = OutColor.r;
 #endif
+
+#if QUAD_MESSAGE_PASSING_BLUR > 0
+	{
+		// .x: AO output, .y:SceneDepth .zw:view space normal
+		float4 CenterPixel = float4(OutColor.r, SceneDepth, normalize(ViewSpaceNormal).xy); 
+
+		float4 dX = ddx_fine(CenterPixel);
+		float4 dY = ddy_fine(CenterPixel);
+
+		int2 Mod = (uint2)ComputePixelPosCenter(ScreenPos, true) % 2;
+
+		float4 PixA = CenterPixel;
+		float4 PixB = CenterPixel - dX * (Mod.x * 2 - 1);
+		float4 PixC = CenterPixel - dY * (Mod.y * 2 - 1);
+
+		float WeightA = 1.0f;
+		float WeightB = 1.0f;
+		float WeightC = 1.0f;
+
+#if QUAD_MESSAGE_PASSING_NORMAL
+		const float NormalTweak = 4.0f;
+		float3 NormalA = ReconstructNormal(PixA.zw);
+		float3 NormalB = ReconstructNormal(PixB.zw);
+		float3 NormalC = ReconstructNormal(PixC.zw);
+		WeightB *= saturate(pow(saturate(dot(NormalA, NormalB)), NormalTweak));
+		WeightC *= saturate(pow(saturate(dot(NormalA, NormalC)), NormalTweak));
+#endif
+
+#if QUAD_MESSAGE_PASSING_DEPTH
+		const float DepthTweak = 1;
+		float InvDepth = 1.0f / PixA.y;
+		WeightB *= 1 - saturate(abs(1 - PixB.y * InvDepth) * DepthTweak);
+		WeightC *= 1 - saturate(abs(1 - PixC.y * InvDepth) * DepthTweak);
+#endif
+
+		// + 1.0f to avoid div by 0
+		float InvWeightABC = 1.0f / (WeightA + WeightB + WeightC);
+
+		WeightA *= InvWeightABC;
+		WeightB *= InvWeightABC;
+		WeightC *= InvWeightABC;
+
+		OutColor = WeightA * PixA.x + WeightB * PixB.x + WeightC * PixC.x;
+		// visualize where we don't want to fade
+//		OutColor = (WeightA - 0.333f) / 0.666f;
+	}
+#endif
 }
 
 
diff --git a/Engine/Source/Runtime/Renderer/Private/CompositionLighting/CompositionLighting.cpp b/Engine/Source/Runtime/Renderer/Private/CompositionLighting/CompositionLighting.cpp
--- a/Engine/Source/Runtime/Renderer/Private/CompositionLighting/CompositionLighting.cpp
+++ b/Engine/Source/Runtime/Renderer/Private/CompositionLighting/CompositionLighting.cpp
@@ -138,7 +138,7 @@ static uint32 ComputeAmbientOcclusionPassCount(FPostprocessContext& Context)
 	if(bEnabled)
 	{
 		// usually in the range 0..100
-		float QualityPercent = GetAmbientOcclusionMaxQualityRT(Context.View);
+		float QualityPercent = GetAmbientOcclusionQualityRT(Context.View);
 
 		// don't expose 0 as the lowest quality should still render
 		Ret = 1 +
diff --git a/Engine/Source/Runtime/Renderer/Private/CompositionLighting/PostProcessAmbientOcclusion.cpp b/Engine/Source/Runtime/Renderer/Private/CompositionLighting/PostProcessAmbientOcclusion.cpp
--- a/Engine/Source/Runtime/Renderer/Private/CompositionLighting/PostProcessAmbientOcclusion.cpp
+++ b/Engine/Source/Runtime/Renderer/Private/CompositionLighting/PostProcessAmbientOcclusion.cpp
@@ -15,8 +15,9 @@ static TAutoConsoleVariable<float> CVarAmbientOcclusionMaxQuality(
 	TEXT("r.AmbientOcclusionMaxQuality"),
 	100.0f,
 	TEXT("Defines the max clamping value from the post process volume's quality level for ScreenSpace Ambient Occlusion\n")
-	TEXT(" 100: don't override quality level from the post process volume (default)\n")
-	TEXT(" <100: clamp down quality level from the post process volume to the maximum set by this cvar"),
+	TEXT("     100: don't override quality level from the post process volume (default)\n")
+	TEXT("   0..99: clamp down quality level from the post process volume to the maximum set by this cvar\n")
+	TEXT(" -100..0: Enforces a different quality (the absolute value) even if the postprocessvolume asks for a lower quality."),
 	ECVF_Scalability | ECVF_RenderThreadSafe);
 
 static TAutoConsoleVariable<float> CVarAmbientOcclusionStepMipLevelFactor(
@@ -125,6 +126,7 @@ FArchive& operator<<(FArchive& Ar, FScreenSpaceAOandSSRShaderParameters& This)
  * Encapsulates the post processing ambient occlusion pixel shader.
  * @param bAOSetupAsInput true:use AO setup instead of full resolution depth and normal
  * @param bDoUpsample true:we have lower resolution pass data we need to upsample, false otherwise
+ * @param ShaderQuality 0..4, 0:low 4:high
  */
 template<uint32 bTAOSetupAsInput, uint32 bDoUpsample, uint32 ShaderQuality>
 class FPostProcessAmbientOcclusionPS : public FGlobalShader
@@ -225,6 +227,7 @@ public:
 	VARIATION0(1)
 	VARIATION0(2)
 	VARIATION0(3)
+	VARIATION0(4)
 	
 #undef VARIATION0
 #undef VARIATION1
@@ -424,9 +427,18 @@ FShader* FRCPassPostProcessAmbientOcclusion::SetShaderTempl(const FRenderingComp
 	return *VertexShader;
 }
 
-float GetAmbientOcclusionMaxQualityRT(const FSceneView& View)
+float GetAmbientOcclusionQualityRT(const FSceneView& View)
 {
-	return FMath::Min(CVarAmbientOcclusionMaxQuality.GetValueOnRenderThread(), View.FinalPostProcessSettings.AmbientOcclusionQuality);
+	float CVarValue = CVarAmbientOcclusionMaxQuality.GetValueOnRenderThread();
+
+	if(CVarValue < 0)
+	{
+		return FMath::Clamp(-CVarValue, 0.0f, 100.0f);
+	}
+	else
+	{
+		return FMath::Min(CVarValue, View.FinalPostProcessSettings.AmbientOcclusionQuality);
+	}
 }
 
 void FRCPassPostProcessAmbientOcclusion::Process(FRenderingCompositePassContext& Context)
@@ -466,11 +478,12 @@ void FRCPassPostProcessAmbientOcclusion::Process(FRenderingCompositePassContext&
 	Context.RHICmdList.SetRasterizerState(TStaticRasterizerState<>::GetRHI());
 	Context.RHICmdList.SetDepthStencilState(TStaticDepthStencilState<false, CF_Always>::GetRHI());
 
-	float QualityPercent = GetAmbientOcclusionMaxQualityRT(Context.View);
+	float QualityPercent = GetAmbientOcclusionQualityRT(Context.View);
 	
-	// 0:low / 1:med: / 2:high / 4:very high
+	// 0..4, 0:low 4:high
 	const int32 ShaderQuality = 
-		(QualityPercent > 60.0f) +
+		(QualityPercent > 75.0f) +
+		(QualityPercent > 55.0f) +
 		(QualityPercent > 25.0f) +
 		(QualityPercent > 5.0f);
 
@@ -501,6 +514,7 @@ void FRCPassPostProcessAmbientOcclusion::Process(FRenderingCompositePassContext&
 		SET_SHADER_CASE(1);
 		SET_SHADER_CASE(2);
 		SET_SHADER_CASE(3);
+		SET_SHADER_CASE(4);
 		default:
 			break;
 	};
diff --git a/Engine/Source/Runtime/Renderer/Private/CompositionLighting/PostProcessAmbientOcclusion.h b/Engine/Source/Runtime/Renderer/Private/CompositionLighting/PostProcessAmbientOcclusion.h
--- a/Engine/Source/Runtime/Renderer/Private/CompositionLighting/PostProcessAmbientOcclusion.h
+++ b/Engine/Source/Runtime/Renderer/Private/CompositionLighting/PostProcessAmbientOcclusion.h
@@ -160,4 +160,4 @@ TUniformBufferRef<FCameraMotionParameters> CreateCameraMotionParametersUniformBu
 
 // for render thread
 // @return usually in 0..100 range but could be outside, combines the view with the cvar setting
-float GetAmbientOcclusionMaxQualityRT(const FSceneView& View);
\ No newline at end of file
+float GetAmbientOcclusionQualityRT(const FSceneView& View);
\ No newline at end of file
