integrated from Orion

fixed SM4 compiling

diff --git a/Engine/Shaders/PostProcessAmbientOcclusion.usf b/Engine/Shaders/PostProcessAmbientOcclusion.usf
--- a/Engine/Shaders/PostProcessAmbientOcclusion.usf
+++ b/Engine/Shaders/PostProcessAmbientOcclusion.usf
@@ -667,7 +667,7 @@ void MainPS(in float4 UVAndScreenPos : TEXCOORD0, out float4 OutColor : SV_Targe
 	OutColor = OutColor.r;
 #endif
 
-#if QUAD_MESSAGE_PASSING_BLUR > 0
+#if QUAD_MESSAGE_PASSING_BLUR > 0 && FEATURE_LEVEL >= FEATURE_LEVEL_SM5
 	{
 		// .x: AO output, .y:SceneDepth .zw:view space normal
 		float4 CenterPixel = float4(OutColor.r, SceneDepth, normalize(ViewSpaceNormal).xy); 
