force a shader re-compile for metal

diff --git a/Engine/Shaders/MetalCommon.usf b/Engine/Shaders/MetalCommon.usf
--- a/Engine/Shaders/MetalCommon.usf
+++ b/Engine/Shaders/MetalCommon.usf
@@ -5,6 +5,6 @@
 =============================================================================*/
 
 // Update this GUID to improve shader recompilation for Metal only shaders
-// GUID = {06C79EC9-644E-4D22-ADA9-47DD4BB343D9a}
+// GUID = {06C79EC9-644E-4D22-ADA9-47DD4BB343D9b}
   
 #pragma once
