HTML5 point light shader compilation error fix -UE-20202

diff --git a/Engine/Shaders/BasePassForForwardShadingPixelShader.usf b/Engine/Shaders/BasePassForForwardShadingPixelShader.usf
--- a/Engine/Shaders/BasePassForForwardShadingPixelShader.usf
+++ b/Engine/Shaders/BasePassForForwardShadingPixelShader.usf
@@ -290,11 +290,23 @@ void Main(
 
 		#if MAX_DYNAMIC_POINT_LIGHTS > 0
 			#if VARIABLE_NUM_DYNAMIC_POINT_LIGHTS
-				for (int i = 0; i < NumDynamicPointLights; i++)
+			#if WEBGL
+			// webgl needs a constant loop counter
+			UNROLL
+			for (int i = 0; i < MAX_DYNAMIC_POINT_LIGHTS; i++)
+			{
+				if(i >= NumDynamicPointLights)
+				{
+					break;
+				}
 			#else
-				for (int i = 0; i < NUM_DYNAMIC_POINT_LIGHTS; i++)
+			for (int i = 0; i < NumDynamicPointLights; i++)
+			{
 			#endif
+			#else
+			for (int i = 0; i < NUM_DYNAMIC_POINT_LIGHTS; i++)
 			{
+			#endif
 				float3 ToLight = LightPositionAndInvRadius[i].xyz - MaterialParameters.WorldPosition;
 				float DistanceSqr = dot(ToLight, ToLight);
 				float3 L = ToLight * rsqrt(DistanceSqr);
