Hair shading model

Switched clear coat material outputs to CustomData with renaming.

First pass. Code needs clean up. BRDF needs optimization. Only works well with directional lights. Point lights don't obey opacity during shadow projection. No solution yet for indirect.

diff --git a/Engine/Shaders/BRDF.usf b/Engine/Shaders/BRDF.usf
--- a/Engine/Shaders/BRDF.usf
+++ b/Engine/Shaders/BRDF.usf
@@ -13,38 +13,6 @@
 // Microfacet specular = D*G*F / (4*NoL*NoV) = D*Vis*F
 // Vis = G / (4*NoL*NoV)
 
-// Diffuse model
-// 0: Lambert
-// 1: Burley
-// 2: Oren-Nayar
-#define PHYSICAL_DIFFUSE	0
-
-// Microfacet distribution function
-// 0: Blinn
-// 1: Beckmann
-// 2: GGX
-#define PHYSICAL_SPEC_D		2
-
-// Geometric attenuation or shadowing
-// 0: Implicit
-// 1: Neumann
-// 2: Kelemen
-// 3: Schlick
-// 4: Smith (matched to GGX)
-#define PHYSICAL_SPEC_G		3
-
-// Fresnel
-// 0: None
-// 1: Schlick
-// 2: Fresnel
-#define PHYSICAL_SPEC_F		1
-
-float Pow5( float x )
-{
-	float x2 = x*x;
-	return x2 * x2 * x;
-}
-
 
 float3 Diffuse_Lambert( float3 DiffuseColor )
 {
@@ -215,55 +183,6 @@ float3 F_Fresnel( float3 SpecularColor, float VoH )
 }
 
 
-float3 Diffuse( float3 DiffuseColor, float Roughness, float NoV, float NoL, float VoH )
-{
-#if   PHYSICAL_DIFFUSE == 0
-	return Diffuse_Lambert( DiffuseColor );
-#elif PHYSICAL_DIFFUSE == 1
-	return Diffuse_Burley( DiffuseColor, Roughness, NoV, NoL, VoH );
-#elif PHYSICAL_DIFFUSE == 2
-	return Diffuse_OrenNayar( DiffuseColor, Roughness, NoV, NoL, VoH );
-#endif
-}
-
-float Distribution( float Roughness, float NoH )
-{
-#if   PHYSICAL_SPEC_D == 0
-	return D_Blinn( Roughness, NoH );
-#elif PHYSICAL_SPEC_D == 1
-	return D_Beckmann( Roughness, NoH );
-#elif PHYSICAL_SPEC_D == 2
-	return D_GGX( Roughness, NoH );
-#endif
-}
-
-// Vis = G / (4*NoL*NoV)
-float GeometricVisibility( float Roughness, float NoV, float NoL, float VoH )
-{
-#if   PHYSICAL_SPEC_G == 0
-	return Vis_Implicit();
-#elif PHYSICAL_SPEC_G == 1
-	return Vis_Neumann( NoV, NoL );
-#elif PHYSICAL_SPEC_G == 2
-	return Vis_Kelemen( VoH );
-#elif PHYSICAL_SPEC_G == 3
-	return Vis_Schlick( Roughness, NoV, NoL );
-#elif PHYSICAL_SPEC_G == 4
-	return Vis_Smith( Roughness, NoV, NoL );
-#endif
-}
-
-float3 Fresnel( float3 SpecularColor, float VoH )
-{
-#if   PHYSICAL_SPEC_F == 0
-	return F_None( SpecularColor );
-#elif PHYSICAL_SPEC_F == 1
-	return F_Schlick( SpecularColor, VoH );
-#elif PHYSICAL_SPEC_F == 2
-	return F_Fresnel( SpecularColor, VoH );
-#endif
-}
-
 
 //---------------
 // EnvBRDF
diff --git a/Engine/Shaders/BasePassPixelShader.usf b/Engine/Shaders/BasePassPixelShader.usf
--- a/Engine/Shaders/BasePassPixelShader.usf
+++ b/Engine/Shaders/BasePassPixelShader.usf
@@ -845,7 +845,6 @@ void Main(
 		GBuffer.IndirectIrradiance = IndirectIrradiance;
 		GBuffer.PrecomputedShadowFactors = PrecomputedShadowFactors;
 		GBuffer.GBufferAO = MaterialAO;
-		GBuffer.Opacity = Opacity;
 		GBuffer.HasDistanceFieldRepresentation = Primitive.HasDistanceFieldRepresentation;
 		GBuffer.HasHeightfieldRepresentation = Primitive.HasHeightfieldRepresentation;
 
@@ -855,18 +854,21 @@ void Main(
 			GBuffer.ShadingModelID = SHADINGMODELID_DEFAULT_LIT;
 		#elif MATERIAL_SHADINGMODEL_SUBSURFACE
 			GBuffer.ShadingModelID = SHADINGMODELID_SUBSURFACE;
-			GBuffer.CustomData = EncodeSubsurfaceColor(SubsurfaceColor);
+			GBuffer.CustomData.rgb = EncodeSubsurfaceColor(SubsurfaceColor);
+			GBuffer.CustomData.a = Opacity;
 		#elif MATERIAL_SHADINGMODEL_PREINTEGRATED_SKIN
 			GBuffer.ShadingModelID = SHADINGMODELID_PREINTEGRATED_SKIN;
-			GBuffer.CustomData = EncodeSubsurfaceColor(SubsurfaceColor);
+			GBuffer.CustomData.rgb = EncodeSubsurfaceColor(SubsurfaceColor);
+			GBuffer.CustomData.a = Opacity;
 		#elif MATERIAL_SHADINGMODEL_SUBSURFACE_PROFILE
 			GBuffer.ShadingModelID = SHADINGMODELID_SUBSURFACE_PROFILE;
-			GBuffer.CustomData = EncodeSubsurfaceProfile(SubsurfaceProfile);
+			GBuffer.CustomData.rgb = EncodeSubsurfaceProfile(SubsurfaceProfile);
+			GBuffer.CustomData.a = Opacity;
 		#elif MATERIAL_SHADINGMODEL_CLEAR_COAT
 			GBuffer.ShadingModelID = SHADINGMODELID_CLEAR_COAT;
 
-			float ClearCoat				= GetMaterialClearCoat(MaterialParameters);
-			float ClearCoatRoughness	= GetMaterialClearCoatRoughness(MaterialParameters);
+			float ClearCoat				= saturate( GetMaterialCustomData0(MaterialParameters) );
+			float ClearCoatRoughness	= saturate( GetMaterialCustomData1(MaterialParameters) );
 			float MetalSpec = 0.9;
 
 			float NoV = saturate( dot( MaterialParameters.WorldNormal, MaterialParameters.CameraVector ) );
@@ -885,7 +887,18 @@ void Main(
 			GBuffer.CustomData.y = ClearCoatRoughness;
 		#elif MATERIAL_SHADINGMODEL_TWOSIDED_FOLIAGE
 			GBuffer.ShadingModelID = SHADINGMODELID_TWOSIDED_FOLIAGE;
-			GBuffer.CustomData = EncodeSubsurfaceColor(SubsurfaceColor);
+			GBuffer.CustomData.rgb = EncodeSubsurfaceColor(SubsurfaceColor);
+			GBuffer.CustomData.a = Opacity;
+		#elif MATERIAL_SHADINGMODEL_HAIR
+			GBuffer.ShadingModelID = SHADINGMODELID_HAIR;
+			#if NUM_MATERIAL_OUTPUTS_GETTANGENTOUTPUT > 0
+				float3 Tangent = GetTangentOutput0(MaterialParameters);
+				Tangent = normalize( TransformTangentVectorToWorld( MaterialParameters.TangentToWorld, Tangent ) );
+				GBuffer.WorldNormal = Tangent;
+			#endif
+			GBuffer.CustomData.xy = UnitVectorToOctahedron( MaterialParameters.WorldNormal ) * 0.5 + 0.5;
+			GBuffer.CustomData.z = saturate( GetMaterialCustomData1(MaterialParameters) );	// Roughness2
+			GBuffer.CustomData.w = saturate( GetMaterialCustomData0(MaterialParameters) );	// Opacity
 		#else
 			// missing shading model, compiler should report ShadingModelID is not set
 		#endif
diff --git a/Engine/Shaders/Common.usf b/Engine/Shaders/Common.usf
--- a/Engine/Shaders/Common.usf
+++ b/Engine/Shaders/Common.usf
@@ -535,6 +535,118 @@ float4 Square( float4 x )
 	return x*x;
 }
 
+float Pow2( float x )
+{
+	return x*x;
+}
+
+float2 Pow2( float2 x )
+{
+	return x*x;
+}
+
+float3 Pow2( float3 x )
+{
+	return x*x;
+}
+
+float4 Pow2( float4 x )
+{
+	return x*x;
+}
+
+float Pow3( float x )
+{
+	return x*x*x;
+}
+
+float2 Pow3( float2 x )
+{
+	return x*x*x;
+}
+
+float3 Pow3( float3 x )
+{
+	return x*x*x;
+}
+
+float4 Pow3( float4 x )
+{
+	return x*x*x;
+}
+
+float Pow4( float x )
+{
+	float xx = x*x;
+	return xx * xx;
+}
+
+float2 Pow4( float2 x )
+{
+	float2 xx = x*x;
+	return xx * xx;
+}
+
+float3 Pow4( float3 x )
+{
+	float3 xx = x*x;
+	return xx * xx;
+}
+
+float4 Pow4( float4 x )
+{
+	float4 xx = x*x;
+	return xx * xx;
+}
+
+float Pow5( float x )
+{
+	float xx = x*x;
+	return xx * xx * x;
+}
+
+float2 Pow5( float2 x )
+{
+	float2 xx = x*x;
+	return xx * xx * x;
+}
+
+float3 Pow5( float3 x )
+{
+	float3 xx = x*x;
+	return xx * xx * x;
+}
+
+float4 Pow5( float4 x )
+{
+	float4 xx = x*x;
+	return xx * xx * x;
+}
+
+float Pow6( float x )
+{
+	float xx = x*x;
+	return xx * xx * xx;
+}
+
+float2 Pow6( float2 x )
+{
+	float2 xx = x*x;
+	return xx * xx * xx;
+}
+
+float3 Pow6( float3 x )
+{
+	float3 xx = x*x;
+	return xx * xx * xx;
+}
+
+float4 Pow6( float4 x )
+{
+	float4 xx = x*x;
+	return xx * xx * xx;
+}
+
 // Only valid for x >= 0
 MaterialFloat AtanFast( MaterialFloat x )
 {
diff --git a/Engine/Shaders/DeferredDecal.usf b/Engine/Shaders/DeferredDecal.usf
--- a/Engine/Shaders/DeferredDecal.usf
+++ b/Engine/Shaders/DeferredDecal.usf
@@ -337,7 +337,6 @@ void MainPS(
 	Data.IndirectIrradiance = 0;
 	Data.PrecomputedShadowFactors = 0;
 	Data.GBufferAO = 1;
-	Data.Opacity = Opacity;
 	Data.ShadingModelID = SHADINGMODELID_DEFAULT_LIT;
 	Data.HasDistanceFieldRepresentation = 1;
 	Data.HasHeightfieldRepresentation = 1;
diff --git a/Engine/Shaders/DeferredLightPixelShaders.usf b/Engine/Shaders/DeferredLightPixelShaders.usf
--- a/Engine/Shaders/DeferredLightPixelShaders.usf
+++ b/Engine/Shaders/DeferredLightPixelShaders.usf
@@ -42,6 +42,7 @@ FDeferredLightData SetupLightDataForStandardDeferred()
 void DirectionalPixelMain(
 	float2 InUV : TEXCOORD0,
 	float3 ScreenVector : TEXCOORD1,
+	float4 SVPos : SV_POSITION,
 	out float4 OutColor : SV_Target0
 	)
 {
@@ -57,7 +58,12 @@ void DirectionalPixelMain(
 		float3 WorldPosition = ScreenVector * SceneDepth + View.WorldCameraOrigin;
 
 		FDeferredLightData LightData = SetupLightDataForStandardDeferred();
-		OutColor = GetDynamicLighting(WorldPosition, CameraVector, InUV, ScreenSpaceData, ScreenSpaceData.GBuffer.ShadingModelID, LightData, GetPerPixelLightAttenuation(InUV), uint2(0,0));
+		
+		uint2 Random = ScrambleTEA( uint2( SVPos.xy ) );
+		Random.x ^= View.Random;
+		Random.y ^= View.Random;
+
+		OutColor = GetDynamicLighting(WorldPosition, CameraVector, InUV, ScreenSpaceData, ScreenSpaceData.GBuffer.ShadingModelID, LightData, GetPerPixelLightAttenuation(InUV), Random);
 	}
 }
 
diff --git a/Engine/Shaders/DeferredLightingCommon.usf b/Engine/Shaders/DeferredLightingCommon.usf
--- a/Engine/Shaders/DeferredLightingCommon.usf
+++ b/Engine/Shaders/DeferredLightingCommon.usf
@@ -119,8 +119,8 @@ float3 PointLightSpecularMIS( FScreenSpaceData ScreenSpaceData, FDeferredLightDa
 			
 			// Generalized microfacet specular
 			float D = D_GGX( Roughness, NoH );
-			float Vis = GeometricVisibility( Roughness, NoV, NoL, VoH );
-			float3 F = Fresnel( GBuffer.SpecularColor, VoH );
+			float Vis = Vis_SmithJointApprox( Roughness, NoV, NoL );
+			float3 F = F_Schlick( GBuffer.SpecularColor, VoH );
 
 			float ConePDF = 1.0 / ( 2 * PI * (1 - ConeCos) );
 			float GGXPDF = D * NoH / (4 * VoH);
@@ -353,9 +353,10 @@ float4 GetDynamicLighting(float3 WorldPosition, float3 CameraVector, float2 InUV
 			{
 #if 1	// for testing if there is a perf impact
 				// correct screen space subsurface scattering
-				float3 SurfaceLightingDiff = SurfaceShading(ScreenSpaceData.GBuffer, LobeRoughness, LobeEnergy, L, V, N, float2(1, 0));
-				float3 SurfaceLightingSpec = SurfaceShading(ScreenSpaceData.GBuffer, LobeRoughness, LobeEnergy, L, V, N, float2(0, 1));
+				float3 SurfaceLightingDiff = SurfaceShading(ScreenSpaceData.GBuffer, LobeRoughness, LobeEnergy, L, V, N, float2(1, 0), Random);
+				float3 SurfaceLightingSpec = SurfaceShading(ScreenSpaceData.GBuffer, LobeRoughness, LobeEnergy, L, V, N, float2(0, 1), Random);
 				LightAccumulator_Add(LightAccumulator, SurfaceLightingDiff, SurfaceLightingSpec, LightColor * (NoL * SurfaceAttenuation));
+				//LightAccumulator_Add(LightAccumulator, SurfaceLightingDiff, SurfaceLightingSpec, LightColor * (SurfaceAttenuation));
 #else
 				// wrong screen space subsurface scattering but potentially faster
 				float3 SurfaceLighting = SurfaceShading(ScreenSpaceData.GBuffer, LobeRoughness, LobeEnergy, L, V, N, float2(1, 1));
@@ -365,7 +366,7 @@ float4 GetDynamicLighting(float3 WorldPosition, float3 CameraVector, float2 InUV
 
 			// accumulate subsurface
 			{
-				float3 SubsurfaceLighting = SubsurfaceShading(ScreenSpaceData.GBuffer, L, V, N);
+				float3 SubsurfaceLighting = SubsurfaceShading(ScreenSpaceData.GBuffer, L, V, N, Random);
 
 				LightAccumulator_Add(LightAccumulator, SubsurfaceLighting, 0, LightColor * SubsurfaceAttenuation);
 
diff --git a/Engine/Shaders/DeferredShadingCommon.usf b/Engine/Shaders/DeferredShadingCommon.usf
--- a/Engine/Shaders/DeferredShadingCommon.usf
+++ b/Engine/Shaders/DeferredShadingCommon.usf
@@ -211,7 +211,8 @@ float Decode71(float Scalar, out uint Mask)
 #define SHADINGMODELID_CLEAR_COAT			4
 #define SHADINGMODELID_SUBSURFACE_PROFILE	5
 #define SHADINGMODELID_TWOSIDED_FOLIAGE		6
-#define SHADINGMODELID_NUM					7
+#define SHADINGMODELID_HAIR					7
+#define SHADINGMODELID_NUM					8
 
 float EncodeShadingModelId(uint ShadingModelId)
 {
@@ -228,7 +229,8 @@ bool IsSubsurfaceModel(int ShadingModel)
 	return ShadingModel == SHADINGMODELID_SUBSURFACE 
 		|| ShadingModel == SHADINGMODELID_PREINTEGRATED_SKIN 
 		|| ShadingModel == SHADINGMODELID_SUBSURFACE_PROFILE
-		|| ShadingModel == SHADINGMODELID_TWOSIDED_FOLIAGE;
+		|| ShadingModel == SHADINGMODELID_TWOSIDED_FOLIAGE
+		|| ShadingModel == SHADINGMODELID_HAIR;
 }
 
 // all values that are output by the forward rendering pass
@@ -247,7 +249,7 @@ struct FGBufferData
 	// 0..1
 	float Specular;
 	// 0..1
-	float3 CustomData;
+	float4 CustomData;
 	// Indirect irradiance luma
 	float IndirectIrradiance;
 	// Static shadow factors for channels assigned by Lightmass
@@ -255,8 +257,6 @@ struct FGBufferData
 	float4 PrecomputedShadowFactors;
 	// 0..1
 	float Roughness;
-	// 0..1
-	float Opacity;
 	// 0..1 ambient occlusion  e.g.SSAO, wet surface mask, skylight mask, ...
 	float GBufferAO;
 	// 0..255 
@@ -349,8 +349,7 @@ void EncodeGBuffer(
 #else
 		OutGBufferC.a = GBuffer.GBufferAO;
 #endif
-		OutGBufferD.r = GBuffer.Opacity;
-		OutGBufferD.gba = GBuffer.CustomData;
+		OutGBufferD = GBuffer.CustomData;
 	}
 
 	OutGBufferB.a = EncodeShadingModelId(GBuffer.ShadingModelID);
@@ -426,10 +425,9 @@ FGBufferData DecodeGBufferData(
 	GBuffer.Metallic = InGBufferB.r;
 	GBuffer.StoredSpecular = InGBufferB.g;
 	GBuffer.Specular = GBuffer.StoredSpecular.g;
-	GBuffer.Opacity = InGBufferD.r;
 	
 	GBuffer.Roughness = InGBufferB.b;
-	GBuffer.CustomData = InGBufferD.gba;
+	GBuffer.CustomData = InGBufferD;
 #if ALLOW_STATIC_LIGHTING
 	GBuffer.GBufferAO = 1;
 	GBuffer.IndirectIrradiance = DecodeIndirectIrradiance(InGBufferC.a);
@@ -524,7 +522,7 @@ void EncodeDBufferData(FGBufferData GBufferData, float3 MultiOpacity,
 
 float3 ExtractSubsurfaceColor(FGBufferData BufferData)
 {
-	return Square(BufferData.CustomData);
+	return Square(BufferData.CustomData.rgb);
 }
 
 uint ExtractSubsurfaceProfileInt(FGBufferData BufferData)
diff --git a/Engine/Shaders/DistanceFieldShadowing.usf b/Engine/Shaders/DistanceFieldShadowing.usf
--- a/Engine/Shaders/DistanceFieldShadowing.usf
+++ b/Engine/Shaders/DistanceFieldShadowing.usf
@@ -570,7 +570,7 @@ void DistanceFieldShadowingCS(
 			{
 				// Note: this has to match shadowmap logic
 				// Derive density from a heuristic using opacity, tweaked for useful falloff ranges and to give a linear depth falloff with opacity
-				SubsurfaceDensity = -.05f * log(1 - min(GBufferData.Opacity, .999f));
+				SubsurfaceDensity = -.05f * log(1 - min(GBufferData.CustomData.a, .999f));
 				bUseSubsurfaceTransmission = true;
 			}
 
diff --git a/Engine/Shaders/MaterialTemplate.usf b/Engine/Shaders/MaterialTemplate.usf
--- a/Engine/Shaders/MaterialTemplate.usf
+++ b/Engine/Shaders/MaterialTemplate.usf
@@ -1051,6 +1051,7 @@ float3 GetShadingModelColor(FGBufferData GBuffer)
 	else if (GBuffer.ShadingModelID == SHADINGMODELID_CLEAR_COAT) return float3(0.1f, 0.4f, 0.4f); 
 	else if (GBuffer.ShadingModelID == SHADINGMODELID_SUBSURFACE_PROFILE) return float3(0.2f, 0.6f, 0.5f); // Cyan
 	else if (GBuffer.ShadingModelID == SHADINGMODELID_TWOSIDED_FOLIAGE) return float3(0.2f, 0.2f, 0.8f); // Blue
+	else if (GBuffer.ShadingModelID == SHADINGMODELID_HAIR) return float3(0.6f, 0.1f, 0.5f);
 	else return float3(1.0f, 1.0f, 1.0f); // White
 #else
 	switch(GBuffer.ShadingModelID)
@@ -1062,6 +1063,7 @@ float3 GetShadingModelColor(FGBufferData GBuffer)
 		case SHADINGMODELID_CLEAR_COAT: return float3(0.1f, 0.4f, 0.4f); // Brown
 		case SHADINGMODELID_SUBSURFACE_PROFILE: return float3(0.2f, 0.6f, 0.5f); // Cyan
 		case SHADINGMODELID_TWOSIDED_FOLIAGE: return float3(0.2f, 0.2f, 0.8f); // Cyan
+		case SHADINGMODELID_HAIR: return float3(0.6f, 0.1f, 0.5f);
 		default: return float3(1.0f, 1.0f, 1.0f); // White
 	}
 #endif
@@ -1117,7 +1119,7 @@ float4 SceneTextureLookup(float2 UV, int SceneTextureIndex, bool bFiltered)
 	else if (SceneTextureIndex == 7)	return ScreenSpaceData.GBuffer.Metallic;
 	else if (SceneTextureIndex == 8)	return float4(ScreenSpaceData.GBuffer.WorldNormal, 0);
 	else if (SceneTextureIndex == 9)	return 1; // todo
-	else if (SceneTextureIndex == 10)	return ScreenSpaceData.GBuffer.Opacity;
+	else if (SceneTextureIndex == 10)	return ScreenSpaceData.GBuffer.CustomData.a;
 	else if (SceneTextureIndex == 11)	return ScreenSpaceData.GBuffer.Roughness;
 	else if (SceneTextureIndex == 12)	return ScreenSpaceData.GBuffer.GBufferAO;
 	else if (SceneTextureIndex == 13)	return ScreenSpaceData.GBuffer.CustomDepth;
@@ -1150,9 +1152,7 @@ float4 SceneTextureLookup(float2 UV, int SceneTextureIndex, bool bFiltered)
 		// PPI_SpecularColor
 		case 3: return float4(ScreenSpaceData.GBuffer.SpecularColor, 0);
 		// PPI_SubsurfaceColor
-		case 4: 
-			float3 CustomData = IsSubsurfaceModel(ScreenSpaceData.GBuffer.ShadingModelID) ? ExtractSubsurfaceColor(ScreenSpaceData.GBuffer) : ScreenSpaceData.GBuffer.CustomData;
-			return float4(CustomData, 0);
+		case 4: return IsSubsurfaceModel(ScreenSpaceData.GBuffer.ShadingModelID) ? float4( ExtractSubsurfaceColor(ScreenSpaceData.GBuffer), ScreenSpaceData.GBuffer.CustomData.a ) : ScreenSpaceData.GBuffer.CustomData;
 		// PPI_BaseColor
 		case 5: return float4(ScreenSpaceData.GBuffer.BaseColor, 0);
 		// PPI_Specular
@@ -1164,7 +1164,7 @@ float4 SceneTextureLookup(float2 UV, int SceneTextureIndex, bool bFiltered)
 		// PPI_SeparateTranslucency
 		case 9: return float4(1, 1, 1, 1);	// todo
 		// PPI_Opacity
-		case 10: return ScreenSpaceData.GBuffer.Opacity;
+		case 10: return ScreenSpaceData.GBuffer.CustomData.a;
 		// PPI_Roughness
 		case 11: return ScreenSpaceData.GBuffer.Roughness;
 		// PPI_MaterialAO
@@ -1315,7 +1315,7 @@ half GetMaterialRoughness(FMaterialPixelParameters Parameters)
 	}
 	#endif
 	
-	return max(0.04f, Roughness);
+	return Roughness;
 #endif
 }
 
@@ -1443,26 +1443,16 @@ half4 GetMaterialSubsurfaceData(FMaterialPixelParameters Parameters)
 	return OutSubsurface;
 }
 
-half GetMaterialClearCoatRaw(FMaterialPixelParameters Parameters)
+half GetMaterialCustomData0(FMaterialPixelParameters Parameters)
 {
 %s;
 }
 
-half GetMaterialClearCoat(FMaterialPixelParameters Parameters)
-{
-	return saturate(GetMaterialClearCoatRaw(Parameters));
-}
-
-half GetMaterialClearCoatRoughnessRaw(FMaterialPixelParameters Parameters)
+half GetMaterialCustomData1(FMaterialPixelParameters Parameters)
 {
 %s;
 }
 
-half GetMaterialClearCoatRoughness(FMaterialPixelParameters Parameters)
-{
-	return saturate(GetMaterialClearCoatRoughnessRaw(Parameters));
-}
-
 half GetMaterialAmbientOcclusionRaw(FMaterialPixelParameters Parameters)
 {
 %s;
diff --git a/Engine/Shaders/PostProcessSubsurface.usf b/Engine/Shaders/PostProcessSubsurface.usf
--- a/Engine/Shaders/PostProcessSubsurface.usf
+++ b/Engine/Shaders/PostProcessSubsurface.usf
@@ -34,8 +34,7 @@ float GetSubsurfaceStrength(float2 UV)
 	FScreenSpaceData ScreenSpaceData = GetScreenSpaceData(UV);
 
 	float Mask = ScreenSpaceData.GBuffer.ShadingModelID == SHADINGMODELID_SUBSURFACE_PROFILE;
-
-	return Mask * ScreenSpaceData.GBuffer.Opacity;
+	return Mask * ScreenSpaceData.GBuffer.CustomData.a;
 }
 
 // 0:full res (slower but reference) / 1:faster half res with some quality loss
@@ -141,7 +140,7 @@ float ComputeFullResLerp(FScreenSpaceData ScreenSpaceData, float2 UVSceneColor,
 	//
 	Ret *= saturate(PixelSizeRadius - PixelSize);
 	// opacity allows to scale the radius - at some point we should fade in the full resolution, we don't have a masking other than that.
-	Ret *= saturate(ScreenSpaceData.GBuffer.Opacity * 10);
+	Ret *= saturate(ScreenSpaceData.GBuffer.CustomData.a * 10);
 	// todo: Subsurface has some non scatter contribution - all that should come from the Full res
 
 	return Ret;
diff --git a/Engine/Shaders/ShadingModels.usf b/Engine/Shaders/ShadingModels.usf
--- a/Engine/Shaders/ShadingModels.usf
+++ b/Engine/Shaders/ShadingModels.usf
@@ -194,19 +194,412 @@ float3 ClearCoatShading( FGBufferData GBuffer, float3 LobeRoughness, float3 Lobe
 #endif //USE_CLEARCOAT
 }
 
+
+float Hair_g( float B, float Theta )
+{
+	return exp( -0.5 * Pow2( Theta ) / (B*B) ) / ( sqrt(2*PI) * B );
+}
+
+float Hair_F( float CosTheta )
+{
+	const float n = 1.55;
+	const float F0 = Pow2( (1 - n) / (1 + n) );
+	return F0 + (1 - F0) * Pow5( 1 - CosTheta );
+}
+
+#define HAIR_REFERENCE	1
+#if HAIR_REFERENCE
+struct FHairTemp
+{
+	float SinThetaL;
+	float SinThetaV;
+	float CosThetaD;
+	float CosThetaT;
+	float CosPhi;
+	float CosHalfPhi;
+	float VoL;
+	float n_prime;
+};
+
+// Modified Bessel function
+float I0( float x )
+{
+	x = abs(x);
+	float a;
+	if( x < 3.75 )
+	{
+		float t = x / 3.75;
+		float t2 = t*t;
+		a =        + 0.0045813;
+		a = a * t2 + 0.0360768;
+		a = a * t2 + 0.2659732;
+		a = a * t2 + 1.2067492;
+		a = a * t2 + 3.0899424;
+		a = a * t2 + 3.5156229;
+		a = a * t2 + 1.0;
+	}
+	else
+	{
+		float t = 3.75 / x;
+		a =       + 0.00392377;
+		a = a * t - 0.01647633;
+		a = a * t + 0.02635537;
+		a = a * t - 0.02057706;
+		a = a * t + 0.00916281;
+		a = a * t - 0.00157565;
+		a = a * t + 0.00225319;
+		a = a * t + 0.01328592;
+		a = a * t + 0.39894228;
+		a *= exp(x) * rsqrt(x);
+	}
+	return a;
+}
+
+float LongitudinalScattering( float B, float SinThetaL, float SinThetaV )
+{
+	float v = B*B;
+	float CosThetaL2 = 1 - SinThetaL * SinThetaL;
+	float CosThetaV2 = 1 - SinThetaV * SinThetaV;
+	float Mp = 0;
+	if( v < 0.1 )
+	{
+		float a = sqrt(CosThetaL2 * CosThetaV2) / v;
+		float b = -SinThetaL * SinThetaV / v;
+		float logI0a = a > 12 ? a + 0.5 * ( -log(2*PI) + log(1/a) + 0.125/a ) : log( I0(a) );
+		Mp = exp( logI0a + b - rcp(v) + 0.6931 + log(0.5 / v) );
+	}
+	else
+	{
+		Mp = rcp( exp(2/v) * v - v ) * exp( ( 1 - SinThetaL * SinThetaV ) / v ) * I0( sqrt(CosThetaL2 * CosThetaV2) / v );
+	}
+	
+	return Mp;
+}
+
+float GaussianDetector( float Bp, float Phi )
+{
+	float Dp = 0;
+	for( int k = -4; k <= 4; k++ )
+	{
+		// TODO use symmetry and detect for both Phi and -Phi
+		Dp += Hair_g( Bp, Phi - (2*PI) * k );
+	}
+	return Dp;
+}
+
+float3 Attenuation( uint p, float h, float3 Color, FHairTemp HairTemp )
+{
+	float3 A;
+	if( p == 0 )
+	{
+		//A = F( cos( 0.5 * acos( HairTemp.VoL ) ) );
+		A = Hair_F( sqrt( 0.5 + 0.5 * HairTemp.VoL ) );
+	}
+	else
+	{
+		// ua is absorption
+		// ua = pe*Sigma_ae + pp*Sigma_ap
+		float3 Sigma_ae = { 0.419, 0.697, 1.37 };
+		float3 Sigma_ap = { 0.187, 0.4, 1.05 };
+		//float3 ua = 0.25 * Sigma_ae + 0.25 * Sigma_ap;
+		float3 ua = -0.25 * log( Color );
+		float3 ua_prime = ua / HairTemp.CosThetaT;
+		//float3 ua_prime = ua / sqrt( 1 - Pow2( HairTemp.CosThetaD ) / 2.4 );
+		
+		float yi = asin(h);
+		float yt = asin(h / HairTemp.n_prime);
+		
+		float f = Hair_F( HairTemp.CosThetaD * sqrt( 1- h*h ) );		// (14)
+		float3 T = exp( -2 * ua_prime * ( 1 + cos(2*yt) ) );
+		A = Pow2(1 - f) * pow( f, p - 1 ) * pow( T, p );		// (13)
+	}
+	return A;
+}
+
+float Omega( uint p, float h, FHairTemp HairTemp )
+{
+	float yi = asin(h);
+	float yt = asin(h / HairTemp.n_prime);
+	return 2*p*yt - 2*yi + p*PI;
+}
+
+float3 AzimuthalScattering( uint p, float Bp, float3 Color, FHairTemp HairTemp, uint2 Random )
+{
+	float Phi = acos( HairTemp.CosPhi );
+	
+	// Np = 0.5 * Integral_-1^1( A(p,h) * Dp( Phi - Omega(p,h) ) * dh )
+
+#if 0
+
+#if 0
+	// Gauss–Legendre quadrature order 5
+	uint Num = 2;
+	float2 w0_x0 = { 0.56888888, 0 };
+	float2 wi_xi[] =
+	{
+		{ 0.47862867, 0.53846931 },
+		{ 0.23692688, 0.90617984 },
+	};
+#elif 0
+	// Gauss–Legendre quadrature order 7
+	const uint Num = 3;
+	float2 w0_x0 = { 0.41795918, 0 };
+	float2 wi_xi[] =
+	{
+		{ 0.38183005, 0.40584515 },
+		{ 0.27970539, 0.74153118 },
+		{ 0.12948496, 0.94910791 },
+	};
+#elif 0
+	// Gauss–Legendre quadrature order 11
+	const uint Num = 5;
+	float2 w0_x0 = { 0.27292508, 0 };
+	float2 wi_xi[] =
+	{
+		{ 0.26280454, 0.26954315 },
+		{ 0.23319376, 0.51909612 },
+		{ 0.18629021, 0.73015200 },
+		{ 0.12558036, 0.88706259 },
+		{ 0.05566856, 0.97822865 },
+	};
+#else
+	// Gauss–Legendre quadrature order 35
+	const uint Num = 17;
+	float2 w0_x0 = { 0.08848679, 0.00000000 };
+	float2 wi_xi[] =
+	{
+		{ 0.08814053, 0.08837134 },
+		{ 0.08710444, 0.17605106 },
+		{ 0.08538665, 0.26235294 },
+		{ 0.08300059, 0.34660155 },
+		{ 0.07996494, 0.42813754 },
+		{ 0.07630345, 0.50632277 },
+		{ 0.07204479, 0.58054534 },
+		{ 0.06722228, 0.65022436 },
+		{ 0.06187367, 0.71481450 },
+		{ 0.05604081, 0.77381025 },
+		{ 0.04976937, 0.82674989 },
+		{ 0.04310842, 0.87321912 },
+		{ 0.03611011, 0.91285426 },
+		{ 0.02882926, 0.94534514 },
+		{ 0.02132297, 0.97043761 },
+		{ 0.01365082, 0.98793576 },
+		{ 0.00588343, 0.99770656 },
+	};
+#endif
+
+	float w = w0_x0.x;
+	float h = w0_x0.y;
+	float3 Np = w * Attenuation( p, h, HairTemp ) * GaussianDetector( Bp, Phi - Omega( p, h, HairTemp ) );
+	
+	for( uint i = 0; i < Num; i++ )
+	{
+		w = wi_xi[i].x;
+		h = wi_xi[i].y;
+		Np += w * Attenuation( p, h, HairTemp ) * GaussianDetector( Bp, Phi - Omega( p, h, HairTemp ) );
+		Np += w * Attenuation( p,-h, HairTemp ) * GaussianDetector( Bp, Phi - Omega( p,-h, HairTemp ) );
+	}
+#else
+	float Offset = float( Random.x & 0xffff ) / (1<<16);
+
+	uint Num = 16;
+	float3 Np = 0;
+	for( uint i = 0; i < Num; i++ )
+	{
+		float h = ( (float)( i + Offset ) / Num ) * 2 - 1;
+		Np += Attenuation( p, h, Color, HairTemp ) * GaussianDetector( Bp, Phi - Omega( p, h, HairTemp ) );
+	}
+	Np *= 2.0 / Num;
+#endif
+
+	return 0.5 * Np;
+}
+
+// [d'Eon et al. 2011, "An Energy-Conserving Hair Reflectance Model"]
+// [d'Eon et al. 2014, "A Fiber Scattering Model with Non-Separable Lobes"]
+float3 HairShadingRef( FGBufferData GBuffer, float3 L, float3 V, half3 N, uint2 Random )
+{
+	float n = 1.55;
+
+	FHairTemp HairTemp;
+	
+	// N is the vector parallel to hair pointing toward root
+	HairTemp.VoL       = dot(V,L);
+	HairTemp.SinThetaL = dot(N,L);
+	HairTemp.SinThetaV = dot(N,V);
+	// SinThetaT = 1/n * SinThetaL
+	HairTemp.CosThetaT = sqrt( 1 - Pow2( (1/n) * HairTemp.SinThetaL ) );
+	HairTemp.CosThetaD = cos( 0.5 * abs( asin( HairTemp.SinThetaV ) - asin( HairTemp.SinThetaL ) ) );
+
+	float3 Lp = L - HairTemp.SinThetaL * N;
+	float3 Vp = V - HairTemp.SinThetaV * N;
+	HairTemp.CosPhi = dot(Lp,Vp) * rsqrt( dot(Lp,Lp) * dot(Vp,Vp) );
+	HairTemp.CosHalfPhi = sqrt( 0.5 + 0.5 * HairTemp.CosPhi );
+
+	HairTemp.n_prime = sqrt( n*n - 1 + Pow2( HairTemp.CosThetaD ) ) / HairTemp.CosThetaD;
+
+	float Alpha[] =
+	{
+		-0.13,
+		0.13 * 0.5,
+		0.13 * 1.5,
+	};
+	float B[] =
+	{
+		Pow2( GBuffer.Roughness ),
+		Pow2( GBuffer.Roughness ) / 2,
+		Pow2( GBuffer.CustomData.z ) * 2,
+	};
+
+	float3 S = 0;
+	UNROLL for( uint p = 0; p < 3; p++ )
+	{
+		float SinThetaV = HairTemp.SinThetaV;
+		float Bp = B[p];
+		if( p == 0 )
+		{
+			Bp *= sqrt(2) * HairTemp.CosHalfPhi;
+			float sa, ca;
+			sincos( Alpha[p], sa, ca );
+			SinThetaV -= 2*sa * ( HairTemp.CosHalfPhi * ca * sqrt( 1 - SinThetaV * SinThetaV ) + sa * SinThetaV );
+		}
+		else
+		{
+			SinThetaV = sin( asin(SinThetaV) - Alpha[p] );
+		}
+		float Mp = LongitudinalScattering( Bp, HairTemp.SinThetaL, SinThetaV );
+		float3 Np = AzimuthalScattering( p, B[p], GBuffer.BaseColor, HairTemp, Random );
+
+		float3 Sp = Mp * Np;
+		S += Sp;
+	}
+	return S;
+}
+#endif
+
+// Approximation to HairShadingRef using concepts from the following papers:
+// [Marschner et al. 2003, "Light Scattering from Human Hair Fibers"]
+// [Pekelis et al. 2015, "A Data-Driven Light Scattering Model for Hair"]
+float3 HairShading( FGBufferData GBuffer, float3 L, float3 V, half3 N, uint2 Random )
+{
+	const float3 DiffuseN	= OctahedronToUnitVector( GBuffer.CustomData.xy * 2 - 1 );
+	const float Roughness2	= GBuffer.CustomData.z;
+	const float Opacity		= GBuffer.CustomData.w;
+
+#if HAIR_REFERENCE
+	if( DiffuseN.x > 0 )
+	{
+		return HairShadingRef( GBuffer, L, V, N, Random );
+	}
+#endif
+	
+	// N is the vector parallel to hair pointing toward root
+
+	const float VoL       = dot(V,L);
+	const float SinThetaL = dot(N,L);
+	const float SinThetaV = dot(N,V);
+	const float CosThetaD = cos( 0.5 * abs( asin( SinThetaV ) - asin( SinThetaL ) ) );
+
+	const float3 Lp = L - SinThetaL * N;
+	const float3 Vp = V - SinThetaV * N;
+	const float CosPhi = dot(Lp,Vp) * rsqrt( dot(Lp,Lp) * dot(Vp,Vp) );
+	const float CosHalfPhi = sqrt( 0.5 + 0.5 * CosPhi );
+	const float Phi = acos( CosPhi );
+
+	float n = 1.55;
+	//float n_prime = sqrt( n*n - 1 + Pow2( CosThetaD ) ) / CosThetaD;
+	float n_prime = 1.2 / CosThetaD + 0.35 * CosThetaD;
+
+	float Alpha[] =
+	{
+		-0.13,
+		0.13 * 0.5,
+		0.13 * 1.5,
+	};
+	float B[] =
+	{
+		Pow2( GBuffer.Roughness ),
+		Pow2( GBuffer.Roughness ) / 2,
+		Pow2( Roughness2 ) * 2,
+	};
+
+	float3 S = 0;
+
+	// R
+	{
+		const float sa = sin( Alpha[0] );
+		const float ca = cos( Alpha[0] );
+		float Shift = (2*sa*ca) * CosHalfPhi * sqrt( 1 - SinThetaV * SinThetaV ) + (2*sa*sa) * SinThetaV;
+
+		float Mp = Hair_g( B[0] * sqrt(2) * CosHalfPhi, SinThetaL + SinThetaV - Shift );
+		float Np = 0.25 * CosHalfPhi;
+		float Fp = Hair_F( sqrt( 0.5 + 0.5 * VoL ) );
+		S += Mp * Np * Fp * ( GBuffer.Specular * 2 );
+	}
+
+	// TT
+	{
+		float Mp = Hair_g( B[1], SinThetaL + SinThetaV - Alpha[1] );
+
+		float a = 1 / n_prime;
+		float h = CosHalfPhi * rsqrt( 1 + a*a - 2*a * sqrt( 0.5 - 0.5 * CosPhi ) );
+		//float h = CosHalfPhi * ( ( 1 - Pow2( CosHalfPhi ) ) * a + 1 );
+		//float h = 0.4;
+		float yi = asin(h);
+		float yt = asin(h / n_prime);
+		
+		float f = Hair_F( CosThetaD * sqrt( 1 - h*h ) );
+		float Fp = Pow2(1 - f);
+		float3 Tp = pow( GBuffer.BaseColor, 0.5 * ( 1 + cos(2*yt) ) / CosThetaD );
+
+		//float t = asin( 1 / n_prime );
+		//float d = ( sqrt(2) - t ) / ( 1 - t );
+		//float s = -0.5 * PI * (1 - 1 / n_prime) * log( 2*d - 1 - 2 * sqrt( d * (d - 1) ) );
+		float s = 0.3;
+		float Np = exp( (Phi - PI) / s ) / ( s * Pow2( 1 + exp( (Phi - PI) / s ) ) );
+
+		S += Mp * Np * Fp * Tp;
+	}
+
+	// TRT
+	{
+		float Mp = Hair_g( B[2], SinThetaL + SinThetaV - Alpha[2] );
+		
+		//float h = 0.75;
+		float f = Hair_F( CosThetaD * 0.66 );
+		float Fp = Pow2(1 - f) * f;
+		float3 Tp = pow( GBuffer.BaseColor, 1.6 / CosThetaD );
+
+		//float s = 0.3;
+		//float Np = exp( Phi / s ) / ( s * Pow2( 1 + exp( Phi / s ) ) );
+		float Np = (1/PI) * 2.6 * exp( 2 * 2.6 * (CosPhi - 1) );
+
+		S += Mp * Np * Fp * Tp;
+	}
+
+
+	// Hack approximation for multiple scattering.
+	float Wrap = 1 - Opacity;
+	float NoL = saturate( ( dot(DiffuseN, L) + Wrap ) / Square( 1 + Wrap ) );
+	float DiffuseScatter = GBuffer.GBufferAO * NoL * (1 / PI) * GBuffer.Metallic;
+	S += GBuffer.BaseColor * DiffuseScatter;
+
+	return S;
+}
+
 float3 SubsurfaceShadingSubsurface( FGBufferData GBuffer, float3 L, float3 V, half3 N )
 {
 	float3 SubsurfaceColor = ExtractSubsurfaceColor(GBuffer);
+	float Opacity = GBuffer.CustomData.a;
 
 	float3 H = normalize(V + L);
 
 	// to get an effect when you see through the material
 	// hard coded pow constant
-	float InScatter = pow(saturate(dot(L, -V)), 12) * lerp(3, .1f, GBuffer.Opacity);
+	float InScatter = pow(saturate(dot(L, -V)), 12) * lerp(3, .1f, Opacity);
 	// wrap around lighting, /(PI*2) to be energy consistent (hack do get some view dependnt and light dependent effect)
-	float OpacityFactor = GBuffer.Opacity;
 	// Opacity of 0 gives no normal dependent lighting, Opacity of 1 gives strong normal contribution
-	float NormalContribution = saturate(dot(N, H) * OpacityFactor + 1 - OpacityFactor);
+	float NormalContribution = saturate(dot(N, H) * Opacity + 1 - Opacity);
 	float BackScatter = GBuffer.GBufferAO * NormalContribution / (PI * 2);
 
 	// lerp to never exceed 1 (energy conserving)
@@ -234,14 +627,14 @@ SamplerState	PreIntegratedBRDFSampler;
 float3 SubsurfaceShadingPreintegratedSkin( FGBufferData GBuffer, float3 L, float3 V, half3 N )
 {
 	float3 SubsurfaceColor = ExtractSubsurfaceColor(GBuffer);
+	float Opacity = GBuffer.CustomData.a;
 
-	float OpacityFactor = GBuffer.Opacity;
-	float3 PreintegratedBRDF = Texture2DSampleLevel(PreIntegratedBRDF, PreIntegratedBRDFSampler, float2(saturate(dot(N, L) * .5 + .5), 1 - OpacityFactor), 0).rgb;
+	float3 PreintegratedBRDF = Texture2DSampleLevel(PreIntegratedBRDF, PreIntegratedBRDFSampler, float2(saturate(dot(N, L) * .5 + .5), 1 - Opacity), 0).rgb;
 	return PreintegratedBRDF * SubsurfaceColor;
 }
 
 // @param DiffSpecMask .r: diffuse, .g:specular e.g. float2(1,1) for both, float2(1,0) for diffuse only
-float3 SurfaceShading( FGBufferData GBuffer, float3 LobeRoughness, float3 LobeEnergy, float3 L, float3 V, half3 N, float2 DiffSpecMask )
+float3 SurfaceShading( FGBufferData GBuffer, float3 LobeRoughness, float3 LobeEnergy, float3 L, float3 V, half3 N, float2 DiffSpecMask, uint2 Random )
 {
 	switch( GBuffer.ShadingModelID )
 	{
@@ -260,7 +653,7 @@ float3 SurfaceShading( FGBufferData GBuffer, float3 LobeRoughness, float3 LobeEn
 	}
 }
 
-float3 SubsurfaceShading( FGBufferData GBuffer, float3 L, float3 V, half3 N )
+float3 SubsurfaceShading( FGBufferData GBuffer, float3 L, float3 V, half3 N, uint2 Random )
 {
 	float3 SubsurfaceColor = ExtractSubsurfaceColor(GBuffer);
 
@@ -272,6 +665,8 @@ float3 SubsurfaceShading( FGBufferData GBuffer, float3 L, float3 V, half3 N )
 			return SubsurfaceShadingPreintegratedSkin( GBuffer, L, V, N );
 		case SHADINGMODELID_TWOSIDED_FOLIAGE:
 			return SubsurfaceShadingTwoSided( SubsurfaceColor, L, V, N );
+		case SHADINGMODELID_HAIR:
+			return HairShading( GBuffer, L, V, N, Random );
 		default:
 			return 0;
 	}
diff --git a/Engine/Shaders/ShadowProjectionPixelShader.usf b/Engine/Shaders/ShadowProjectionPixelShader.usf
--- a/Engine/Shaders/ShadowProjectionPixelShader.usf
+++ b/Engine/Shaders/ShadowProjectionPixelShader.usf
@@ -106,13 +106,14 @@ void Main(
 		BRANCH
 		if (IsSubsurfaceModel(InGBufferData.ShadingModelID))
 		{
+			float Opacity = InGBufferData.CustomData.a;
 			// Derive density from a heuristic using opacity, tweaked for useful falloff ranges and to give a linear depth falloff with opacity
-			float Density = -.05f * log(1 - min(InGBufferData.Opacity, .999f));
+			float Density = -.05f * log(1 - min(Opacity, .999f));
 			//@todo - Use a larger kernel filter based on material opacity
 			// This is not being done currently because a varying opacity causes discontinuities in the shadow result
 			// Note: Snapping based on the largest size kernel is done in FProjectedShadowInfo::FProjectedShadowInfo, MaxDownsampleFactor is set to 4
 			// Only factors of MaxDownsampleFactor will be stable
-			float SquareRootFilterScale = lerp(1.999f, 0, InGBufferData.Opacity);
+			float SquareRootFilterScale = lerp(1.999f, 0, Opacity);
 			int SquareRootFilterScaleInt = SquareRootFilterScale + 1;
 
 			#if UNFILTERED_SHADOW_PROJECTION
diff --git a/Engine/Source/Editor/MaterialEditor/Private/MaterialEditor.cpp b/Engine/Source/Editor/MaterialEditor/Private/MaterialEditor.cpp
--- a/Engine/Source/Editor/MaterialEditor/Private/MaterialEditor.cpp
+++ b/Engine/Source/Editor/MaterialEditor/Private/MaterialEditor.cpp
@@ -3544,7 +3544,8 @@ void FMaterialEditor::NotifyPostChange( const FPropertyChangedEvent& PropertyCha
 				SetPreviewAsset(GUnrealEd->GetThumbnailManager()->EditorSphere);
 			}
 		}
-		else if ( NameOfPropertyThatChanged == GET_MEMBER_NAME_CHECKED(UMaterial, MaterialDomain) )
+		else if( NameOfPropertyThatChanged == GET_MEMBER_NAME_CHECKED(UMaterial, MaterialDomain) ||
+				 NameOfPropertyThatChanged == GET_MEMBER_NAME_CHECKED(UMaterial, ShadingModel) )
 		{
 			Material->MaterialGraph->RebuildGraph();
 			TArray<TWeakObjectPtr<UObject>> SelectedObjects = MaterialDetailsView->GetSelectedObjects();
diff --git a/Engine/Source/Editor/UnrealEd/Classes/MaterialGraph/MaterialGraph.h b/Engine/Source/Editor/UnrealEd/Classes/MaterialGraph/MaterialGraph.h
--- a/Engine/Source/Editor/UnrealEd/Classes/MaterialGraph/MaterialGraph.h
+++ b/Engine/Source/Editor/UnrealEd/Classes/MaterialGraph/MaterialGraph.h
@@ -168,4 +168,6 @@ private:
 
 	FText GetEmissivePinName() const;
 	FText GetBaseColorPinName() const;
+	FText GetMetallicPinName() const;
+	FText GetCustomDataPinName( uint32 Index ) const;
 };
diff --git a/Engine/Source/Editor/UnrealEd/Private/MaterialGraph.cpp b/Engine/Source/Editor/UnrealEd/Private/MaterialGraph.cpp
--- a/Engine/Source/Editor/UnrealEd/Private/MaterialGraph.cpp
+++ b/Engine/Source/Editor/UnrealEd/Private/MaterialGraph.cpp
@@ -32,7 +32,7 @@ void UMaterialGraph::RebuildGraph()
 	{
 		// Initialize the material input list.
 		MaterialInputs.Add( FMaterialInputInfo( GetBaseColorPinName(), MP_BaseColor ) );	
-		MaterialInputs.Add( FMaterialInputInfo( LOCTEXT("Metallic", "Metallic"), MP_Metallic ) );
+		MaterialInputs.Add( FMaterialInputInfo( GetMetallicPinName(), MP_Metallic ) );
 		MaterialInputs.Add( FMaterialInputInfo( LOCTEXT("Specular", "Specular"), MP_Specular ) );
 		MaterialInputs.Add( FMaterialInputInfo( LOCTEXT("Roughness", "Roughness"), MP_Roughness ) );
 		MaterialInputs.Add( FMaterialInputInfo( GetEmissivePinName(), MP_EmissiveColor ) );
@@ -43,8 +43,8 @@ void UMaterialGraph::RebuildGraph()
 		MaterialInputs.Add( FMaterialInputInfo( LOCTEXT("WorldDisplacement", "World Displacement"), MP_WorldDisplacement ) );
 		MaterialInputs.Add( FMaterialInputInfo( LOCTEXT("TessellationMultiplier", "Tessellation Multiplier"), MP_TessellationMultiplier ) );
 		MaterialInputs.Add( FMaterialInputInfo( LOCTEXT("SubsurfaceColor", "Subsurface Color"), MP_SubsurfaceColor ) );
-		MaterialInputs.Add( FMaterialInputInfo( LOCTEXT("ClearCoat", "Clear Coat"), MP_ClearCoat ) );
-		MaterialInputs.Add( FMaterialInputInfo( LOCTEXT("ClearCoatRoughness", "Clear Coat Roughness"), MP_ClearCoatRoughness ) );
+		MaterialInputs.Add( FMaterialInputInfo( GetCustomDataPinName(0), MP_CustomData0 ) );
+		MaterialInputs.Add( FMaterialInputInfo( GetCustomDataPinName(1), MP_CustomData1 ) );
 		MaterialInputs.Add( FMaterialInputInfo( LOCTEXT("AmbientOcclusion", "Ambient Occlusion"), MP_AmbientOcclusion ) );
 		MaterialInputs.Add( FMaterialInputInfo( LOCTEXT("Refraction", "Refraction"), MP_Refraction) );
 
@@ -465,4 +465,39 @@ FText UMaterialGraph::GetBaseColorPinName() const
 	return LOCTEXT("BaseColor", "Base Color");
 }
 
+FText UMaterialGraph::GetMetallicPinName() const
+{
+	return Material->GetShadingModel() == MSM_Hair ? LOCTEXT("Diffuse", "Diffuse") : LOCTEXT("Metallic", "Metallic");
+}
+
+FText UMaterialGraph::GetCustomDataPinName( uint32 Index ) const
+{
+	if( Index == 0 )
+	{
+		switch( Material->GetShadingModel() )
+		{
+		case MSM_ClearCoat:
+			return LOCTEXT("ClearCoat", "Clear Coat");
+		case MSM_Hair:
+			return LOCTEXT("HairOpacity", "Hair Opacity");
+		default:
+			return LOCTEXT("CustomData0", "Custom Data 0");
+		}
+	}
+	else if( Index == 1 )
+	{
+		switch( Material->GetShadingModel() )
+		{
+		case MSM_ClearCoat:
+			return LOCTEXT("ClearCoatRoughness", "Clear Coat Roughness");
+		case MSM_Hair:
+			return LOCTEXT("SecondaryRoughness", "Secondary Roughness");
+		default:
+			return LOCTEXT("CustomData1", "Custom Data 1");
+		}
+	}
+
+	return LOCTEXT("CustomData", "Custom Data");	
+}
+
 #undef LOCTEXT_NAMESPACE
diff --git a/Engine/Source/Editor/UnrealEd/Private/MaterialGraphNode.cpp b/Engine/Source/Editor/UnrealEd/Private/MaterialGraphNode.cpp
--- a/Engine/Source/Editor/UnrealEd/Private/MaterialGraphNode.cpp
+++ b/Engine/Source/Editor/UnrealEd/Private/MaterialGraphNode.cpp
@@ -244,6 +244,11 @@ FLinearColor UMaterialGraphNode::GetNodeTitleColor() const
 		// Previously FColor(255, 155, 0);
 		return Settings->ResultNodeTitleColor;
 	}
+	else if (MaterialExpression->IsA(UMaterialExpressionCustomOutput::StaticClass()))
+	{
+		// Previously FColor(255, 155, 0);
+		return Settings->ResultNodeTitleColor;
+	}
 	else if (UMaterial::IsParameter(MaterialExpression))
 	{
 		if (Material->HasDuplicateParameters(MaterialExpression))
diff --git a/Engine/Source/Runtime/Engine/Classes/Engine/EngineTypes.h b/Engine/Source/Runtime/Engine/Classes/Engine/EngineTypes.h
--- a/Engine/Source/Runtime/Engine/Classes/Engine/EngineTypes.h
+++ b/Engine/Source/Runtime/Engine/Classes/Engine/EngineTypes.h
@@ -208,7 +208,8 @@ enum EMaterialShadingModel
 	MSM_PreintegratedSkin	UMETA(DisplayName="Preintegrated Skin"),
 	MSM_ClearCoat			UMETA(DisplayName="Clear Coat"),
 	MSM_SubsurfaceProfile	UMETA(DisplayName="Subsurface Profile"),
-	MSM_TwoSidedFoliage		UMETA(DisplayName="Two Sided Foliage"),
+	MSM_TwoSidedFoliage		UMETA(DisplayName="Two Sided"),
+	MSM_Hair				UMETA(DisplayName="Hair"),
 	MSM_MAX,
 };
 
diff --git a/Engine/Source/Runtime/Engine/Classes/Materials/MaterialExpressionTangentOutput.h b/Engine/Source/Runtime/Engine/Classes/Materials/MaterialExpressionTangentOutput.h
new file mode 100644
--- /dev/null
+++ b/Engine/Source/Runtime/Engine/Classes/Materials/MaterialExpressionTangentOutput.h
@@ -0,0 +1,25 @@
+// Copyright 1998-2015 Epic Games, Inc. All Rights Reserved.
+
+#pragma once
+
+#include "MaterialExpressionCustomOutput.h"
+#include "MaterialExpressionTangentOutput.generated.h"
+
+UCLASS(collapsecategories, hidecategories=Object, MinimalAPI)
+class UMaterialExpressionTangentOutput : public UMaterialExpressionCustomOutput
+{
+	GENERATED_UCLASS_BODY()
+
+	UPROPERTY(meta = (RequiredInput = "true"))
+	FExpressionInput Input;
+
+	virtual int32 Compile(class FMaterialCompiler* Compiler, int32 OutputIndex, int32 MultiplexIndex) override;
+	virtual void GetCaption(TArray<FString>& OutCaptions) const override;
+
+#if WITH_EDITOR
+	virtual uint32 GetInputType(int32 InputIndex) override { return MCT_Float; }
+#endif
+
+	virtual int32 GetNumOutputs() const override { return 1; }
+	virtual FString GetFunctionName() const override { return TEXT("GetTangentOutput"); }
+};
\ No newline at end of file
diff --git a/Engine/Source/Runtime/Engine/Private/Materials/HLSLMaterialTranslator.h b/Engine/Source/Runtime/Engine/Private/Materials/HLSLMaterialTranslator.h
--- a/Engine/Source/Runtime/Engine/Private/Materials/HLSLMaterialTranslator.h
+++ b/Engine/Source/Runtime/Engine/Private/Materials/HLSLMaterialTranslator.h
@@ -315,9 +315,9 @@ public:
 				Chunk[MP_SubsurfaceColor] = AppendVector(SubsurfaceColor, CodeSubsurfaceProfile);		
 			}
 
-			Chunk[MP_ClearCoat]						= Material->CompilePropertyAndSetMaterialProperty(MP_ClearCoat			   ,this);
-			Chunk[MP_ClearCoatRoughness]			= Material->CompilePropertyAndSetMaterialProperty(MP_ClearCoatRoughness    ,this);
-			Chunk[MP_AmbientOcclusion]				= Material->CompilePropertyAndSetMaterialProperty(MP_AmbientOcclusion      ,this);
+			Chunk[MP_CustomData0]					= Material->CompilePropertyAndSetMaterialProperty(MP_CustomData0		,this);
+			Chunk[MP_CustomData1]					= Material->CompilePropertyAndSetMaterialProperty(MP_CustomData1		,this);
+			Chunk[MP_AmbientOcclusion]				= Material->CompilePropertyAndSetMaterialProperty(MP_AmbientOcclusion	,this);
 
 			if(IsTranslucentBlendMode(Material->GetBlendMode()))
 			{
@@ -699,8 +699,8 @@ public:
 		LazyPrintf.PushParam(*FString::Printf(TEXT("return %.5f"),Material->GetMaxDisplacement()));
 		LazyPrintf.PushParam(*GenerateFunctionCode(MP_TessellationMultiplier));
 		LazyPrintf.PushParam(*GenerateFunctionCode(MP_SubsurfaceColor));
-		LazyPrintf.PushParam(*GenerateFunctionCode(MP_ClearCoat));
-		LazyPrintf.PushParam(*GenerateFunctionCode(MP_ClearCoatRoughness));
+		LazyPrintf.PushParam(*GenerateFunctionCode(MP_CustomData0));
+		LazyPrintf.PushParam(*GenerateFunctionCode(MP_CustomData1));
 		LazyPrintf.PushParam(*GenerateFunctionCode(MP_AmbientOcclusion));
 		LazyPrintf.PushParam(*GenerateFunctionCode(MP_Refraction));
 
diff --git a/Engine/Source/Runtime/Engine/Private/Materials/Material.cpp b/Engine/Source/Runtime/Engine/Private/Materials/Material.cpp
--- a/Engine/Source/Runtime/Engine/Private/Materials/Material.cpp
+++ b/Engine/Source/Runtime/Engine/Private/Materials/Material.cpp
@@ -3393,8 +3393,8 @@ FExpressionInput* UMaterial::GetExpressionInputForProperty(EMaterialProperty InP
 		case MP_WorldDisplacement:		return &WorldDisplacement;
 		case MP_TessellationMultiplier:	return &TessellationMultiplier;
 		case MP_SubsurfaceColor:		return &SubsurfaceColor;
-		case MP_ClearCoat:				return &ClearCoat;
-		case MP_ClearCoatRoughness:		return &ClearCoatRoughness;
+		case MP_CustomData0:			return &ClearCoat;
+		case MP_CustomData1:			return &ClearCoatRoughness;
 		case MP_AmbientOcclusion:		return &AmbientOcclusion;
 		case MP_Refraction:				return &Refraction;
 		case MP_MaterialAttributes:		return &MaterialAttributes;
@@ -3708,8 +3708,8 @@ int32 UMaterial::CompilePropertyEx( FMaterialCompiler* Compiler, EMaterialProper
 		case MP_Specular:				return Specular.CompileWithDefault(Compiler, Property);
 		case MP_Roughness:				return Roughness.CompileWithDefault(Compiler, Property);
 		case MP_TessellationMultiplier:	return TessellationMultiplier.CompileWithDefault(Compiler, Property);
-		case MP_ClearCoat:				return ClearCoat.CompileWithDefault(Compiler, Property);
-		case MP_ClearCoatRoughness:		return ClearCoatRoughness.CompileWithDefault(Compiler, Property);
+		case MP_CustomData0:			return ClearCoat.CompileWithDefault(Compiler, Property);
+		case MP_CustomData1:			return ClearCoatRoughness.CompileWithDefault(Compiler, Property);
 		case MP_AmbientOcclusion:		return AmbientOcclusion.CompileWithDefault(Compiler, Property);
 		case MP_Refraction:				return Refraction.CompileWithDefault(Compiler, Property);
 		case MP_EmissiveColor:			return EmissiveColor.CompileWithDefault(Compiler, Property);
@@ -3835,8 +3835,9 @@ EMaterialShadingModel UMaterial::GetShadingModel(bool bIsInGameThread) const
 	switch (MaterialDomain)
 	{
 		case MD_Surface:
-		case MD_DeferredDecal:
 			return ShadingModel;
+		case MD_DeferredDecal:
+			return MSM_DefaultLit;
 
 		// Post process and light function materials must be rendered with the unlit model.
 		case MD_PostProcess:
@@ -4016,9 +4017,9 @@ bool UMaterial::IsPropertyActive(EMaterialProperty InProperty) const
 	case MP_SubsurfaceColor:
 		Active = ShadingModel == MSM_Subsurface || ShadingModel == MSM_PreintegratedSkin || ShadingModel == MSM_TwoSidedFoliage;
 		break;
-	case MP_ClearCoat:
-	case MP_ClearCoatRoughness:
-		Active = ShadingModel == MSM_ClearCoat;
+	case MP_CustomData0:
+	case MP_CustomData1:
+		Active = ShadingModel == MSM_ClearCoat || ShadingModel == MSM_Hair;
 		break;
 	case MP_TessellationMultiplier:
 	case MP_WorldDisplacement:
diff --git a/Engine/Source/Runtime/Engine/Private/Materials/MaterialExpressions.cpp b/Engine/Source/Runtime/Engine/Private/Materials/MaterialExpressions.cpp
--- a/Engine/Source/Runtime/Engine/Private/Materials/MaterialExpressions.cpp
+++ b/Engine/Source/Runtime/Engine/Private/Materials/MaterialExpressions.cpp
@@ -111,6 +111,7 @@
 #include "Materials/MaterialExpressionStaticBool.h"
 #include "Materials/MaterialExpressionStaticSwitch.h"
 #include "Materials/MaterialExpressionSubtract.h"
+#include "Materials/MaterialExpressionTangentOutput.h"
 #include "Materials/MaterialExpressionTextureBase.h"
 #include "Materials/MaterialExpressionTextureObject.h"
 #include "Materials/MaterialExpressionTextureProperty.h"
@@ -3338,8 +3339,8 @@ int32 UMaterialExpressionMakeMaterialAttributes::Compile(class FMaterialCompiler
 	case MP_WorldDisplacement: Ret = WorldDisplacement.Compile(Compiler); Expression = WorldDisplacement.Expression; break;
 	case MP_TessellationMultiplier: Ret = TessellationMultiplier.Compile(Compiler); Expression = TessellationMultiplier.Expression; break;
 	case MP_SubsurfaceColor: Ret = SubsurfaceColor.Compile(Compiler); Expression = SubsurfaceColor.Expression; break;
-	case MP_ClearCoat: Ret = ClearCoat.Compile(Compiler); Expression = ClearCoat.Expression; break;
-	case MP_ClearCoatRoughness: Ret = ClearCoatRoughness.Compile(Compiler); Expression = ClearCoatRoughness.Expression; break;
+	case MP_CustomData0: Ret = ClearCoat.Compile(Compiler); Expression = ClearCoat.Expression; break;
+	case MP_CustomData1: Ret = ClearCoatRoughness.Compile(Compiler); Expression = ClearCoatRoughness.Expression; break;
 	case MP_AmbientOcclusion: Ret = AmbientOcclusion.Compile(Compiler); Expression = AmbientOcclusion.Expression; break;
 	case MP_Refraction: Ret = Refraction.Compile(Compiler); Expression = Refraction.Expression; break;
 	case MP_PixelDepthOffset: Ret = PixelDepthOffset.Compile(Compiler); Expression = PixelDepthOffset.Expression; break;
@@ -9499,4 +9500,47 @@ void UMaterialExpressionEyeAdaptation::GetCaption(TArray<FString>& OutCaptions)
 	OutCaptions.Add(FString(TEXT("EyeAdaptation")));
 }
 
+
+//
+// UMaterialExpressionTangentOutput
+//
+UMaterialExpressionTangentOutput::UMaterialExpressionTangentOutput(const FObjectInitializer& ObjectInitializer)
+: Super(ObjectInitializer)
+{
+	// Structure to hold one-time initialization
+	struct FConstructorStatics
+	{
+		FText NAME_Custom;
+		FConstructorStatics()
+			: NAME_Custom(LOCTEXT( "Custom", "Custom" ))
+		{
+		}
+	};
+	static FConstructorStatics ConstructorStatics;
+
+	MenuCategories.Add(ConstructorStatics.NAME_Custom);
+
+	// No outputs
+	Outputs.Reset();
+}
+
+int32 UMaterialExpressionTangentOutput::Compile(class FMaterialCompiler* Compiler, int32 OutputIndex, int32 MultiplexIndex)
+{
+	if( Input.Expression )
+	{
+		return Compiler->CustomOutput(this, OutputIndex, Input.Compile(Compiler, MultiplexIndex));
+	}
+	else
+	{
+		return CompilerError(Compiler, TEXT("Input missing"));
+	}
+
+	return INDEX_NONE;
+}
+
+void UMaterialExpressionTangentOutput::GetCaption(TArray<FString>& OutCaptions) const
+{
+	OutCaptions.Add(TEXT("Tangent output"));
+}
+
 #undef LOCTEXT_NAMESPACE
diff --git a/Engine/Source/Runtime/Engine/Private/Materials/MaterialShared.cpp b/Engine/Source/Runtime/Engine/Private/Materials/MaterialShared.cpp
--- a/Engine/Source/Runtime/Engine/Private/Materials/MaterialShared.cpp
+++ b/Engine/Source/Runtime/Engine/Private/Materials/MaterialShared.cpp
@@ -209,8 +209,8 @@ EMaterialValueType GetMaterialPropertyType(EMaterialProperty Property)
 	case MP_WorldDisplacement : return MCT_Float3;
 	case MP_TessellationMultiplier: return MCT_Float;
 	case MP_SubsurfaceColor: return MCT_Float3;
-	case MP_ClearCoat: return MCT_Float;
-	case MP_ClearCoatRoughness: return MCT_Float;
+	case MP_CustomData0: return MCT_Float;
+	case MP_CustomData1: return MCT_Float;
 	case MP_AmbientOcclusion: return MCT_Float;
 	case MP_Refraction: return MCT_Float;
 	case MP_MaterialAttributes: return MCT_MaterialAttributes;
@@ -1252,6 +1252,7 @@ void FMaterial::SetupMaterialEnvironment(
 		case MSM_SubsurfaceProfile: OutEnvironment.SetDefine(TEXT("MATERIAL_SHADINGMODEL_SUBSURFACE_PROFILE"),	TEXT("1")); break;
 		case MSM_ClearCoat:			OutEnvironment.SetDefine(TEXT("MATERIAL_SHADINGMODEL_CLEAR_COAT"),			TEXT("1")); break;
 		case MSM_TwoSidedFoliage:	OutEnvironment.SetDefine(TEXT("MATERIAL_SHADINGMODEL_TWOSIDED_FOLIAGE"),	TEXT("1")); break;
+		case MSM_Hair:				OutEnvironment.SetDefine(TEXT("MATERIAL_SHADINGMODEL_HAIR"),				TEXT("1")); break;
 		default: 
 			UE_LOG(LogMaterial, Warning, TEXT("Unknown material shading model: %u  Setting to MSM_DefaultLit"),(int32)GetShadingModel());
 			OutEnvironment.SetDefine(TEXT("MATERIAL_SHADINGMODEL_DEFAULT_LIT"),TEXT("1"));
@@ -2164,8 +2165,8 @@ const TMap<EMaterialProperty, int32> CreatePropertyToIOIndexMap()
 	Ret.Add(MP_WorldDisplacement, 9);
 	Ret.Add(MP_TessellationMultiplier, 10);
 	Ret.Add(MP_SubsurfaceColor, 11);
-	Ret.Add(MP_ClearCoat, 12);
-	Ret.Add(MP_ClearCoatRoughness, 13);
+	Ret.Add(MP_CustomData0, 12);
+	Ret.Add(MP_CustomData1, 13);
 	Ret.Add(MP_AmbientOcclusion, 14);
 	Ret.Add(MP_Refraction, 15);
 	Ret.Add(MP_CustomizedUVs0, 16);
@@ -2212,8 +2213,8 @@ int32 GetDefaultExpressionForMaterialProperty(FMaterialCompiler* Compiler, EMate
 		case MP_Specular:				return Compiler->Constant(0.5f);
 		case MP_Roughness:				return Compiler->Constant(0.5f);
 		case MP_TessellationMultiplier:	return Compiler->Constant(1.0f);
-		case MP_ClearCoat:				return Compiler->Constant(1.0f);
-		case MP_ClearCoatRoughness:		return Compiler->Constant(0.1f);
+		case MP_CustomData0:			return Compiler->Constant(1.0f);
+		case MP_CustomData1:			return Compiler->Constant(0.1f);
 		case MP_AmbientOcclusion:		return Compiler->Constant(1.0f);
 		case MP_PixelDepthOffset:		return Compiler->Constant(0.0f);
 
@@ -2265,8 +2266,8 @@ FString GetNameOfMaterialProperty(EMaterialProperty Property)
 	case MP_WorldDisplacement:		return TEXT("WorldDisplacement");
 	case MP_TessellationMultiplier: return TEXT("TessellationMultiplier");
 	case MP_SubsurfaceColor:		return TEXT("SubsurfaceColor");
-	case MP_ClearCoat:				return TEXT("ClearCoat");
-	case MP_ClearCoatRoughness:		return TEXT("ClearCoatRoughness");
+	case MP_CustomData0:			return TEXT("ClearCoat");
+	case MP_CustomData1:			return TEXT("ClearCoatRoughness");
 	case MP_AmbientOcclusion:		return TEXT("AmbientOcclusion");
 	case MP_Refraction:				return TEXT("Refraction");
 	case MP_PixelDepthOffset:		return TEXT("PixelDepthOffset");
diff --git a/Engine/Source/Runtime/Engine/Public/SceneTypes.h b/Engine/Source/Runtime/Engine/Public/SceneTypes.h
--- a/Engine/Source/Runtime/Engine/Public/SceneTypes.h
+++ b/Engine/Source/Runtime/Engine/Public/SceneTypes.h
@@ -140,8 +140,8 @@ enum EMaterialProperty
 	MP_WorldDisplacement,
 	MP_TessellationMultiplier,
 	MP_SubsurfaceColor,
-	MP_ClearCoat,
-	MP_ClearCoatRoughness,
+	MP_CustomData0,
+	MP_CustomData1,
 	MP_AmbientOcclusion,
 	MP_Refraction,
 	MP_CustomizedUVs0,
