Jira: UE-14401  Improved Vis tool visualization for shadows.

diff --git a/Engine/Shaders/VisualizeTexture.usf b/Engine/Shaders/VisualizeTexture.usf
--- a/Engine/Shaders/VisualizeTexture.usf
+++ b/Engine/Shaders/VisualizeTexture.usf
@@ -28,7 +28,7 @@ SamplerState VisualizeTextureCubeArraySampler;
 
 // [0]: r:RGBAdd, g:SingleChannelMul, b:Add, a:FracScale
 // [1]: r:BlinkState=0/1, g:Frac when 0, saturate when 1, b:array index, a:custom mip
-// [2]: r:0=normal 1=depth, g:unwrapcubemap, b:singlechannel(0-3), a:unused
+// [2]: r:0=normal 1=depth, 2=shadow depth, g:unwrapcubemap, b:singlechannel(0-3), a:unused
 float4 VisualizeParam[3];
 // x, y (0 for 2D textures), z(only !=0 if 3d texture), w unused
 float4 TextureExtent;
@@ -92,26 +92,35 @@ void VisualizeTexturePS(
 	#endif
 #endif
 	
-
+	// Normal = 0, SceneDepth = 1 and ShadowDepth = 2
 	int InputValueMapping = (int)VisualizeParam[2].r;
-	if(InputValueMapping == 1)
+	if(InputValueMapping != 0)
 	{
-		// current z buffer is negated (1 is near, 0 is far)
 		float Depth = TexLookup.r;
+		if (InputValueMapping == 1)  // SceneDepth 
+		{
+			// current z buffer is negated (1 is near, 0 is far)
 #if HAS_INVERTED_Z_BUFFER
-		Depth = 1 - Depth;
+			Depth = 1 - Depth;
 #endif
 
 #if FEATURE_LEVEL >= FEATURE_LEVEL_SM4
-		// map depth to nice colors
-		float Scale = 64.0; // This 64 looks like a hardware limit.
-		Depth = (exp2(Depth * Scale) - 1.0) * (1.0 / (exp2(Scale) - 1.0));
-		Depth *= Depth;
-		Depth *= Depth;
-
-		float Value = Depth * Depth;
-		TexLookup = float4(pow(Value, 0.25f), pow(Value, 1), pow(Value, 4), 0);		
-
+			// map depth to nice colors
+			float Scale = 64.0; // This 64 looks like a hardware limit.
+			Depth = (exp2(Depth * Scale) - 1.0) * (1.0 / (exp2(Scale) - 1.0));
+			Depth *= Depth;
+			Depth *= Depth;
+
+			float Value = Depth * Depth;
+			TexLookup = float4(pow(Value, 0.25f), pow(Value, 1), pow(Value, 4), 0);
+		} 
+		else if (InputValueMapping == 2) // ShadowDepth
+		{
+			// linearly blend from White (near) to Teal (far).
+			float4 TealBlue = float4(54.f/255.f, 117.f/255.f, 136.f/255.f, 0.f);
+			float4 White = float4(1.0f, 1.0f, 1.0f, 0.f);
+			TexLookup.rgb = lerp(White, TealBlue, Depth);
+		}
 #if TEXTURE_TYPE == 6
 		// inside #ifdef to avoid D3DDebug error:
 		//    The resource return type for component 0 declared in the shader code (UINT) is not compatible with the Shader Resource View format bound to slot 1 of the Pixel Shader unit (UNORM).  This mismatch is invalid if the shader actually uses the view (e.g. it is not skipped due to shader code branching). [ EXECUTION ERROR #361: DEVICE_DRAW_RESOURCE_RETURN_TYPE_MISMATCH]
diff --git a/Engine/Source/Runtime/Renderer/Private/PostProcess/VisualizeTexture.cpp b/Engine/Source/Runtime/Renderer/Private/PostProcess/VisualizeTexture.cpp
--- a/Engine/Source/Runtime/Renderer/Private/PostProcess/VisualizeTexture.cpp
+++ b/Engine/Source/Runtime/Renderer/Private/PostProcess/VisualizeTexture.cpp
@@ -414,8 +414,10 @@ void FVisualizeTexture::GenerateContent(FRHICommandListImmediate& RHICmdList, co
 
 	FVisualizeTextureData VisualizeTextureData(RenderTargetItem, Desc);
 
-	bool bDepthTexture = (Desc.TargetableFlags & TexCreate_DepthStencilTargetable) != 0;
-	
+	// distinguish between standard depth and shadow depth to produce more reasonable default value mapping in the pixel shader.
+	const bool bDepthTexture = (Desc.TargetableFlags & TexCreate_DepthStencilTargetable) != 0;
+	const bool bShadowDepth = (Desc.Format == PF_ShadowDepth);
+
 	VisualizeTextureData.RGBMul = RGBMul;
 	VisualizeTextureData.SingleChannelMul = SingleChannelMul;
 	VisualizeTextureData.SingleChannel = SingleChannel;
@@ -423,7 +425,7 @@ void FVisualizeTexture::GenerateContent(FRHICommandListImmediate& RHICmdList, co
 	VisualizeTextureData.Tex00 = Tex00;
 	VisualizeTextureData.Tex11 = Tex11;
 	VisualizeTextureData.bSaturateInsteadOfFrac = (Flags & 1) != 0;
-	VisualizeTextureData.InputValueMapping = bDepthTexture ? 1 : 0;
+	VisualizeTextureData.InputValueMapping = bShadowDepth ? 2 : (bDepthTexture ? 1 : 0);
 	VisualizeTextureData.ArrayIndex = ArrayIndex;
 	VisualizeTextureData.CustomMip = CustomMip;
 	VisualizeTextureData.StencilSRV = StencilSRV;
@@ -641,6 +643,18 @@ void FVisualizeTexture::PresentContent(FRHICommandListImmediate& RHICmdList, con
 	{
 		Canvas.DrawShadowedString( X, Y += YStep, TEXT("Blinking Red: <0"), GetStatsFont(), FLinearColor(1,0,0));
 		Canvas.DrawShadowedString( X, Y += YStep, TEXT("Blinking Blue: NAN or Inf"), GetStatsFont(), FLinearColor(0,0,1));
+
+		// add explicit legend for SceneDepth and ShadowDepth as the display coloring is an artificial choice. 
+		const bool bDepthTexture = (Desc.TargetableFlags & TexCreate_DepthStencilTargetable) != 0;
+		const bool bShadowDepth = (Desc.Format == PF_ShadowDepth);
+		if (bShadowDepth)
+		{
+			Canvas.DrawShadowedString(X, Y += YStep, TEXT("Color Key: Linear with white near and teal distant"), GetStatsFont(), FLinearColor(54.f / 255.f, 117.f / 255.f, 136.f / 255.f));
+		} 
+		else if (bDepthTexture)  
+		{
+			Canvas.DrawShadowedString(X, Y += YStep, TEXT("Color Key: Nonlinear with white distant"), GetStatsFont(), FLinearColor(0.5, 0, 0));
+		}
 	}
 
 	Canvas.Flush_RenderThread(RHICmdList);
