Fixed GL. Disabled hair reference.

diff --git a/Engine/Shaders/ShadingModels.usf b/Engine/Shaders/ShadingModels.usf
--- a/Engine/Shaders/ShadingModels.usf
+++ b/Engine/Shaders/ShadingModels.usf
@@ -207,7 +207,7 @@ float Hair_F( float CosTheta )
 	return F0 + (1 - F0) * Pow5( 1 - CosTheta );
 }
 
-#define HAIR_REFERENCE	1
+#define HAIR_REFERENCE	0
 #if HAIR_REFERENCE
 struct FHairTemp
 {
@@ -458,7 +458,7 @@ float3 HairShadingRef( FGBufferData GBuffer, float3 L, float3 V, half3 N, uint2
 		float Bp = B[p];
 		if( p == 0 )
 		{
-			Bp *= sqrt(2) * HairTemp.CosHalfPhi;
+			Bp *= sqrt(2.0) * HairTemp.CosHalfPhi;
 			float sa, ca;
 			sincos( Alpha[p], sa, ca );
 			SinThetaV -= 2*sa * ( HairTemp.CosHalfPhi * ca * sqrt( 1 - SinThetaV * SinThetaV ) + sa * SinThetaV );
@@ -531,7 +531,7 @@ float3 HairShading( FGBufferData GBuffer, float3 L, float3 V, half3 N, uint2 Ran
 		const float ca = cos( Alpha[0] );
 		float Shift = (2*sa*ca) * CosHalfPhi * sqrt( 1 - SinThetaV * SinThetaV ) + (2*sa*sa) * SinThetaV;
 
-		float Mp = Hair_g( B[0] * sqrt(2) * CosHalfPhi, SinThetaL + SinThetaV - Shift );
+		float Mp = Hair_g( B[0] * sqrt(2.0) * CosHalfPhi, SinThetaL + SinThetaV - Shift );
 		float Np = 0.25 * CosHalfPhi;
 		float Fp = Hair_F( sqrt( 0.5 + 0.5 * VoL ) );
 		S += Mp * Np * Fp * ( GBuffer.Specular * 2 );
