Fixed dbuffer compile error.

diff --git a/Engine/Shaders/DeferredDecal.usf b/Engine/Shaders/DeferredDecal.usf
--- a/Engine/Shaders/DeferredDecal.usf
+++ b/Engine/Shaders/DeferredDecal.usf
@@ -363,7 +363,7 @@ void MainPS(
 		// before base pass (DBuffer)
 
 		// @param MultiOpacity .x: Color, .y:Normal, .z:Roughness
-		float3 MultiOpacity = Data.Opacity;
+		float3 MultiOpacity = Opacity;
 
 		EncodeDBufferData(Data, MultiOpacity, OutTarget0, OutTarget1, OutTarget2);
 	}
