Fix PS4 compile error

diff --git a/Engine/Shaders/MaterialTemplate.usf b/Engine/Shaders/MaterialTemplate.usf
--- a/Engine/Shaders/MaterialTemplate.usf
+++ b/Engine/Shaders/MaterialTemplate.usf
@@ -1113,7 +1113,7 @@ float4 SceneTextureLookup(float2 UV, int SceneTextureIndex, bool bFiltered)
 	else if (SceneTextureIndex == 1)	return ScreenSpaceData.GBuffer.Depth;
 	else if (SceneTextureIndex == 2)	return float4(ScreenSpaceData.GBuffer.DiffuseColor, 0);
 	else if (SceneTextureIndex == 3)	return float4(ScreenSpaceData.GBuffer.SpecularColor, 0);
-	else if (SceneTextureIndex == 4)	return float4(ScreenSpaceData.GBuffer.CustomData, 0);
+	else if (SceneTextureIndex == 4)	return float4(ScreenSpaceData.GBuffer.CustomData.rgb, 0);
 	else if (SceneTextureIndex == 5)	return float4(ScreenSpaceData.GBuffer.BaseColor, 0);
 	else if (SceneTextureIndex == 6)	return ScreenSpaceData.GBuffer.Specular;
 	else if (SceneTextureIndex == 7)	return ScreenSpaceData.GBuffer.Metallic;
