Jira UE-1399.  Fixed a compliation error that resulted from mismatched brackets under certain #if conditions.

diff --git a/Engine/Shaders/VisualizeTexture.usf b/Engine/Shaders/VisualizeTexture.usf
--- a/Engine/Shaders/VisualizeTexture.usf
+++ b/Engine/Shaders/VisualizeTexture.usf
@@ -130,6 +130,7 @@ void VisualizeTexturePS(
 #endif
 
 #else
+		}
 		TexLookup.r = Depth;
 #endif
 	}
