Merging

//depot/UE4-Orion/Engine/Shaders/...

to //depot/UE4-Orion/Engine/Shaders/...

diff --git a/Engine/Shaders/FilterPixelShader.usf b/Engine/Shaders/FilterPixelShader.usf
--- a/Engine/Shaders/FilterPixelShader.usf
+++ b/Engine/Shaders/FilterPixelShader.usf
@@ -57,7 +57,7 @@ int SampleCount;
 # if (ES2_PROFILE || ES3_1_PROFILE)
 
 void Main(
-	float2 InUV : TEXCOORD0,
+	noperspective float2 InUV : TEXCOORD0,
 	out half4 OutColor : SV_Target0
 	)
 {
@@ -93,7 +93,7 @@ void Main(
 # else //!(ES2_PROFILE || ES3_1_PROFILE)
 
 void Main(
-	float2 InUV : TEXCOORD0,
+	noperspective float2 InUV : TEXCOORD0,
 	out half4 OutColor : SV_Target0
 	)
 {
@@ -138,8 +138,8 @@ half4 SampleWeights[NUM_SAMPLES];
 # if (ES2_PROFILE || ES3_1_PROFILE)
 
 void Main(
-	float2 InUV : TEXCOORD0,
-	float2 InOffsetUVs[NUM_SAMPLES] : TEXCOORD1,
+	noperspective float2 InUV : TEXCOORD0,
+	noperspective float2 InOffsetUVs[NUM_SAMPLES] : TEXCOORD1,
 	out half4 OutColor : SV_Target0
 	)
 {
@@ -159,8 +159,8 @@ void Main(
 # else //!(ES2_PROFILE || ES3_1_PROFILE)
 
 void Main(
-	float2 InUV : TEXCOORD0,
-	float4 InOffsetUVs[(NUM_SAMPLES + 1) / 2] : TEXCOORD1,
+	noperspective float2 InUV : TEXCOORD0,
+	noperspective float4 InOffsetUVs[(NUM_SAMPLES + 1) / 2] : TEXCOORD1,
 	out float4 OutColor : SV_Target0
 	)
 {
@@ -199,7 +199,7 @@ void Main(
 // currently only used for downsampling, doesn't need to be masked as it doesn't read outside bounds
 // reads 4 samples
 void DownsampleScene(
-	float4 InOffsetUVs[2] : TEXCOORD0,
+	noperspective float4 InOffsetUVs[2] : TEXCOORD0,
 	out float4 OutColor : SV_Target0
 	)
 {
diff --git a/Engine/Shaders/PostProcessAmbient.usf b/Engine/Shaders/PostProcessAmbient.usf
--- a/Engine/Shaders/PostProcessAmbient.usf
+++ b/Engine/Shaders/PostProcessAmbient.usf
@@ -444,7 +444,7 @@ float3 ApproximateSpecularIBL( uint2 Random, float3 SpecularColor, float Roughne
 	return PrefilteredColor * ( SpecularColor * AB.x + AB.y );
 }
 
-void MainPS(in float4 UVAndScreenPos : TEXCOORD0, out float4 OutColor : SV_Target0)
+void MainPS(in noperspective float4 UVAndScreenPos : TEXCOORD0, out float4 OutColor : SV_Target0)
 {
 	float2 UV = UVAndScreenPos.xy;
 
diff --git a/Engine/Shaders/PostProcessAmbientOcclusion.usf b/Engine/Shaders/PostProcessAmbientOcclusion.usf
--- a/Engine/Shaders/PostProcessAmbientOcclusion.usf
+++ b/Engine/Shaders/PostProcessAmbientOcclusion.usf
@@ -201,7 +201,7 @@ float ComputeDepthSimilarity(float DepthA, float DepthB, float TweakScale)
 }
 
 // downsample the input of the ambient occlusion pass for better performance, can take input from setup or another downsample pass
-void MainSetupPS(in float4 UVAndScreenPos : TEXCOORD0, out float4 OutColor0 : SV_Target0)
+void MainSetupPS(in noperspective float4 UVAndScreenPos : TEXCOORD0, out float4 OutColor0 : SV_Target0)
 {
 	float2 InUV = UVAndScreenPos.xy;
 
@@ -500,7 +500,7 @@ float3 ReconstructNormal(float2 In)
 }
 
 // the main pixel shader that computes ambient occlusion
-void MainPS(in float4 UVAndScreenPos : TEXCOORD0, out float4 OutColor : SV_Target0)
+void MainPS(in noperspective float4 UVAndScreenPos : TEXCOORD0, out float4 OutColor : SV_Target0)
 {
 	OutColor = 0;
 
@@ -718,7 +718,7 @@ void MainPS(in float4 UVAndScreenPos : TEXCOORD0, out float4 OutColor : SV_Targe
 
 
 // used if StaticFraction > small number
-void BasePassAOPS(in float4 UVAndScreenPos : TEXCOORD0, out float4 OutColor : SV_Target0)
+void BasePassAOPS(in noperspective float4 UVAndScreenPos : TEXCOORD0, out float4 OutColor : SV_Target0)
 {
 	float2 UV = UVAndScreenPos.xy;
 
diff --git a/Engine/Shaders/PostProcessBloom.usf b/Engine/Shaders/PostProcessBloom.usf
--- a/Engine/Shaders/PostProcessBloom.usf
+++ b/Engine/Shaders/PostProcessBloom.usf
@@ -25,7 +25,7 @@ void MainPostprocessCommonVS(
 }
 
 void VisualizeBloomSetupPS(
-	  float4 UVAndScreenPos : TEXCOORD0
+	noperspective float4 UVAndScreenPos : TEXCOORD0
 	, out float4 OutColor : SV_Target0)
 {
 	float2 UV = UVAndScreenPos.xy;
@@ -56,7 +56,7 @@ void VisualizeBloomSetupPS(
 }
 
 void VisualizeBloomOverlayPS(
-	  float4 UVAndScreenPos : TEXCOORD0
+	noperspective float4 UVAndScreenPos : TEXCOORD0
 	, out float4 OutColor : SV_Target0)
 {
 	float2 UV = UVAndScreenPos.xy;
@@ -150,9 +150,9 @@ float4 BloomThreshold;
 // -----------------------------
 // bloom threshold
 void MainPS(
-	  float4 UVAndScreenPos : TEXCOORD0
-	, float InExposureScale : TEXCOORD1
-	, out float4 OutColor : SV_Target0)
+	noperspective float4 UVAndScreenPos : TEXCOORD0,
+	nointerpolation float InExposureScale : TEXCOORD1,
+	out float4 OutColor : SV_Target0)
 {
 	float2 UV = UVAndScreenPos.xy;
 
diff --git a/Engine/Shaders/PostProcessBokehDOF.usf b/Engine/Shaders/PostProcessBokehDOF.usf
--- a/Engine/Shaders/PostProcessBokehDOF.usf
+++ b/Engine/Shaders/PostProcessBokehDOF.usf
@@ -347,7 +347,7 @@ float4 CommonDOFSetup(in float2 CenterUV, out bool bFrontLayer, out float4 Mask
 }
 
 // downsample to half res, put depth in alpha
-void MainSetupPS(in float4 UVAndScreenPos : TEXCOORD0, out float4 OutColor : SV_Target0)
+void MainSetupPS(in noperspective float4 UVAndScreenPos : TEXCOORD0, out float4 OutColor : SV_Target0)
 {
 	// unused for this pass
 	bool bFrontLayer;
@@ -361,7 +361,7 @@ void MainSetupPS(in float4 UVAndScreenPos : TEXCOORD0, out float4 OutColor : SV_
 
 // pixel shader to accumulate the Bokeh shaped elements
 void MainPS(
-	float2 TexCoord : TEXCOORD0,
+	noperspective float2 TexCoord : TEXCOORD0,
 	nointerpolation float4 InColor : TEXCOORD1,
 	out float4 OutColor : SV_Target0
 	)
@@ -378,7 +378,7 @@ void MainPS(
 
 // pixel shader to combine the full res scene and the blurred images behind and in front of the the focal plane
 void MainRecombinePS(
-	in float4 UVAndScreenPos : TEXCOORD0,
+	noperspective in float4 UVAndScreenPos : TEXCOORD0,
 	out float4 OutColor : SV_Target0
 	)
 {
diff --git a/Engine/Shaders/PostProcessCircleDOF.usf b/Engine/Shaders/PostProcessCircleDOF.usf
--- a/Engine/Shaders/PostProcessCircleDOF.usf
+++ b/Engine/Shaders/PostProcessCircleDOF.usf
@@ -11,7 +11,7 @@
 #include "CircleDOFCommon.usf"
 
 // pixel shader entry point
-void CircleSetupPS(float4 UVAndScreenPos : TEXCOORD0, out float4 OutColor0 : SV_Target0)
+void CircleSetupPS(noperspective float4 UVAndScreenPos : TEXCOORD0, out float4 OutColor0 : SV_Target0)
 {
 	float2 UV = UVAndScreenPos.xy;
 
diff --git a/Engine/Shaders/PostProcessCombineLUTs.usf b/Engine/Shaders/PostProcessCombineLUTs.usf
--- a/Engine/Shaders/PostProcessCombineLUTs.usf
+++ b/Engine/Shaders/PostProcessCombineLUTs.usf
@@ -162,7 +162,7 @@ void MainPS(FWriteToSliceGeometryOutput Input, out float4 OutColor : SV_Target0)
 		Neutral = float4(UV * LUTSize / (LUTSize - 1), Input.LayerIndex / (LUTSize - 1), 0);
 	}
 #else
-void MainPS(float4 InUV : TEXCOORD0, out float4 OutColor : SV_Target0)
+void MainPS(noperspective float4 InUV : TEXCOORD0, out float4 OutColor : SV_Target0)
 {
 	// construct the neutral color from a 2d position in 256x16
 	float4 Neutral;
diff --git a/Engine/Shaders/PostProcessDownsample.usf b/Engine/Shaders/PostProcessDownsample.usf
--- a/Engine/Shaders/PostProcessDownsample.usf
+++ b/Engine/Shaders/PostProcessDownsample.usf
@@ -18,16 +18,22 @@ void MainDownsampleVS(
 	DrawRectangle(InPosition, UV, OutPosition, OutUV);
 }
 
+// xy: UV offset, zw:unused
+float4 DownsampleParams;
+
 // pixel shader entry point
-void MainPS(float2 InUV : TEXCOORD0, out float4 OutColor : SV_Target0)
+void MainPS(noperspective float2 InUV : TEXCOORD0, out float4 OutColor : SV_Target0)
 {
-	float2 Offset = PostprocessInput0Size.zw;
+	float2 Offset = DownsampleParams.xy;
 
 #if METHOD == 0
-	// low quality 1 sample
+	// Output: low quality, 1 filtered sample
+	
 	OutColor = Texture2DSample(PostprocessInput0, PostprocessInput0Sampler, InUV);
+
 #elif METHOD == 1
-	// high quality 4 samples
+	// Output: float4(RGBA), 4 filtered samples
+	
 	float2 UV[4];
 
 	// blur during downsample (4x4 kernel) to get better quality especially for HDR content, can be made an option of this shader
@@ -38,29 +44,29 @@ void MainPS(float2 InUV : TEXCOORD0, out float4 OutColor : SV_Target0)
 
 	float4 Sample[4];
 
-	for(uint i = 0; i < 4; ++i)
+	UNROLL for(uint i = 0; i < 4; ++i)
 	{
 		Sample[i] = Texture2DSample(PostprocessInput0, PostprocessInput0Sampler, UV[i]);
 	}
 
 	OutColor = (Sample[0] + Sample[1] + Sample[2] + Sample[3]) * 0.25f;
 
-	// fixed rarely occuring yellow color tine of the whole viewport (certain view port size, need to investigate more)
+	// fixed rarely occuring yellow color tint of the whole viewport (certain view port size, need to investigate more)
 	OutColor.rgb = max(float3(0,0,0), OutColor.rgb);
 #else
-	// high quality 4 samples with blurring
-	// Depth in Alpha
+	// Output: float4(RGB,depth), 4 unfiltered samples (no leaking)
+
 	float2 UV[4];
 
 	// no filtering (2x2 kernel) to get no leaking in Depth of Field
-	UV[0] = InUV + Offset * float2(-0.5f, -0.5f);
-	UV[1] = InUV + Offset * float2( 0.5f, -0.5f);
-	UV[2] = InUV + Offset * float2(-0.5f,  0.5f);
-	UV[3] = InUV + Offset * float2( 0.5f,  0.5f);
+	UV[0] = InUV + Offset * float2(-1, -1);
+	UV[1] = InUV + Offset * float2( 1, -1);
+	UV[2] = InUV + Offset * float2(-1,  1);
+	UV[3] = InUV + Offset * float2( 1,  1);
 
 	float3 Sample[4];
 
-	for(uint i = 0; i < 4; ++i)
+	UNROLL for(uint i = 0; i < 4; ++i)
 	{
 		Sample[i] = Texture2DSample(PostprocessInput0, PostprocessInput0Sampler, UV[i]).rgb;
 	}
diff --git a/Engine/Shaders/PostProcessLensBlur.usf b/Engine/Shaders/PostProcessLensBlur.usf
--- a/Engine/Shaders/PostProcessLensBlur.usf
+++ b/Engine/Shaders/PostProcessLensBlur.usf
@@ -79,8 +79,8 @@ void MainVS(
 
 // pixel shader
 void MainPS(
-	float2 TexCoord : TEXCOORD0,
-	float4 InColor : TEXCOORD1,
+	noperspective float2 TexCoord : TEXCOORD0,
+	noperspective float4 InColor : TEXCOORD1,
 	out float4 OutColor : SV_Target0
 	)
 {
diff --git a/Engine/Shaders/PostProcessLensFlares.usf b/Engine/Shaders/PostProcessLensFlares.usf
--- a/Engine/Shaders/PostProcessLensFlares.usf
+++ b/Engine/Shaders/PostProcessLensFlares.usf
@@ -14,13 +14,13 @@ float2 TexScale;
 
 
 // pixel shader, can be optimized by using additive blending
-void CopyPS(float4 InUV : TEXCOORD0, out float4 OutColor : SV_Target0)
+void CopyPS(noperspective float4 InUV : TEXCOORD0, out float4 OutColor : SV_Target0)
 {
 	OutColor = Texture2DSample(PostprocessInput0, PostprocessInput0Sampler, InUV.xy);
 }
 
 // pixel shader entry point
-void MainPS(in float4 UVAndScreenPos : TEXCOORD0, out float4 OutColor : SV_Target0)
+void MainPS(in noperspective float4 UVAndScreenPos : TEXCOORD0, out float4 OutColor : SV_Target0)
 {
 	float2 UV = UVAndScreenPos.xy;
 	float2 ScreenPos = UVAndScreenPos.zw;
diff --git a/Engine/Shaders/PostProcessPassThrough.usf b/Engine/Shaders/PostProcessPassThrough.usf
--- a/Engine/Shaders/PostProcessPassThrough.usf
+++ b/Engine/Shaders/PostProcessPassThrough.usf
@@ -8,7 +8,7 @@
 #include "PostProcessCommon.usf"
 
 // pixel shader
-void MainPS(float4 InUV : TEXCOORD0, out float4 OutColor : SV_Target0)
+void MainPS(noperspective float4 InUV : TEXCOORD0, out float4 OutColor : SV_Target0)
 {
 	OutColor = Texture2DSample(PostprocessInput0, PostprocessInput0Sampler, InUV.xy);
 }
diff --git a/Engine/Shaders/PostProcessTemporalAA.usf b/Engine/Shaders/PostProcessTemporalAA.usf
--- a/Engine/Shaders/PostProcessTemporalAA.usf
+++ b/Engine/Shaders/PostProcessTemporalAA.usf
@@ -244,7 +244,7 @@ float RandFast( uint2 PixelPos, float Magic = 3571.0 )
 	return Random.x;
 }
 
-void SSRTemporalAAPS( float4 UVAndScreenPos : TEXCOORD0, float3 InExposureScaleVignette : TEXCOORD1, out float4 OutColor : SV_Target0 )
+void SSRTemporalAAPS( noperspective float4 UVAndScreenPos : TEXCOORD0, float3 InExposureScaleVignette : TEXCOORD1, out float4 OutColor : SV_Target0 )
 {
 #if ENABLE_TEMPORAL_AA
 	float InExposureScale = InExposureScaleVignette.x;
@@ -260,7 +260,7 @@ void SSRTemporalAAPS( float4 UVAndScreenPos : TEXCOORD0, float3 InExposureScaleV
 #endif
 }
 
-void DOFTemporalAAPS( float4 UVAndScreenPos : TEXCOORD0, float3 InExposureScaleVignette : TEXCOORD1, out float4 OutColor : SV_Target0 )
+void DOFTemporalAAPS( noperspective float4 UVAndScreenPos : TEXCOORD0, float3 InExposureScaleVignette : TEXCOORD1, out float4 OutColor : SV_Target0 )
 {
 #if ENABLE_TEMPORAL_AA
 	float InExposureScale = InExposureScaleVignette.x;
@@ -281,7 +281,7 @@ void DOFTemporalAAPS( float4 UVAndScreenPos : TEXCOORD0, float3 InExposureScaleV
 #endif
 }
 
-void LightShaftTemporalAAPS( float4 UVAndScreenPos : TEXCOORD0, float3 InExposureScaleVignette : TEXCOORD1, out float4 OutColor : SV_Target0 )
+void LightShaftTemporalAAPS( noperspective float4 UVAndScreenPos : TEXCOORD0, float3 InExposureScaleVignette : TEXCOORD1, out float4 OutColor : SV_Target0 )
 {
 #if ENABLE_TEMPORAL_AA
 	float InExposureScale = InExposureScaleVignette.x;
@@ -297,11 +297,11 @@ void LightShaftTemporalAAPS( float4 UVAndScreenPos : TEXCOORD0, float3 InExposur
 }
 
 void MainTemporalAAPS(
-	in float4 UVAndScreenPos : TEXCOORD0,
-	in float3 InExposureScaleVignette : TEXCOORD1,
+	in noperspective float4 UVAndScreenPos : TEXCOORD0,
+	in noperspective float3 InExposureScaleVignette : TEXCOORD1,
 	// Note: these are needed to match up with the vertex shader outputs which is FPostProcessTonemapVS, otherwise a d3d debug error is issued
-	in float4 Unused1 : TEXCOORD2,
-	in float4 Unused2 : TEXCOORD3,
+	in noperspective float4 Unused1 : TEXCOORD2,
+	in noperspective float4 Unused2 : TEXCOORD3,
 	in float4 SvPosition : SV_Position,
 	out float4 OutColor : SV_Target0 )
 {
@@ -329,8 +329,8 @@ void MainTemporalAAPS(
 }
 
 void MainFastTemporalAAPS(
-	in float4 UVAndScreenPos : TEXCOORD0,
-	in float3 InExposureScaleVignette : TEXCOORD1,
+	in noperspective float4 UVAndScreenPos : TEXCOORD0,
+	in noperspective float3 InExposureScaleVignette : TEXCOORD1,
 	in float4 SvPosition : SV_Position,
 	out float4 OutColor : SV_Target0 )
 {
diff --git a/Engine/Shaders/PostProcessTestImage.usf b/Engine/Shaders/PostProcessTestImage.usf
--- a/Engine/Shaders/PostProcessTestImage.usf
+++ b/Engine/Shaders/PostProcessTestImage.usf
@@ -22,7 +22,7 @@ float Quantize(float x, float Count)
 	return ceil(x * Count - 1.0f) / (Count - 1.0f);
 }
 
-void MainPS(float4 UVAndScreenPos : TEXCOORD0, out float4 OutColor : SV_Target0)
+void MainPS(noperspective float4 UVAndScreenPos : TEXCOORD0, out float4 OutColor : SV_Target0)
 {
 	OutColor = 0;
 
diff --git a/Engine/Shaders/PostProcessTonemap.usf b/Engine/Shaders/PostProcessTonemap.usf
--- a/Engine/Shaders/PostProcessTonemap.usf
+++ b/Engine/Shaders/PostProcessTonemap.usf
@@ -171,10 +171,10 @@ float3 Graph( float2 ScreenSpacePos )
 
 // pixel shader entry point
 void MainPS(
-	in float4 UVAndScreenPos : TEXCOORD0,
-	in float3 InExposureScaleVignette : TEXCOORD1,
-	in float4 GrainUV : TEXCOORD2,
-	in float4 FringeUV : TEXCOORD3,
+	in noperspective float4 UVAndScreenPos : TEXCOORD0,
+	in noperspective float3 InExposureScaleVignette : TEXCOORD1,
+	in noperspective float4 GrainUV : TEXCOORD2,
+	in noperspective float4 FringeUV : TEXCOORD3,
 	out float4 OutColor : SV_Target0
 	)
 {
diff --git a/Engine/Shaders/PostProcessUpscale.usf b/Engine/Shaders/PostProcessUpscale.usf
--- a/Engine/Shaders/PostProcessUpscale.usf
+++ b/Engine/Shaders/PostProcessUpscale.usf
@@ -84,7 +84,7 @@ void MainVS(
 }
 
 
-void MainPS(float4 UVAndScreenPos : TEXCOORD0, float4 SvPosition : SV_POSITION, out float4 OutColor : SV_Target0)
+void MainPS(noperspective float4 UVAndScreenPos : TEXCOORD0, float4 SvPosition : SV_POSITION, out float4 OutColor : SV_Target0)
 {
 #if METHOD == 0
 	// Nearest sampling (not blurry but blocky, more for testing)
diff --git a/Engine/Shaders/PostProcessVisualizeBuffer.usf b/Engine/Shaders/PostProcessVisualizeBuffer.usf
--- a/Engine/Shaders/PostProcessVisualizeBuffer.usf
+++ b/Engine/Shaders/PostProcessVisualizeBuffer.usf
@@ -10,7 +10,7 @@
 #include "PostProcessHistogramCommon.usf"
 
 /** Main pixel shader function */
-void MainPS(float4 UVAndScreenPos : TEXCOORD0, out float4 OutColor : SV_Target0)
+void MainPS(noperspective float4 UVAndScreenPos : TEXCOORD0, out float4 OutColor : SV_Target0)
 {
 	float Alpha = 1;
 
diff --git a/Engine/Shaders/PostProcessVisualizeGBuffer.usf b/Engine/Shaders/PostProcessVisualizeGBuffer.usf
--- a/Engine/Shaders/PostProcessVisualizeGBuffer.usf
+++ b/Engine/Shaders/PostProcessVisualizeGBuffer.usf
@@ -96,7 +96,7 @@ float3 ComputeQuarterColor(float2 UV, int2 QuarterTileXY)
 }
 
 /** Main pixel shader function */
-void MainPS(float4 UVAndScreenPos : TEXCOORD0, out float4 OutColor : SV_Target0)
+void MainPS(noperspective float4 UVAndScreenPos : TEXCOORD0, out float4 OutColor : SV_Target0)
 {
 #if GBUFFER_VISUALIZATION_MODE == GBUFFER_VISUALIZATION_MODE_ALL
 	int2 PixelPos = (int2)(UVAndScreenPos.zw * ScreenPosToPixel.xy + ScreenPosToPixel.zw + 0.5f);
diff --git a/Engine/Shaders/PostProcessVisualizeHDR.usf b/Engine/Shaders/PostProcessVisualizeHDR.usf
--- a/Engine/Shaders/PostProcessVisualizeHDR.usf
+++ b/Engine/Shaders/PostProcessVisualizeHDR.usf
@@ -110,7 +110,7 @@ float3 Colorize(float x)
 }
 
 //
-void MainPS(float4 UVAndScreenPos : TEXCOORD0, out float4 OutColor : SV_Target0)
+void MainPS(noperspective float4 UVAndScreenPos : TEXCOORD0, out float4 OutColor : SV_Target0)
 {
 	float2 UV = UVAndScreenPos.xy;
 	int2 PixelPos = (int2)(UVAndScreenPos.zw * ScreenPosToPixel.xy + ScreenPosToPixel.zw + 0.5f);
