Integrated from Orion

diff --git a/Engine/Shaders/PostProcessSubsurface.usf b/Engine/Shaders/PostProcessSubsurface.usf
--- a/Engine/Shaders/PostProcessSubsurface.usf
+++ b/Engine/Shaders/PostProcessSubsurface.usf
@@ -147,7 +147,7 @@ float ComputeFullResLerp(FScreenSpaceData ScreenSpaceData, float2 UVSceneColor,
 }
 
 // visualization (doesn't have to be fast)
-void VisualizePS(in float4 UVAndScreenPos : TEXCOORD0, out float4 OutColor : SV_Target0)
+void VisualizePS(in noperspective float4 UVAndScreenPos : TEXCOORD0, out float4 OutColor : SV_Target0)
 {
 	float2 UV = UVAndScreenPos.xy;
 
@@ -239,7 +239,7 @@ SDiffuseAndSpecular ReconstructDiffuse(float2 UVSceneColor)
 	// todo: expose as scalability setting (can be evaluate best without TemporalAA)
 	// 0:fast but pattern can appear, 1:better, 2: even better, 3: best but expensive
 
-	// With TemporalAA they all look the same but 0 has an offset which we can fixed by alternating the 2 (L R) or 4 directions (L R U D)
+	// 0/1/2/3 With TemporalAA they all look the same but 0 has an offset which we can fixed by alternating the 2 (L R) or 4 directions (L R U D)
 	#define QUALITY_LEVEL 1
 
 #if QUALITY_LEVEL == 0
@@ -254,7 +254,7 @@ SDiffuseAndSpecular ReconstructDiffuse(float2 UVSceneColor)
 		Texture2DSample(PostprocessInput0, PostprocessInput0Sampler, UVSceneColor + float2( 0,  1) * View.BufferSizeAndInvSize.zw).rgb +
 		Texture2DSample(PostprocessInput0, PostprocessInput0Sampler, UVSceneColor + float2(-1,  0) * View.BufferSizeAndInvSize.zw).rgb +
 		Texture2DSample(PostprocessInput0, PostprocessInput0Sampler, UVSceneColor + float2( 0, -1) * View.BufferSizeAndInvSize.zw).rgb);
-#else
+#else // QUALITY_LEVEL == 3
 	float3 A = Texture2DSample(PostprocessInput0, PostprocessInput0Sampler, UVSceneColor + float2( 1,  0) * View.BufferSizeAndInvSize.zw).rgb;
 	float3 B = Texture2DSample(PostprocessInput0, PostprocessInput0Sampler, UVSceneColor + float2(-1,  0) * View.BufferSizeAndInvSize.zw).rgb;
 	float3 C = Texture2DSample(PostprocessInput0, PostprocessInput0Sampler, UVSceneColor + float2( 0,  1) * View.BufferSizeAndInvSize.zw).rgb;
@@ -301,7 +301,7 @@ float4 SetupSubsurfaceForOnePixel(float2 UVSceneColor)
 }
 
 
-void SetupPS(in float4 UVAndScreenPos : TEXCOORD0, out float4 OutColor : SV_Target0)
+void SetupPS(in noperspective float4 UVAndScreenPos : TEXCOORD0, out float4 OutColor : SV_Target0)
 {
 	float2 UV = UVAndScreenPos.xy;
 
@@ -358,7 +358,7 @@ void SetupPS(in float4 UVAndScreenPos : TEXCOORD0, out float4 OutColor : SV_Targ
 }
 
 // input0 is created by the SetupPS shader
-void MainPS(float4 UVAndScreenPos : TEXCOORD0, out float4 OutColor : SV_Target0)
+void MainPS(noperspective float4 UVAndScreenPos : TEXCOORD0, out float4 OutColor : SV_Target0)
 {
 	float2 ViewportUV = UVAndScreenPos.xy;
 
@@ -389,7 +389,7 @@ void MainPS(float4 UVAndScreenPos : TEXCOORD0, out float4 OutColor : SV_Target0)
 
 // Recombines the half res Subsurface filtered lighting contribution (upsampled and renormalized with the alpha)
 // with the SceneColor.
-void SubsurfaceRecombinePS(float4 UVAndScreenPos : TEXCOORD0, out float4 OutColor : SV_Target0)
+void SubsurfaceRecombinePS(noperspective float4 UVAndScreenPos : TEXCOORD0, out float4 OutColor : SV_Target0)
 {
 	float2 ViewportUV = UVAndScreenPos.xy;
 
