Merging CL 2692508

Added custom UV support to Slate vertex shaders

diff --git a/Engine/Shaders/SlateElementPixelShader.usf b/Engine/Shaders/SlateElementPixelShader.usf
--- a/Engine/Shaders/SlateElementPixelShader.usf
+++ b/Engine/Shaders/SlateElementPixelShader.usf
@@ -3,10 +3,7 @@
 
 #include "Common.usf"
 #include "GammaCorrectionCommon.usf"
-
-#if USE_MATERIALS
-#include "Material.usf"
-#endif
+#include "SlateShaderCommon.usf"
 
 /** Shader types (mirrored from ESlateShader::Type in RenderingCommon.h */
 #define ST_Default		0
@@ -14,33 +11,18 @@
 #define ST_Font			2
 #define ST_LineSegment	3
 
-/** Color Vision Deficiency Types (mirrored in EditorUserSettings.h) */
-#define CVD_NormalVision	0
-#define CVD_Deuteranomly	1
-#define CVD_Deuteranopia	2
-#define CVD_Protanomly		3
-#define CVD_Protanopia		4
-#define CVD_Tritanomaly		5
-#define CVD_Tritanopia		6
-#define CVD_Achromatopsia	7
-
 /** Enable to debug clipping issues */
 #define DEBUG_CLIPPING 0
 
-
 /** Display gamma x:gamma curve adjustment, y:inverse gamma (1/GEngine->DisplayGamma) */
 float2 GammaValues;
 /** Per element params */
 float4 ShaderParams;
 
-//#if !USE_MATERIALS
 /** Texture for the element */
 Texture2D ElementTexture;
 SamplerState ElementTextureSampler;
 
-//#endif
-
-
 float3 GammaCorrect(float3 InColor)
 {
 	float3 CorrectedColor = InColor;
@@ -56,54 +38,60 @@ float3 GammaCorrect(float3 InColor)
 }
 
 
-struct VertexOut
-{
-	float4 Position : SV_POSITION;
-	float4 Color : COLOR0;
-	float4 TextureCoordinates : TEXCOORD0;
-	float4 ClipOriginAndPos : TEXCOORD1;
-	float4 ClipExtents : TEXCOORD2;
-	float2 MaterialTexCoords : TEXCOORD3;
-};
-	
 #if USE_MATERIALS
-float4 GetMaterialColor( VertexOut InVertex, float2 CustomUV )
+void FillTexCoords( inout FMaterialPixelParameters Parameters, VertexToPixelInterpolants InVertex, float2 InDefaultUV )
 {
-	FMaterialPixelParameters Parameters = MakeInitializedMaterialPixelParameters();
-	CalcMaterialParameters(Parameters, InVertex.Position, true, InVertex.Position);
-
 #if SHADER_TYPE == ST_Font
 	#if NUM_MATERIAL_TEXCOORDS == 1
-		Parameters.TexCoords[0] = InVertex.TextureCoordinates.zw;
+		// Note: The first channel is used to look up into the texture atlas. The second channel is the custom 0-1 UV set for the text run
+		Parameters.TexCoords[0] = GetUV( InVertex, 1 );
 	#elif NUM_MATERIAL_TEXCOORDS == 2
-		Parameters.TexCoords[0] = InVertex.TextureCoordinates.zw;
+		// Note: The first channel is used to look up into the texture atlas. The second channel is the custom 0-1 UV set for the text run
+		Parameters.TexCoords[0] = GetUV( InVertex, 1 );
 		Parameters.TexCoords[1] = InVertex.MaterialTexCoords;
 	#endif
 #else
-	#if NUM_MATERIAL_TEXCOORDS == 1
-		Parameters.TexCoords[0] = CustomUV;
-	#elif NUM_MATERIAL_TEXCOORDS == 2
-		Parameters.TexCoords[0] = CustomUV;
-		Parameters.TexCoords[1] = InVertex.TextureCoordinates.xy;
-	#elif NUM_MATERIAL_TEXCOORDS == 3
-		Parameters.TexCoords[0] = CustomUV;
-		Parameters.TexCoords[1] = InVertex.TextureCoordinates.xy;
-		Parameters.TexCoords[2] = InVertex.TextureCoordinates.zw;
-	#endif
-#endif
+	#if NUM_CUSTOMIZED_UVS == 0
+		#if NUM_MATERIAL_TEXCOORDS == 1
+			Parameters.TexCoords[0] = InDefaultUV;
+		#elif NUM_MATERIAL_TEXCOORDS == 2
+			Parameters.TexCoords[0] = InDefaultUV;
+			Parameters.TexCoords[1] = GetUV( InVertex, 0 );
+		#elif NUM_MATERIAL_TEXCOORDS == 3
+			Parameters.TexCoords[0] = InDefaultUV;
+			Parameters.TexCoords[1] = GetUV( InVertex, 0 );
+			Parameters.TexCoords[2] = GetUV( InVertex, 1 );
+		#endif
+	#else
+		// Note:  We assume users with customized UVs are overriding behavior and know what they are doing
+		UNROLL
+		for( int CoordinateIndex = 0; CoordinateIndex < NUM_MATERIAL_TEXCOORDS; CoordinateIndex++ )
+		{
+			Parameters.TexCoords[CoordinateIndex] = GetUV( InVertex, CoordinateIndex );
+		}
+	#endif // NUM_CUSTOMIZED_UVS == 0
+#endif	// SHADER_TYPE
+}
+
+float4 GetMaterialColor( VertexToPixelInterpolants InVertex, float2 MeshUV )
+{
+	FMaterialPixelParameters Parameters = MakeInitializedMaterialPixelParameters();
+	CalcMaterialParameters( Parameters, InVertex.Position, true, InVertex.Position );
 
 	Parameters.VertexColor = InVertex.Color;
 
+	FillTexCoords( Parameters, InVertex, MeshUV );
+
 	// Clip if masked
 	GetMaterialCoverageAndClipping( Parameters );
 	
-	float Opacity = GetMaterialOpacity(Parameters);
-	float3 Color = GetMaterialEmissive(Parameters);
+	float Opacity = GetMaterialOpacity( Parameters );
+	float3 Color = GetMaterialEmissive( Parameters );
 
 	float4 OutColor;
 
 	#if MATERIALBLENDING_ADDITIVE
-		OutColor = float4(Color * Opacity, 0.0f);
+		OutColor = float4( Color * Opacity, 0.0f );
 	#else
 		OutColor = float4( Color, Opacity );
 	#endif
@@ -113,7 +101,7 @@ float4 GetMaterialColor( VertexOut InVertex, float2 CustomUV )
 #endif
 
 #if !USE_MATERIALS
-float4 GetTextureColor( VertexOut InVertex, float2 UV )
+float4 GetTextureColor( VertexToPixelInterpolants InVertex, float2 UV )
 {
 	float4 BaseColor = Texture2DSample(ElementTexture, ElementTextureSampler, UV );
 #if !USE_TEXTURE_ALPHA
@@ -125,7 +113,7 @@ float4 GetTextureColor( VertexOut InVertex, float2 UV )
 #endif
 
 /** Gets a color from a texture and supplied vertex color.  Adjusting for tint masking if necessary */
-float4 GetColor( VertexOut InVertex, float2 UV )
+float4 GetColor( VertexToPixelInterpolants InVertex, float2 UV )
 {
 	float4 FinalColor;
 	
@@ -141,8 +129,9 @@ float4 GetColor( VertexOut InVertex, float2 UV )
 }
 
 #if SHADER_TYPE == ST_Font
-float4 GetFontElementColor( VertexOut InVertex )
+float4 GetFontElementColor( VertexToPixelInterpolants InVertex )
 {
+	float2 AtlasTextureCoordinates = GetUV( InVertex, 0 );
 #if USE_MATERIALS
 	float4 BaseColor = GetMaterialColor( InVertex, float2(0,0) )*InVertex.Color;
 #else
@@ -150,23 +139,23 @@ float4 GetFontElementColor( VertexOut InVertex )
 #endif
 	float4 OutColor = BaseColor;
 
-	OutColor.a *= Texture2DSample_A8(ElementTexture, ElementTextureSampler, InVertex.TextureCoordinates.xy);
+	OutColor.a *= Texture2DSample_A8( ElementTexture, ElementTextureSampler, AtlasTextureCoordinates );
 
 	return OutColor;
 }
 #endif
 
 #if SHADER_TYPE == ST_Default
-float4 GetDefaultElementColor( VertexOut InVertex )
+float4 GetDefaultElementColor( VertexToPixelInterpolants InVertex )
 {
-	return GetColor( InVertex, InVertex.TextureCoordinates.xy*InVertex.TextureCoordinates.zw );
+	return GetColor( InVertex, GetUV( InVertex, 0) * GetUV( InVertex, 1 ) );
 }
 #endif
 
 #if SHADER_TYPE == ST_Border 
-float4 GetBorderElementColor( VertexOut InVertex )
+float4 GetBorderElementColor( VertexToPixelInterpolants InVertex )
 {
-	float4 TextureCoordinates = InVertex.TextureCoordinates;
+	float4 TextureCoordinates = InVertex.TextureCoordinates[0];
 	float2 NewUV;
 	if( TextureCoordinates.z == 0.0f && TextureCoordinates.w == 0.0f )
 	{
@@ -204,13 +193,13 @@ float4 GetBorderElementColor( VertexOut InVertex )
  * Generates an anti-aliased line segment pixel
  * The idea for the anti-alising technique is based on the fast prefiltered lines technique published in GPU Gems 2 
  */
-float4 GetLineSegmentElementColor( VertexOut InVertex )
+float4 GetLineSegmentElementColor( VertexToPixelInterpolants InVertex )
 {
 	float Width = ShaderParams.x;
 	float Radius = ShaderParams.y;
 
-	float2 StartPos = InVertex.TextureCoordinates.xy;
-	float2 EndPos = InVertex.TextureCoordinates.zw;
+	float2 StartPos = GetUV( InVertex, 0 );
+	float2 EndPos = GetUV( InVertex, 1 );
 
 	float2 Diff = float2( StartPos.y - EndPos.y, EndPos.x - StartPos.x ) ;
 
@@ -260,82 +249,7 @@ float4 GetLineSegmentElementColor( VertexOut InVertex )
 } 
 #endif
 
-float4 ApplyColorVisionDeficiencyFilter(float4 Color)
-{
-	float4 OutColor = Color;
-#if COLOR_VISION_DEFICIENCY_TYPE == CVD_Deuteranomly
-	OutColor.rgb = float3(
-		OutColor.r * 0.80000 + OutColor.g * 0.20000 + OutColor.b * 0.0,
-		OutColor.r * 0.25833 + OutColor.g * 0.74167 + OutColor.b * 0.0,
-		OutColor.r * 0.0     + OutColor.g * 0.14167 + OutColor.b * 0.85833);
-#elif COLOR_VISION_DEFICIENCY_TYPE == CVD_Deuteranopia
-	OutColor.rgb = float3(
-		OutColor.r * 0.625 + OutColor.g * 0.375 + OutColor.b * 0.0,
-		OutColor.r * 0.700 + OutColor.g * 0.300 + OutColor.b * 0.0,
-		OutColor.r * 0.0   + OutColor.g * 0.300 + OutColor.b * 0.700);
-#elif COLOR_VISION_DEFICIENCY_TYPE == CVD_Protanomly
-	OutColor.rgb = float3(
-		OutColor.r * 0.81667 + OutColor.g * 0.18333 + OutColor.b * 0.0,
-		OutColor.r * 0.33333 + OutColor.g * 0.66667 + OutColor.b * 0.0,
-		OutColor.r * 0.0     + OutColor.g * 0.12500 + OutColor.b * 0.87500);
-#elif COLOR_VISION_DEFICIENCY_TYPE == CVD_Protanopia
-	OutColor.rgb = float3(
-		OutColor.r * 0.56667 + OutColor.g * 0.43333 + OutColor.b * 0.0,
-		OutColor.r * 0.55833 + OutColor.g * 0.44167 + OutColor.b * 0.0,
-		OutColor.r * 0.0     + OutColor.g * 0.24167 + OutColor.b * 0.75833);
-#elif COLOR_VISION_DEFICIENCY_TYPE == CVD_Tritanomaly
-	OutColor.rgb = float3(
-		OutColor.r * 0.96667 + OutColor.g * 0.03333 + OutColor.b * 0.0,
-		OutColor.r * 0.0     + OutColor.g * 0.73333 + OutColor.b * 0.26667,
-		OutColor.r * 0.0     + OutColor.g * 0.18333 + OutColor.b * 0.81667);
-#elif COLOR_VISION_DEFICIENCY_TYPE == CVD_Tritanopia
-	OutColor.rgb = float3(
-		OutColor.r * 0.950 + OutColor.g * 0.050 + OutColor.b * 0.0,
-		OutColor.r * 0.0   + OutColor.g * 0.433 + OutColor.b * 0.567,
-		OutColor.r * 0.0   + OutColor.g * 0.475 + OutColor.b * 0.525);
-#elif COLOR_VISION_DEFICIENCY_TYPE == CVD_Achromatopsia
-	OutColor.rgb = float3(
-		OutColor.r * 0.299 + OutColor.g * 0.587 + OutColor.b * 0.114,
-		OutColor.r * 0.299 + OutColor.g * 0.587 + OutColor.b * 0.114,
-		OutColor.r * 0.299 + OutColor.g * 0.587 + OutColor.b * 0.114);
-#endif
-	return OutColor;
-}
-
-float cross(float2 a, float2 b)
-{
-	return a.x*b.y - a.y*b.x;
-}
-
-/**
- * Given a point p and a parallelogram defined by point a and vectors b and c, determines in p is inside the parallelogram. 
- * returns a 4-vector that can be used with the clip instruction.
- */
-float4 PointInParallelogram(float2 p, float2 a, float4 bc)
-{
-	// unoptomized form:
-	//float2 o = p - a;
-	//float2 b = bc.xy;
-	//float2 c = bc.zw;
-	//float d = cross(b, c);
-	//float s = -cross(o, b) / d;
-	//float t = cross(o, c) / d;
-	// test for s and t between 0 and 1
-	//return float4(s, 1 - s, t, 1 - t);
-
-	float2 o = p - a;
-	// precompute 1/d
-	float invD = 1/cross(bc.xy, bc.zw);
-	// Compute an optimized o x b and o x c, leveraging that b and c are in the same vector register already (and free swizzles):
-	//   (o.x * b .y  - o.y * b .x, o.x *  c.y - o.y *  c.x) ==
-	//   (o.x * bc.y  - o.y * bc.x, o.x * bc.w - o.y * bc.z) ==
-	//    o.x * bc.yw - o.y * bc.xz
-	float2 st = (o.x * bc.yw - o.y * bc.xz) * float2(-invD, invD);
-	// test for s and t between 0 and 1
-	return float4(st, float2(1,1) - st);
-}
-
-float4 Main( VertexOut VIn ) : SV_Target0
+float4 Main( VertexToPixelInterpolants VIn ) : SV_Target0
 {
 	// Clip pixels which are outside of the clipping rect
 	float2 ClipOrigin = VIn.ClipOriginAndPos.xy;
@@ -372,11 +286,9 @@ float4 Main( VertexOut VIn ) : SV_Target0
 #if !USE_MATERIALS
 	OutColor = ApplyColorVisionDeficiencyFilter(OutColor);
 #endif
-
 	// gamma correct
 	OutColor.rgb = GammaCorrect(OutColor.rgb);
 
-
 #if DEBUG_CLIPPING
 	if (any(ClipTest < 0))
 	{
@@ -388,7 +300,7 @@ float4 Main( VertexOut VIn ) : SV_Target0
 	return OutColor;
 }
 
-float4 DebugOverdrawMain( VertexOut VIn ) : SV_Target0
+float4 DebugOverdrawMain( VertexToPixelInterpolants VIn ) : SV_Target0
 {
 
 #if !DEBUG_CLIPPING
diff --git a/Engine/Shaders/SlateShaderCommon.usf b/Engine/Shaders/SlateShaderCommon.usf
new file mode 100644
--- /dev/null
+++ b/Engine/Shaders/SlateShaderCommon.usf
@@ -0,0 +1,119 @@
+// Copyright 1998-2015 Epic Games, Inc. All Rights Reserved.
+
+#if USE_MATERIALS
+#include "Material.usf"
+
+	#if NUM_MATERIAL_TEXCOORDS
+		#define NUM_SLATE_TEXCOORDS NUM_MATERIAL_TEXCOORDS
+	#else
+		#define NUM_SLATE_TEXCOORDS 1
+	#endif // NUM_MATERIAL_TEXCOORDS
+
+	#define HAS_MATERIAL_TEXCOORDS NUM_MATERIAL_TEXCOORDS > 0
+#else
+	#define HAS_MATERIAL_TEXCOORDS 0
+	// This is defined automatically based on the number of textures in a material
+	// In the case that this is a basic texture lookup shader, we only need 1 uv set
+	#define NUM_SLATE_TEXCOORDS 1
+#endif
+
+/** Color Vision Deficiency Types (mirrored in EditorUserSettings.h) */
+#define CVD_NormalVision	0
+#define CVD_Deuteranomly	1
+#define CVD_Deuteranopia	2
+#define CVD_Protanomly		3
+#define CVD_Protanopia		4
+#define CVD_Tritanomaly		5
+#define CVD_Tritanopia		6
+#define CVD_Achromatopsia	7
+
+struct VertexToPixelInterpolants
+{
+	float4 Position : SV_POSITION;
+	float4 Color : COLOR0;
+	float4 ClipOriginAndPos : TEXCOORD0;
+	float4 ClipExtents : TEXCOORD1;
+	float2 MaterialTexCoords : TEXCOORD2;
+	float4 TextureCoordinates[NUM_SLATE_TEXCOORDS] : TEXCOORD3;
+};
+
+float2 GetUV(VertexToPixelInterpolants Interpolants, int UVIndex)
+{
+	float4 UVVector = Interpolants.TextureCoordinates[UVIndex / 2];
+	return UVIndex % 2 ? UVVector.zw : UVVector.xy;
+}
+
+float4 ApplyColorVisionDeficiencyFilter(float4 Color)
+{
+	float4 OutColor = Color;
+#if COLOR_VISION_DEFICIENCY_TYPE == CVD_Deuteranomly
+	OutColor.rgb = float3(
+		OutColor.r * 0.80000 + OutColor.g * 0.20000 + OutColor.b * 0.0,
+		OutColor.r * 0.25833 + OutColor.g * 0.74167 + OutColor.b * 0.0,
+		OutColor.r * 0.0     + OutColor.g * 0.14167 + OutColor.b * 0.85833);
+#elif COLOR_VISION_DEFICIENCY_TYPE == CVD_Deuteranopia
+	OutColor.rgb = float3(
+		OutColor.r * 0.625 + OutColor.g * 0.375 + OutColor.b * 0.0,
+		OutColor.r * 0.700 + OutColor.g * 0.300 + OutColor.b * 0.0,
+		OutColor.r * 0.0   + OutColor.g * 0.300 + OutColor.b * 0.700);
+#elif COLOR_VISION_DEFICIENCY_TYPE == CVD_Protanomly
+	OutColor.rgb = float3(
+		OutColor.r * 0.81667 + OutColor.g * 0.18333 + OutColor.b * 0.0,
+		OutColor.r * 0.33333 + OutColor.g * 0.66667 + OutColor.b * 0.0,
+		OutColor.r * 0.0     + OutColor.g * 0.12500 + OutColor.b * 0.87500);
+#elif COLOR_VISION_DEFICIENCY_TYPE == CVD_Protanopia
+	OutColor.rgb = float3(
+		OutColor.r * 0.56667 + OutColor.g * 0.43333 + OutColor.b * 0.0,
+		OutColor.r * 0.55833 + OutColor.g * 0.44167 + OutColor.b * 0.0,
+		OutColor.r * 0.0     + OutColor.g * 0.24167 + OutColor.b * 0.75833);
+#elif COLOR_VISION_DEFICIENCY_TYPE == CVD_Tritanomaly
+	OutColor.rgb = float3(
+		OutColor.r * 0.96667 + OutColor.g * 0.03333 + OutColor.b * 0.0,
+		OutColor.r * 0.0     + OutColor.g * 0.73333 + OutColor.b * 0.26667,
+		OutColor.r * 0.0     + OutColor.g * 0.18333 + OutColor.b * 0.81667);
+#elif COLOR_VISION_DEFICIENCY_TYPE == CVD_Tritanopia
+	OutColor.rgb = float3(
+		OutColor.r * 0.950 + OutColor.g * 0.050 + OutColor.b * 0.0,
+		OutColor.r * 0.0   + OutColor.g * 0.433 + OutColor.b * 0.567,
+		OutColor.r * 0.0   + OutColor.g * 0.475 + OutColor.b * 0.525);
+#elif COLOR_VISION_DEFICIENCY_TYPE == CVD_Achromatopsia
+	OutColor.rgb = float3(
+		OutColor.r * 0.299 + OutColor.g * 0.587 + OutColor.b * 0.114,
+		OutColor.r * 0.299 + OutColor.g * 0.587 + OutColor.b * 0.114,
+		OutColor.r * 0.299 + OutColor.g * 0.587 + OutColor.b * 0.114);
+#endif
+	return OutColor;
+}
+
+float cross(float2 a, float2 b)
+{
+	return a.x*b.y - a.y*b.x;
+}
+
+/**
+ * Given a point p and a parallelogram defined by point a and vectors b and c, determines in p is inside the parallelogram. 
+ * returns a 4-vector that can be used with the clip instruction.
+ */
+float4 PointInParallelogram(float2 p, float2 a, float4 bc)
+{
+	// unoptomized form:
+	//float2 o = p - a;
+	//float2 b = bc.xy;
+	//float2 c = bc.zw;
+	//float d = cross(b, c);
+	//float s = -cross(o, b) / d;
+	//float t = cross(o, c) / d;
+	// test for s and t between 0 and 1
+	//return float4(s, 1 - s, t, 1 - t);
+
+	float2 o = p - a;
+	// precompute 1/d
+	float invD = 1/cross(bc.xy, bc.zw);
+	// Compute an optimized o x b and o x c, leveraging that b and c are in the same vector register already (and free swizzles):
+	//   (o.x * b .y  - o.y * b .x, o.x *  c.y - o.y *  c.x) ==
+	//   (o.x * bc.y  - o.y * bc.x, o.x * bc.w - o.y * bc.z) ==
+	//    o.x * bc.yw - o.y * bc.xz
+	float2 st = (o.x * bc.yw - o.y * bc.xz) * float2(-invD, invD);
+	// test for s and t between 0 and 1
+	return float4(st, float2(1,1) - st);
+}
diff --git a/Engine/Shaders/SlateVertexShader.usf b/Engine/Shaders/SlateVertexShader.usf
--- a/Engine/Shaders/SlateVertexShader.usf
+++ b/Engine/Shaders/SlateVertexShader.usf
@@ -3,27 +3,13 @@
 
 #include "Common.usf"
 #include "GammaCorrectionCommon.usf"
-
-#if USE_MATERIALS
-#include "Material.usf"
-#endif
+#include "SlateShaderCommon.usf"
 
 float4x4 ViewProjection;
 
 float SwitchVerticalAxisMultiplier;
 
-struct VertexOut
-{
-	float4 Position : SV_POSITION;
-	float4 Color : COLOR0;
-	float4 TextureCoordinates : TEXCOORD0;
-	float4 ClipOriginAndPos : TEXCOORD1;
-	float4 ClipExtents : TEXCOORD2;
-	float2 MaterialTexCoords : TEXCOORD3;
-};
-
-
-VertexOut Main(
+VertexToPixelInterpolants Main(
 	in float4 InTextureCoordinates : ATTRIBUTE0,
 	in float2 InMaterialTextureCoordinates : ATTRIBUTE1,
 	in float2 InPosition : ATTRIBUTE2,
@@ -32,34 +18,64 @@ VertexOut Main(
 	in float4 InColor : ATTRIBUTE5
 	)
 {
-	VertexOut VOut;
+	VertexToPixelInterpolants VOut = (VertexToPixelInterpolants)0;
 
 	float4 WorldPosition = float4(InPosition.xy,0,1);
+
+	InColor.rgb = sRGBToLinear(InColor.rgb);
+	float4 FinalVertexColor = InColor FCOLOR_COMPONENT_SWIZZLE;
+
+	VOut.MaterialTexCoords = InMaterialTextureCoordinates;
+	VOut.ClipOriginAndPos = float4(InClipOrigin, InPosition.xy);
+	VOut.ClipExtents = InClipExtents;
+	VOut.Color = FinalVertexColor;
+	VOut.TextureCoordinates[0] = InTextureCoordinates;
+
 #if USE_MATERIALS
-	FMaterialVertexParameters VertexParameters;
+	FMaterialVertexParameters VertexParameters = (FMaterialVertexParameters)0;
 	VertexParameters.WorldPosition = WorldPosition.xyz;
-#if USE_INSTANCING
-	VertexParameters.PerInstanceParams = float4(0,0,0,0);
-#endif
-	VertexParameters.VertexColor = InColor;
+	
+	// Todo do not use per-instance parameters
+	/*#if USE_INSTANCING
+		VertexParameters.PerInstanceParams = float4(0,0,0,0);
+	#endif // USE_INSTANCING*/
+
+	VertexParameters.VertexColor = FinalVertexColor;
 	WorldPosition.xyz += GetMaterialWorldPositionOffset(VertexParameters);
-#endif
+		
+	#if HAS_MATERIAL_TEXCOORDS
+		#if NUM_MATERIAL_TEXCOORDS_VERTEX == 1
+			VertexParameters.TexCoords[0] = InTextureCoordinates.xy;
+		#elif NUM_MATERIAL_TEXCOORDS_VERTEX >= 2
+			VertexParameters.TexCoords[0] = InTextureCoordinates.xy;
+			VertexParameters.TexCoords[1] = InTextureCoordinates.zw;
+		#endif
+
+		float2 CustomizedUVs[NUM_SLATE_TEXCOORDS];
+		GetMaterialCustomizedUVs(VertexParameters, CustomizedUVs);
+
+		UNROLL
+		for( int CoordinateIndex = 0; CoordinateIndex < NUM_SLATE_TEXCOORDS; CoordinateIndex++ )
+		{
+			FLATTEN
+			if (CoordinateIndex % 2)
+			{
+				VOut.TextureCoordinates[CoordinateIndex / 2].zw = CustomizedUVs[CoordinateIndex];
+			}
+			else
+			{
+				VOut.TextureCoordinates[CoordinateIndex / 2].xy = CustomizedUVs[CoordinateIndex];
+			}
+		}
+	#endif // NUM_MATERIAL_TEXCOORDS
+#endif // USE_MATERIALS
 
 	VOut.Position = mul(WorldPosition, ViewProjection);
 
 #if ES2_PROFILE && COMPILER_GLSL_ES2
 	// @todo-mobile: Fix this in the projection matrix
 	VOut.Position.y *= SwitchVerticalAxisMultiplier;
-#endif
-
-	// TextureCoordinates contains both the first and second texture coordinates in xy and zw respectively.
-	VOut.TextureCoordinates = InTextureCoordinates;
-	VOut.MaterialTexCoords = InMaterialTextureCoordinates;
-	VOut.ClipOriginAndPos = float4(InClipOrigin, InPosition.xy);
-	VOut.ClipExtents = InClipExtents;
-
-	InColor.rgb = sRGBToLinear(InColor.rgb);
-	VOut.Color = InColor FCOLOR_COMPONENT_SWIZZLE;
+#endif //ES2_PROFILE && COMPILER_GLSL_ES2
 
 	return VOut;
 }
diff --git a/Engine/Source/Editor/MaterialEditor/Private/MaterialEditorDetailCustomization.cpp b/Engine/Source/Editor/MaterialEditor/Private/MaterialEditorDetailCustomization.cpp
--- a/Engine/Source/Editor/MaterialEditor/Private/MaterialEditorDetailCustomization.cpp
+++ b/Engine/Source/Editor/MaterialEditor/Private/MaterialEditorDetailCustomization.cpp
@@ -403,7 +403,8 @@ void FMaterialDetailCustomization::CustomizeDetails( IDetailLayoutBuilder& Detai
 
 			if(		PropertyName != GET_MEMBER_NAME_CHECKED(UMaterial, MaterialDomain) 
 				&&	PropertyName != GET_MEMBER_NAME_CHECKED(UMaterial, BlendMode) 
-				&&	PropertyName != GET_MEMBER_NAME_CHECKED(UMaterial, OpacityMaskClipValue) )
+				&&	PropertyName != GET_MEMBER_NAME_CHECKED(UMaterial, OpacityMaskClipValue) 
+				&&  	PropertyName != GET_MEMBER_NAME_CHECKED(UMaterial, NumCustomizedUVs) )
 			{
 				DetailLayout.HideProperty( PropertyHandle );
 			}
diff --git a/Engine/Source/Editor/UnrealEd/Classes/MaterialGraph/MaterialGraph.h b/Engine/Source/Editor/UnrealEd/Classes/MaterialGraph/MaterialGraph.h
--- a/Engine/Source/Editor/UnrealEd/Classes/MaterialGraph/MaterialGraph.h
+++ b/Engine/Source/Editor/UnrealEd/Classes/MaterialGraph/MaterialGraph.h
@@ -46,7 +46,18 @@ struct FMaterialInputInfo
 		}
 		else if( Material->IsUIMaterial() )
 		{
-			return Property == MP_EmissiveColor || Property == MP_Opacity || Property == MP_OpacityMask || Property == MP_WorldPositionOffset;
+			if( Property == MP_EmissiveColor || Property == MP_Opacity || Property == MP_OpacityMask || Property == MP_WorldPositionOffset )
+			{
+				return true;
+			}
+
+			if (Property >= MP_CustomizedUVs0 && Property <= MP_CustomizedUVs7)
+			{
+				return (Property - MP_CustomizedUVs0) < Material->NumCustomizedUVs;
+			}
+
+			return false;
+
 		}
 		else
 		{
diff --git a/Engine/Source/Runtime/Engine/Private/Materials/Material.cpp b/Engine/Source/Runtime/Engine/Private/Materials/Material.cpp
--- a/Engine/Source/Runtime/Engine/Private/Materials/Material.cpp
+++ b/Engine/Source/Runtime/Engine/Private/Materials/Material.cpp
@@ -3977,7 +3977,11 @@ bool UMaterial::IsPropertyActive(EMaterialProperty InProperty) const
 		return InProperty == MP_EmissiveColor
 			|| ( InProperty == MP_WorldPositionOffset )
 			|| ( InProperty == MP_OpacityMask && BlendMode == BLEND_Masked ) 
-			|| ( InProperty == MP_Opacity && IsTranslucentBlendMode((EBlendMode)BlendMode) && BlendMode != BLEND_Modulate );
+			|| ( InProperty == MP_Opacity && IsTranslucentBlendMode((EBlendMode)BlendMode) && BlendMode != BLEND_Modulate )
+			|| ( InProperty >= MP_CustomizedUVs0 && InProperty <= MP_CustomizedUVs7);
+		{
+			return true;
+		}
 	}
 
 	bool Active = true;
diff --git a/Engine/Source/Runtime/Engine/Private/Materials/MaterialShared.cpp b/Engine/Source/Runtime/Engine/Private/Materials/MaterialShared.cpp
--- a/Engine/Source/Runtime/Engine/Private/Materials/MaterialShared.cpp
+++ b/Engine/Source/Runtime/Engine/Private/Materials/MaterialShared.cpp
@@ -900,6 +900,11 @@ bool FMaterialResource::IsDefaultMaterial() const
 	return Material->IsDefaultMaterial();
 }
 
+int32 FMaterialResource::GetNumCustomizedUVs() const
+{
+	return Material->NumCustomizedUVs;
+}
+
 void FMaterialResource::NotifyCompilationFinished()
 {
 	Material->NotifyCompilationFinished(this);
diff --git a/Engine/Source/Runtime/Engine/Public/MaterialShared.h b/Engine/Source/Runtime/Engine/Public/MaterialShared.h
--- a/Engine/Source/Runtime/Engine/Public/MaterialShared.h
+++ b/Engine/Source/Runtime/Engine/Public/MaterialShared.h
@@ -978,7 +978,7 @@ public:
 	virtual bool HasNormalConnected() const { return false; }
 	virtual bool RequiresSynchronousCompilation() const { return false; };
 	virtual bool IsDefaultMaterial() const { return false; };
-
+	virtual int32 GetNumCustomizedUVs() const { return 0; }
 	/**
 	 * Should shaders compiled for this material be saved to disk?
 	 */
@@ -1536,6 +1536,7 @@ public:
 	ENGINE_API virtual FString GetFriendlyName() const override;
 	ENGINE_API virtual bool RequiresSynchronousCompilation() const override;
 	ENGINE_API virtual bool IsDefaultMaterial() const override;
+	ENGINE_API virtual int32 GetNumCustomizedUVs() const override;
 	ENGINE_API virtual float GetRefractionDepthBiasValue() const override;
 	ENGINE_API virtual float GetMaxDisplacement() const override;
 	ENGINE_API virtual bool UseTranslucencyVertexFog() const override;
diff --git a/Engine/Source/Runtime/SlateRHIRenderer/Private/SlateMaterialShader.cpp b/Engine/Source/Runtime/SlateRHIRenderer/Private/SlateMaterialShader.cpp
--- a/Engine/Source/Runtime/SlateRHIRenderer/Private/SlateMaterialShader.cpp
+++ b/Engine/Source/Runtime/SlateRHIRenderer/Private/SlateMaterialShader.cpp
@@ -17,6 +17,7 @@ void FSlateMaterialShaderVS::ModifyCompilationEnvironment(EShaderPlatform Platfo
 {
 	// Set defines based on what this shader will be used for
 	OutEnvironment.SetDefine( TEXT("USE_MATERIALS"), 1 );
+	OutEnvironment.SetDefine( TEXT("NUM_CUSTOMIZED_UVS"), Material->GetNumCustomizedUVs() );
 
 	FMaterialShader::ModifyCompilationEnvironment( Platform, Material, OutEnvironment );
 }
@@ -66,6 +67,7 @@ void FSlateMaterialShaderPS::ModifyCompilationEnvironment(EShaderPlatform Platfo
 {
 	// Set defines based on what this shader will be used for
 	OutEnvironment.SetDefine( TEXT("USE_MATERIALS"), 1 );
+	OutEnvironment.SetDefine( TEXT("NUM_CUSTOMIZED_UVS"), Material->GetNumCustomizedUVs() );
 
 	FMaterialShader::ModifyCompilationEnvironment( Platform, Material, OutEnvironment );
 }
diff --git a/Engine/Source/Runtime/SlateRHIRenderer/Private/SlateMaterialShader.h b/Engine/Source/Runtime/SlateRHIRenderer/Private/SlateMaterialShader.h
--- a/Engine/Source/Runtime/SlateRHIRenderer/Private/SlateMaterialShader.h
+++ b/Engine/Source/Runtime/SlateRHIRenderer/Private/SlateMaterialShader.h
@@ -89,7 +89,7 @@ public:
 	{
 		FSlateMaterialShaderVS::ModifyCompilationEnvironment(Platform, Material,OutEnvironment);
 
-		OutEnvironment.SetDefine(TEXT("USE_INSTANCING"), (uint32)(bUseInstancing ? 1 : 0));
+		OutEnvironment.SetDefine(TEXT("USE_SLATE_INSTANCING"), (uint32)(bUseInstancing ? 1 : 0));
 	}
 
 	virtual bool Serialize(FArchive& Ar) override
diff --git a/Engine/Source/Runtime/SlateRHIRenderer/Private/SlateRHIRenderingPolicy.cpp b/Engine/Source/Runtime/SlateRHIRenderer/Private/SlateRHIRenderingPolicy.cpp
--- a/Engine/Source/Runtime/SlateRHIRenderer/Private/SlateRHIRenderingPolicy.cpp
+++ b/Engine/Source/Runtime/SlateRHIRenderer/Private/SlateRHIRenderingPolicy.cpp
@@ -525,8 +525,8 @@ void FSlateRHIRenderingPolicy::DrawElements(FRHICommandListImmediate& RHICmdList
 	// Should only be called by the rendering thread
 	check(IsInRenderingThread());
 
-	float TimeSeconds = FPlatformTime::Seconds() - GStartTime;
-	float RealTimeSeconds = FPlatformTime::Seconds() - GStartTime;
+	float TimeSeconds = FApp::GetCurrentTime() - GStartTime;
+	float RealTimeSeconds =  FApp::GetCurrentTime() - GStartTime;
 	float DeltaTimeSeconds = FApp::GetDeltaTime();
 
 	static const FEngineShowFlags DefaultShowFlags(ESFIM_Game);
@@ -538,10 +538,10 @@ void FSlateRHIRenderingPolicy::DrawElements(FRHICommandListImmediate& RHICmdList
 		FSceneViewFamily::ConstructionValues
 		(
 			&BackBuffer,
-			NULL,
+			nullptr,
 			DefaultShowFlags
 		)
-		.SetWorldTimes( TimeSeconds, RealTimeSeconds, DeltaTimeSeconds )
+		.SetWorldTimes( TimeSeconds, DeltaTimeSeconds, RealTimeSeconds )
 		.SetGammaCorrection( DisplayGamma )
 	);
 
