Fixed: Launching a Project which uses Depth of Field crashes to Home Screen on iOS8 device

diff --git a/Engine/Shaders/PostProcessMobile.usf b/Engine/Shaders/PostProcessMobile.usf
--- a/Engine/Shaders/PostProcessMobile.usf
+++ b/Engine/Shaders/PostProcessMobile.usf
@@ -499,7 +499,7 @@ void DofNearVS_ES2(
 void DofNearPS_ES2(
 	float2 InUVs2 : TEXCOORD0,
 	float4 InUVs[4] : TEXCOORD1,
-	out half OutColor : SV_Target0
+	out half4 OutColor : SV_Target0
 	)
 {
 	half N0 = PostprocessInput0.Sample(PostprocessInput0Sampler, InUVs2).a;
@@ -537,8 +537,9 @@ void DofNearPS_ES2(
 	#endif
 
 	// The first sample is 1/4 the size as the rest of the samples.
-	OutColor = (N0 * 0.25 + N1 + N2 + N3 + N4 + N5 + N6 + N7 + N8) / 8.25;
-	if(OutColor > 0.0) OutColor = sqrt(OutColor);
+	half Out = (N0 * 0.25 + N1 + N2 + N3 + N4 + N5 + N6 + N7 + N8) / 8.25;
+	if(Out > 0.0) Out = sqrt(Out);
+	OutColor.r = Out;
 }
 
 
