Fix ES2 shader compile errors

diff --git a/Engine/Shaders/SlateShaderCommon.usf b/Engine/Shaders/SlateShaderCommon.usf
--- a/Engine/Shaders/SlateShaderCommon.usf
+++ b/Engine/Shaders/SlateShaderCommon.usf
@@ -39,8 +39,14 @@ struct VertexToPixelInterpolants
 
 float2 GetUV(VertexToPixelInterpolants Interpolants, int UVIndex)
 {
+#if ES2_PROFILE
+	// Only support 2 uv coords on ES2 for now.  It doesnt support modulo operations
+	float4 UVVector = Interpolants.TextureCoordinates[UVIndex / 2];
+	return UVIndex == 1 ? UVVector.zw : UVVector.xy;
+#else
 	float4 UVVector = Interpolants.TextureCoordinates[UVIndex / 2];
 	return UVIndex % 2 ? UVVector.zw : UVVector.xy;
+#endif
 }
 
 float4 ApplyColorVisionDeficiencyFilter(float4 Color)
diff --git a/Engine/Shaders/SlateVertexShader.usf b/Engine/Shaders/SlateVertexShader.usf
--- a/Engine/Shaders/SlateVertexShader.usf
+++ b/Engine/Shaders/SlateVertexShader.usf
@@ -54,6 +54,15 @@ VertexToPixelInterpolants Main(
 		float2 CustomizedUVs[NUM_SLATE_TEXCOORDS];
 		GetMaterialCustomizedUVs(VertexParameters, CustomizedUVs);
 
+		#if ES2_PROFILE
+			#if NUM_SLATE_TEXCOORDS == 1
+				VOut.TextureCoordinates[0].xy = CustomizedUVs[0];
+			#elif NUM_SLATE_TEXCOORDS >= 2
+				VOut.TextureCoordinates[0].xy = CustomizedUVs[0];
+				VOut.TextureCoordinates[0].zw = CustomizedUVs[1];
+			#endif
+		#endif
+	#else
 		UNROLL
 		for( int CoordinateIndex = 0; CoordinateIndex < NUM_SLATE_TEXCOORDS; CoordinateIndex++ )
 		{
