Fixed incorrect #if/#else

diff --git a/Engine/Shaders/SlateVertexShader.usf b/Engine/Shaders/SlateVertexShader.usf
--- a/Engine/Shaders/SlateVertexShader.usf
+++ b/Engine/Shaders/SlateVertexShader.usf
@@ -61,21 +61,21 @@ VertexToPixelInterpolants Main(
 				VOut.TextureCoordinates[0].xy = CustomizedUVs[0];
 				VOut.TextureCoordinates[0].zw = CustomizedUVs[1];
 			#endif
-		#endif
-	#else
-		UNROLL
-		for( int CoordinateIndex = 0; CoordinateIndex < NUM_SLATE_TEXCOORDS; CoordinateIndex++ )
-		{
-			FLATTEN
-			if (CoordinateIndex % 2)
-			{
-				VOut.TextureCoordinates[CoordinateIndex / 2].zw = CustomizedUVs[CoordinateIndex];
-			}
-			else
+		#else
+			UNROLL
+			for( int CoordinateIndex = 0; CoordinateIndex < NUM_SLATE_TEXCOORDS; CoordinateIndex++ )
 			{
-				VOut.TextureCoordinates[CoordinateIndex / 2].xy = CustomizedUVs[CoordinateIndex];
+				FLATTEN
+				if (CoordinateIndex % 2)
+				{
+					VOut.TextureCoordinates[CoordinateIndex / 2].zw = CustomizedUVs[CoordinateIndex];
+				}
+				else
+				{
+					VOut.TextureCoordinates[CoordinateIndex / 2].xy = CustomizedUVs[CoordinateIndex];
+				}
 			}
-		}
+		#endif
 	#endif // NUM_MATERIAL_TEXCOORDS
 #endif // USE_MATERIALS
 
