fixed shader after recent changes broke it

diff --git a/Engine/Shaders/ForEachLight.usf b/Engine/Shaders/ForEachLight.usf
--- a/Engine/Shaders/ForEachLight.usf
+++ b/Engine/Shaders/ForEachLight.usf
@@ -63,11 +63,10 @@ float3 AccumulateForwardLights(FMaterialPixelParameters Parameters, float3 Diffu
 	ScreenSpaceData.GBuffer.DiffuseColor = DiffuseColor;
 	ScreenSpaceData.GBuffer.BaseColor = 1;
 	ScreenSpaceData.GBuffer.Specular = 1;
-	ScreenSpaceData.GBuffer.Opacity = GetMaterialOpacity(Parameters);
 	ScreenSpaceData.GBuffer.PrecomputedShadowFactors = 1;
 //	ScreenSpaceData.GBuffer.ShadingModelID = SHADINGMODELID_DEFAULT_LIT;
 	ScreenSpaceData.GBuffer.ShadingModelID = SHADINGMODELID_TWOSIDED_FOLIAGE;
-	ScreenSpaceData.GBuffer.CustomData = DiffuseColor;
+	ScreenSpaceData.GBuffer.CustomData = float4(DiffuseColor, GetMaterialOpacity(Parameters));
 	ScreenSpaceData.AmbientOcclusion = 1;
 
 	// todo: support more than 32 lights
