PR #1556: Add c_ChannelFlags to MeshPaintPixelShader to fix texture painting bug (Contributed by [REDACTED_PERSON])

diff --git a/Engine/Shaders/MeshPaintPixelShader.usf b/Engine/Shaders/MeshPaintPixelShader.usf
--- a/Engine/Shaders/MeshPaintPixelShader.usf
+++ b/Engine/Shaders/MeshPaintPixelShader.usf
@@ -17,6 +17,7 @@ float4 c_BrushMetrics;
 
 float c_BrushStrength;
 float4 c_BrushColor;
+float4 c_ChannelFlags;
 
 // @todo MeshPaint: Remove this?
 float c_Gamma;
@@ -88,9 +89,11 @@ void Main( float4 InPosition : SV_POSITION,
 			PaintAmount *= c_BrushStrength;
 
 			float4 BrushColor = c_BrushColor;
-						
-			NewColor.a = lerp(OldColor.a, BrushColor.a, PaintAmount);
-			NewColor.rgb = lerp(OldColor.rgb * OldColor.a, BrushColor.rgb * BrushColor.a, PaintAmount);
+
+			NewColor.a = lerp(OldColor.a, BrushColor.a, PaintAmount * c_ChannelFlags.a);
+			NewColor.r = lerp(OldColor.r * OldColor.a, BrushColor.r * BrushColor.a, PaintAmount * c_ChannelFlags.r);
+			NewColor.g = lerp(OldColor.g * OldColor.a, BrushColor.g * BrushColor.a, PaintAmount * c_ChannelFlags.g);
+			NewColor.b = lerp(OldColor.b * OldColor.a, BrushColor.b * BrushColor.a, PaintAmount * c_ChannelFlags.b);
 			if (NewColor.a > 0.0f)
 			{
 				NewColor.rgb /= NewColor.a;
