integrated from Orion

optimized shaders to use SvPosition where possible, following change could remove some unused interpolator

diff --git a/Engine/Shaders/BasePassCommon.usf b/Engine/Shaders/BasePassCommon.usf
--- a/Engine/Shaders/BasePassCommon.usf
+++ b/Engine/Shaders/BasePassCommon.usf
@@ -25,10 +25,10 @@ struct FSharedBasePassInterpolants
 #endif
 
 #if !PC_D3D && !GL4_PROFILE
-	float4 PixelPosition	: TEXCOORD8;
+	// Note: TEXCOORD8 is unused
 
 	#if USE_WORLD_POSITION_EXCLUDING_SHADER_OFFSETS
-	float4 PixelPositionExcludingWPO : TEXCOORD9;
+	float3 PixelPositionExcludingWPO : TEXCOORD9;
 	#endif
 #endif
 
@@ -46,10 +46,10 @@ struct FSharedBasePassInterpolants
 /** Interpolants passed from the vertex shader to the pixel shader. */
 struct FBasePassInterpolantsVSToPS : FSharedBasePassInterpolants
 {
-	float4 PixelPosition	: TEXCOORD8;
+	// Note: TEXCOORD8 is unused
 
 	#if USE_WORLD_POSITION_EXCLUDING_SHADER_OFFSETS
-	float4 PixelPositionExcludingWPO : TEXCOORD9;
+	float3 PixelPositionExcludingWPO : TEXCOORD9;
 	#endif
 };
 
diff --git a/Engine/Shaders/BasePassForForwardShadingCommon.usf b/Engine/Shaders/BasePassForForwardShadingCommon.usf
--- a/Engine/Shaders/BasePassForForwardShadingCommon.usf
+++ b/Engine/Shaders/BasePassForForwardShadingCommon.usf
@@ -32,7 +32,7 @@ struct FSharedForwardShadingBasePassInterpolants
 	float4 PixelPosition	: TEXCOORD8; // xyz = world position, w = clip z
 
 #if USE_WORLD_POSITION_EXCLUDING_SHADER_OFFSETS
-	float4 PixelPositionExcludingWPO	: TEXCOORD9;
+	float3 PixelPositionExcludingWPO	: TEXCOORD9;
 #endif
 };
 
diff --git a/Engine/Shaders/BasePassForForwardShadingPixelShader.usf b/Engine/Shaders/BasePassForForwardShadingPixelShader.usf
--- a/Engine/Shaders/BasePassForForwardShadingPixelShader.usf
+++ b/Engine/Shaders/BasePassForForwardShadingPixelShader.usf
@@ -86,9 +86,9 @@ float ClipForEditorPrimitives(FMaterialPixelParameters MaterialParameters)
 			bool bIsPerspective = (View.ViewToClip._m33 < 1.0f);
 
 			// dejitter the sample position and make a filtered lookup - for planes this allows to reconstruct a much less jittery depth comparison function, it however doesn't fix silhuetes
-			float DeviceZ = Texture2DSampleLevel(FilteredSceneDepthTexture, FilteredSceneDepthTextureSampler, (MaterialParameters.SVPosition.xy - View.TemporalAAParams.zw) * View.BufferSizeAndInvSize.zw, 0).r;
+			float DeviceZ = Texture2DSampleLevel(FilteredSceneDepthTexture, FilteredSceneDepthTextureSampler, (MaterialParameters.SvPosition.xy - View.TemporalAAParams.zw) * View.BufferSizeAndInvSize.zw, 0).r;
 
-			float PixelDeviceZ = MaterialParameters.SVPosition.z;
+			float PixelDeviceZ = MaterialParameters.SvPosition.z;
 
 			// Soft Bias with DeviceZ for best quality
 			const float DeviceDepthFade = 0.00005f;
@@ -140,13 +140,16 @@ void Main(
 	float4 PackedInterpolants[NUM_VF_PACKED_INTERPOLANTS];
 	VertexFactoryUnpackInterpolants(Interpolants, PackedInterpolants);
 #endif
-	FMaterialPixelParameters MaterialParameters = GetMaterialPixelParameters(Interpolants, BasePassInterpolants.PixelPosition);
+	FMaterialPixelParameters MaterialParameters = GetMaterialPixelParameters(Interpolants, SvPosition);
 
-	CalcMaterialParameters(MaterialParameters, SvPosition, bIsFrontFace, BasePassInterpolants.PixelPosition
-#if USE_WORLD_POSITION_EXCLUDING_SHADER_OFFSETS
-		, BasePassInterpolants.PixelPositionExcludingWPO
-#endif
-	);
+	#if USE_WORLD_POSITION_EXCLUDING_SHADER_OFFSETS
+		{
+			float4 ScreenPosition = ComputeScreenPosition(SvPosition);
+			CalcMaterialParametersEx(MaterialParameters, SvPosition, ScreenPosition, bIsFrontFace, ComputeWorldWorldPosCamRelative(ScreenPosition), BasePassInterpolants.PixelPositionExcludingWPO);
+		}
+	#else
+		CalcMaterialParameters(MaterialParameters, SvPosition, bIsFrontFace);
+	#endif
 	  
 	//Clip if the blend mode requires it.
 	GetMaterialCoverageAndClipping(MaterialParameters);	
@@ -371,7 +374,7 @@ void Main(
 			OutColor.a = 0.0;
 		#else
 			// Place Z in FP16 alpha value.
-			OutColor.a = BasePassInterpolants.PixelPosition.w;
+			OutColor.a = MaterialParameters.ScreenPosition.w;
 		#endif
 	#endif
 
diff --git a/Engine/Shaders/BasePassForForwardShadingVertexShader.usf b/Engine/Shaders/BasePassForForwardShadingVertexShader.usf
--- a/Engine/Shaders/BasePassForForwardShadingVertexShader.usf
+++ b/Engine/Shaders/BasePassForForwardShadingVertexShader.usf
@@ -55,7 +55,7 @@ void Main(
 	Output.BasePassInterpolants.PixelPosition = WorldPosition;
 
 #if USE_WORLD_POSITION_EXCLUDING_SHADER_OFFSETS
-	Output.BasePassInterpolants.PixelPositionExcludingWPO = WorldPositionExcludingWPO;
+	Output.BasePassInterpolants.PixelPositionExcludingWPO = WorldPositionExcludingWPO.xyz;
 #endif
 
 #if USE_VERTEX_FOG
diff --git a/Engine/Shaders/BasePassPixelShader.usf b/Engine/Shaders/BasePassPixelShader.usf
--- a/Engine/Shaders/BasePassPixelShader.usf
+++ b/Engine/Shaders/BasePassPixelShader.usf
@@ -123,9 +123,9 @@ half3 GetImageBasedReflectionLighting(FMaterialPixelParameters MaterialParameter
 		//uint FrameRandom = (uint)(View.TemporalAAParams.r * 1551);
 		uint FrameRandom = (View.FrameNumber & 15) * 1551;
 
-		uint Morton = MortonCode( (uint)MaterialParameters.SVPosition.x & 3 ) | ( MortonCode( (uint)MaterialParameters.SVPosition.y & 3 ) * 2 );
+		uint Morton = MortonCode( (uint)MaterialParameters.SvPosition.x & 3 ) | ( MortonCode( (uint)MaterialParameters.SvPosition.y & 3 ) * 2 );
 		uint PixelIndex = ReverseBits32( Morton ) >> 28;
-		//uint PixelIndex = ( (uint)MaterialParameters.SVPosition.x & 3 ) | ( ( (uint)MaterialParameters.SVPosition.y & 3 ) * 2 );
+		//uint PixelIndex = ( (uint)MaterialParameters.SvPosition.x & 3 ) | ( ( (uint)MaterialParameters.SvPosition.y & 3 ) * 2 );
 		//PixelIndex = ( PixelIndex * 1551 ) & 15;
 
 		uint Offset = ( PixelIndex + FrameRandom ) & 15;
@@ -560,9 +560,9 @@ float ClipForEditorPrimitives(FMaterialPixelParameters MaterialParameters)
 #endif	// HAS_INVERTED_Z_BUFFER
  		
 		// dejitter the sample position and make a filtered lookup - for planes this allows to reconstruct a much less jittery depth comparison function, it however doesn't fix silhuetes
-		float DeviceZ = Texture2DSampleLevel(FilteredSceneDepthTexture, FilteredSceneDepthTextureSampler, (MaterialParameters.SVPosition.xy - View.TemporalAAParams.zw) * View.BufferSizeAndInvSize.zw, 0).r;
+		float DeviceZ = Texture2DSampleLevel(FilteredSceneDepthTexture, FilteredSceneDepthTextureSampler, (MaterialParameters.SvPosition.xy - View.TemporalAAParams.zw) * View.BufferSizeAndInvSize.zw, 0).r;
 
-		float PixelDeviceZ = MaterialParameters.SVPosition.z;
+		float PixelDeviceZ = MaterialParameters.SvPosition.z;
 
 		// Soft Bias with DeviceZ for best quality
 		const float DeviceDepthFade = 0.00005f;
@@ -649,13 +649,16 @@ void Main(
 	OutCoverage = InCoverage;
 #endif
 
-	FMaterialPixelParameters MaterialParameters = GetMaterialPixelParameters(Interpolants, BasePassInterpolants.PixelPosition);
-
-	CalcMaterialParameters(MaterialParameters, SvPosition, bIsFrontFace, BasePassInterpolants.PixelPosition
-#if USE_WORLD_POSITION_EXCLUDING_SHADER_OFFSETS
-		, BasePassInterpolants.PixelPositionExcludingWPO
-#endif
-	);
+	FMaterialPixelParameters MaterialParameters = GetMaterialPixelParameters(Interpolants, SvPosition);
+	
+	#if USE_WORLD_POSITION_EXCLUDING_SHADER_OFFSETS
+		{
+			float4 ScreenPosition = ComputeScreenPosition(SvPosition);
+			CalcMaterialParametersEx(MaterialParameters, SvPosition, ScreenPosition, bIsFrontFace, ComputeWorldWorldPosCamRelative(ScreenPosition), BasePassInterpolants.PixelPositionExcludingWPO);
+		}
+	#else
+		CalcMaterialParameters(MaterialParameters, SvPosition, bIsFrontFace);
+	#endif
 
 #if USE_EDITOR_COMPOSITING && (FEATURE_LEVEL >= FEATURE_LEVEL_SM4 || MOBILE_EMULATION)
 	const bool bEditorWeightedZBuffering = true;
@@ -924,7 +927,7 @@ void Main(
 		#endif
 
 		// -0.5 .. 0.5, could be optimzed as lower quality noise would be sufficient
-		float QuantizationBias = PseudoRandom( MaterialParameters.SVPosition.xy ) - 0.5f;
+		float QuantizationBias = PseudoRandom( MaterialParameters.SvPosition.xy ) - 0.5f;
 		EncodeGBuffer(GBuffer, OutGBufferA, OutGBufferB, OutGBufferC, OutGBufferD, OutGBufferE, OutGBufferVelocity, QuantizationBias);
 	#endif 
 
diff --git a/Engine/Shaders/BasePassTessellationShaders.usf b/Engine/Shaders/BasePassTessellationShaders.usf
--- a/Engine/Shaders/BasePassTessellationShaders.usf
+++ b/Engine/Shaders/BasePassTessellationShaders.usf
@@ -54,12 +54,9 @@ FBasePassVSToPS PassFinalizeTessellationOutput(FBasePassVSToDS Interpolants, flo
 		#endif
 #endif
 	}
-	
-	// Calc position-dependent interps
-	O.BasePassInterpolants.PixelPosition = WorldPosition;
 
 	#if USE_WORLD_POSITION_EXCLUDING_SHADER_OFFSETS
-		O.BasePassInterpolants.PixelPositionExcludingWPO = float4(Interpolants.BasePassInterpolants.WorldPositionExcludingWPO, 1);
+		O.BasePassInterpolants.PixelPositionExcludingWPO = Interpolants.BasePassInterpolants.WorldPositionExcludingWPO;
 	#endif
 
 	return O;
diff --git a/Engine/Shaders/BasePassVertexShader.usf b/Engine/Shaders/BasePassVertexShader.usf
--- a/Engine/Shaders/BasePassVertexShader.usf
+++ b/Engine/Shaders/BasePassVertexShader.usf
@@ -38,7 +38,7 @@ void Main(
 	Output.Position = WorldPosition;	
 
 	#if USE_WORLD_POSITION_EXCLUDING_SHADER_OFFSETS
-		Output.BasePassInterpolants.WorldPositionExcludingWPO = WorldPositionExcludingWPO;
+		Output.BasePassInterpolants.WorldPositionExcludingWPO = WorldPositionExcludingWPO.xyz;
 	#endif
 #else
 	ISOLATE
@@ -47,8 +47,6 @@ void Main(
 		Output.Position = mul(RasterizedWorldPosition, View.TranslatedWorldToClip);
 	}
 
-	Output.BasePassInterpolants.PixelPosition = WorldPosition;
-
 	#if USE_WORLD_POSITION_EXCLUDING_SHADER_OFFSETS
 		Output.BasePassInterpolants.PixelPositionExcludingWPO = WorldPositionExcludingWPO;
 	#endif
diff --git a/Engine/Shaders/Common.usf b/Engine/Shaders/Common.usf
--- a/Engine/Shaders/Common.usf
+++ b/Engine/Shaders/Common.usf
@@ -941,11 +941,27 @@ float ConvertToDeviceZ(float SceneDepth)
 	return 1.0f / ((SceneDepth + View.InvDeviceZToWorldZTransform[3]) * View.InvDeviceZToWorldZTransform[2]);
 }
 
+/** Returns clip space W, which is world space distance along the View Z axis. */
+float LookupDeviceZ( float2 ScreenUV )
+{
+#if FEATURE_LEVEL > FEATURE_LEVEL_ES3_1
+	return Texture2DSampleLevel(SceneDepthTexture, SceneDepthTextureSampler, ScreenUV, 0).r;
+#else
+	#if COMPILER_GLSL_ES2
+		return ConvertToDeviceZ(CalcSceneDepth(ScreenUV));
+	#elif METAL_PROFILE
+		return ConvertToDeviceZ(CalcSceneDepth(ScreenUV));
+	#else
+		return Texture2DSampleLevel(SceneDepthTexture, SceneDepthTextureSampler, ScreenUV, 0).r;
+	#endif
+#endif
+}
+
 /** Returns clip space W, which is world space distance along the View Z axis. */
 float CalcSceneDepth( float2 ScreenUV )
 {
 #if FEATURE_LEVEL > FEATURE_LEVEL_ES3_1
-	return ConvertFromDeviceZ(Texture2DSampleLevel(SceneDepthTexture, SceneDepthTextureSampler, ScreenUV, 0).r);
+	return ConvertFromDeviceZ(LookupDeviceZ(ScreenUV));
 #else
 	#if COMPILER_GLSL_ES2
 		#if IOS
@@ -960,7 +976,7 @@ float CalcSceneDepth( float2 ScreenUV )
 	#elif METAL_PROFILE
 			return FramebufferFetchES2().w;
 	#else
-		return ConvertFromDeviceZ(Texture2DSampleLevel(SceneDepthTexture, SceneDepthTextureSampler, ScreenUV, 0).r);
+		return ConvertFromDeviceZ(LookupDeviceZ(ScreenUV));
 	#endif
 #endif
 }
diff --git a/Engine/Shaders/ConvertToUniformMesh.usf b/Engine/Shaders/ConvertToUniformMesh.usf
--- a/Engine/Shaders/ConvertToUniformMesh.usf
+++ b/Engine/Shaders/ConvertToUniformMesh.usf
@@ -47,7 +47,9 @@ void ConvertToUniformMeshGS(triangle FConvertToUniformMeshVSToGS Inputs[3], inou
 	for (int i = 0; i < 3; i++)
 	{
 		FConvertToUniformMeshVSToGS Input = Inputs[i];
-		FMaterialPixelParameters MaterialParameters = GetMaterialPixelParameters(Input.FactoryInterpolants, float4(Input.PixelPosition, 1));
+
+		float4 DummySvPosition;
+		FMaterialPixelParameters MaterialParameters = GetMaterialPixelParameters(Input.FactoryInterpolants, DummySvPosition);
 
 		FConvertToUniformMeshVertex Vertex = (FConvertToUniformMeshVertex)0;
 		Vertex.Position = float4(Input.PixelPosition, 40);
diff --git a/Engine/Shaders/DepthOnlyPixelShader.usf b/Engine/Shaders/DepthOnlyPixelShader.usf
--- a/Engine/Shaders/DepthOnlyPixelShader.usf
+++ b/Engine/Shaders/DepthOnlyPixelShader.usf
@@ -21,8 +21,8 @@ void Main(
 {
 #if !MATERIALBLENDING_SOLID || OUTPUT_PIXEL_DEPTH_OFFSET
 
-	FMaterialPixelParameters MaterialParameters = GetMaterialPixelParameters(FactoryInterpolants, PixelPosition);
-	CalcMaterialParameters(MaterialParameters, SvPosition, bIsFrontFace, PixelPosition);
+	FMaterialPixelParameters MaterialParameters = GetMaterialPixelParameters(FactoryInterpolants, SvPosition);
+	CalcMaterialParameters(MaterialParameters, SvPosition, bIsFrontFace);
 
 	#if OUTPUT_PIXEL_DEPTH_OFFSET
 		ApplyPixelDepthOffsetToMaterialParameters(MaterialParameters, OutDepth);
diff --git a/Engine/Shaders/DistortAccumulatePS.usf b/Engine/Shaders/DistortAccumulatePS.usf
--- a/Engine/Shaders/DistortAccumulatePS.usf
+++ b/Engine/Shaders/DistortAccumulatePS.usf
@@ -24,8 +24,8 @@ void Main(
 	)
 {
 	// material parameter inputs
-	FMaterialPixelParameters MaterialParameters = GetMaterialPixelParameters(Interpolants,PixelPosition);
-	CalcMaterialParameters(MaterialParameters, SvPosition, bIsFrontFace, PixelPosition);
+	FMaterialPixelParameters MaterialParameters = GetMaterialPixelParameters(Interpolants, SvPosition);
+	CalcMaterialParameters(MaterialParameters, SvPosition, bIsFrontFace);
 	
 	// material distortion offset
 	half3 Normal = GetMaterialNormal(MaterialParameters);
diff --git a/Engine/Shaders/EvaluateSurfelMaterial.usf b/Engine/Shaders/EvaluateSurfelMaterial.usf
--- a/Engine/Shaders/EvaluateSurfelMaterial.usf
+++ b/Engine/Shaders/EvaluateSurfelMaterial.usf
@@ -1,7 +1,7 @@
 // Copyright 1998-2014 Epic Games, Inc. All Rights Reserved.
 
 /*=============================================================================
-	ConvertToUniformMesh.usf
+	EvaluateSurfelMaterial.usf
 =============================================================================*/
 
 #include "Common.usf"
@@ -56,7 +56,7 @@ void EvaluateSurfelMaterialCS(
 		AssembleMaterialParameters(MaterialParameters, SurfelIndex);
 
 		float4 PixelPosition = float4(PositionAndRadius.xyz + View.PreViewTranslation.xyz, 0);
-		CalcMaterialParameters(MaterialParameters, float4(0, 0, 0, 0), true, PixelPosition);
+		CalcMaterialParameters(MaterialParameters, float4(0, 0, 0, 0), true);
 
 		float3 BaseColor = GetMaterialBaseColor(MaterialParameters);	
 		float Metallic = GetMaterialMetallic(MaterialParameters);
diff --git a/Engine/Shaders/ForEachLight.usf b/Engine/Shaders/ForEachLight.usf
--- a/Engine/Shaders/ForEachLight.usf
+++ b/Engine/Shaders/ForEachLight.usf
@@ -37,7 +37,7 @@ float3 LightingForOneLight(FMaterialPixelParameters Parameters, uint LightIndex,
 	// The only type of shadowing supported for lights using tiled is static shadowing, so the light should only compute shadowing if it has static shadowing
 	LightData.bShadowed = dot(LightData.ShadowMapChannelMask, float4(1, 1, 1, 1));
 	
-	float2 ScreenUV = Parameters.SVPosition.xy * View.BufferSizeAndInvSize.zw;
+	float2 ScreenUV = Parameters.SvPosition.xy * View.BufferSizeAndInvSize.zw;
 
 	return GetDynamicLighting(Parameters.WorldPosition, CameraVector, ScreenUV, ScreenSpaceData, ScreenSpaceData.GBuffer.ShadingModelID, LightData, float4(1, 1, 1, 1), uint2(0,0)).rgb;
 }
@@ -51,7 +51,7 @@ float3 AccumulateForwardLights(FMaterialPixelParameters Parameters, float3 Diffu
 #if !COMPILER_GLSL && !COMPILER_GLSL_ES2
 
 	// Which tile we are?
-	int2 Tile = int2((Parameters.SVPosition.xy - View.ViewRectMin.xy) * ForwardLightData.InvTileSize);
+	int2 Tile = int2((Parameters.SvPosition.xy - View.ViewRectMin.xy) * ForwardLightData.InvTileSize);
 
 	// Get bitmask from CPU for this tile
 	uint culling = LightGrid[Tile.x + Tile.y * ForwardLightData.TileCountX];
diff --git a/Engine/Shaders/GpuSkinVertexFactory.usf b/Engine/Shaders/GpuSkinVertexFactory.usf
--- a/Engine/Shaders/GpuSkinVertexFactory.usf
+++ b/Engine/Shaders/GpuSkinVertexFactory.usf
@@ -121,7 +121,7 @@ struct FVertexFactoryInterpolantsVSToDS
 };
 
 /** Converts from vertex factory specific interpolants to a FMaterialPixelParameters, which is used by material inputs. */
-FMaterialPixelParameters GetMaterialPixelParameters(FVertexFactoryInterpolantsVSToPS Interpolants, float4 PixelPosition)
+FMaterialPixelParameters GetMaterialPixelParameters(FVertexFactoryInterpolantsVSToPS Interpolants, float4 SvPosition)
 {
 	// GetMaterialPixelParameters is responsible for fully initializing the result
 	FMaterialPixelParameters Result = MakeInitializedMaterialPixelParameters();
diff --git a/Engine/Shaders/HitProxyPixelShader.usf b/Engine/Shaders/HitProxyPixelShader.usf
--- a/Engine/Shaders/HitProxyPixelShader.usf
+++ b/Engine/Shaders/HitProxyPixelShader.usf
@@ -25,8 +25,8 @@ void Main(
 	,out float4 OutColor		: SV_Target0
 	)
 {
-	FMaterialPixelParameters MaterialParameters = GetMaterialPixelParameters(Interpolants, PixelPosition);
-	CalcMaterialParameters(MaterialParameters, SvPosition, bIsFrontFace, PixelPosition);
+	FMaterialPixelParameters MaterialParameters = GetMaterialPixelParameters(Interpolants, SvPosition);
+	CalcMaterialParameters(MaterialParameters, SvPosition, bIsFrontFace);
 
 	#if OUTPUT_PIXEL_DEPTH_OFFSET
 		ApplyPixelDepthOffsetToMaterialParameters(MaterialParameters, OutDepth);
diff --git a/Engine/Shaders/LandscapeGrassWeight.usf b/Engine/Shaders/LandscapeGrassWeight.usf
--- a/Engine/Shaders/LandscapeGrassWeight.usf
+++ b/Engine/Shaders/LandscapeGrassWeight.usf
@@ -60,9 +60,9 @@ void PSMain(
 	out half4 OutColor	: SV_Target0
 	)
 {
-	FMaterialPixelParameters MaterialParameters = GetMaterialPixelParameters(Interpolants.FactoryInterpolants, Interpolants.PixelPosition);
+	FMaterialPixelParameters MaterialParameters = GetMaterialPixelParameters(Interpolants.FactoryInterpolants, SvPosition);
 
-	CalcMaterialParameters(MaterialParameters, SvPosition, bIsFrontFace, Interpolants.PixelPosition); 
+	CalcMaterialParameters(MaterialParameters, SvPosition, bIsFrontFace); 
 
 #if NUM_MATERIAL_OUTPUTS_GETGRASSWEIGHT > 2
 	switch(OutputPass)
diff --git a/Engine/Shaders/LandscapeVertexFactory.usf b/Engine/Shaders/LandscapeVertexFactory.usf
--- a/Engine/Shaders/LandscapeVertexFactory.usf
+++ b/Engine/Shaders/LandscapeVertexFactory.usf
@@ -283,7 +283,7 @@ FMaterialTessellationParameters GetMaterialTessellationParameters(FVertexFactory
 }
 
 /** Converts from vertex factory specific interpolants to a FMaterialPixelParameters, which is used by material inputs. */
-FMaterialPixelParameters GetMaterialPixelParameters(FVertexFactoryInterpolantsVSToPS Interpolants, float4 PixelPosition)
+FMaterialPixelParameters GetMaterialPixelParameters(FVertexFactoryInterpolantsVSToPS Interpolants, float4 SvPosition)
 {
 	// GetMaterialPixelParameters is responsible for fully initializing the result
 	FMaterialPixelParameters Result = MakeInitializedMaterialPixelParameters();
diff --git a/Engine/Shaders/LightMapDensityShader.usf b/Engine/Shaders/LightMapDensityShader.usf
--- a/Engine/Shaders/LightMapDensityShader.usf
+++ b/Engine/Shaders/LightMapDensityShader.usf
@@ -158,8 +158,8 @@ void MainPixelShader(
 	out float4 OutColor		: SV_Target0
 	)
 {
-	FMaterialPixelParameters MaterialParameters = GetMaterialPixelParameters( FactoryInterpolants, WorldPosition );
-	CalcMaterialParameters(MaterialParameters, SvPosition, bIsFrontFace, float4(0,0,.00001f,1));
+	FMaterialPixelParameters MaterialParameters = GetMaterialPixelParameters(FactoryInterpolants, SvPosition);
+	CalcMaterialParameters(MaterialParameters, SvPosition, bIsFrontFace);
 	GetMaterialCoverageAndClipping(MaterialParameters);
 
 	float2 LightMapUV,LightMapUV1;
diff --git a/Engine/Shaders/LocalVertexFactory.usf b/Engine/Shaders/LocalVertexFactory.usf
--- a/Engine/Shaders/LocalVertexFactory.usf
+++ b/Engine/Shaders/LocalVertexFactory.usf
@@ -226,7 +226,7 @@ float3 GetInstanceOrigin(FVertexFactoryInput Input)
 #endif	// USE_INSTANCING
 
 /** Converts from vertex factory specific interpolants to a FMaterialPixelParameters, which is used by material inputs. */
-FMaterialPixelParameters GetMaterialPixelParameters(FVertexFactoryInterpolantsVSToPS Interpolants, float4 PixelPosition)
+FMaterialPixelParameters GetMaterialPixelParameters(FVertexFactoryInterpolantsVSToPS Interpolants, float4 SvPosition)
 {
 	// GetMaterialPixelParameters is responsible for fully initializing the result
 	FMaterialPixelParameters Result = MakeInitializedMaterialPixelParameters();
diff --git a/Engine/Shaders/MaterialTemplate.usf b/Engine/Shaders/MaterialTemplate.usf
--- a/Engine/Shaders/MaterialTemplate.usf
+++ b/Engine/Shaders/MaterialTemplate.usf
@@ -115,17 +115,14 @@ struct FMaterialPixelParameters
 	/** World space light vector, only valid when rendering a light function. */
 	half3 LightVector;
 
-	/** Post projection position, before the divide by W. */
-	float4 ScreenPosition;
-
 	/**
 	 * Like SV_Position (.xy is pixel position at pixel center, z:DeviceZ, .w:SceneDepth)
-	 * Can be higher or lower quality (depending on HIGHQUALITY_PS_POSITION).
-	 * low quality: reconstructed from interpolated world position
-	 * high quality: using shader generated value SV_POSITION
-	 * Later we can replace "ScreenPosition" with that.
+	 * using shader generated value SV_POSITION
 	 */
-	float4 SVPosition;
+	float4 SvPosition;
+		
+	/** Post projection position reconstructed from SvPosition, before the divide by W. left..top -1..1, bottom..top -1..1  within the viewport, W is the SceneDepth */
+	float4 ScreenPosition;
 
 	half UnMirrored;
 
@@ -228,7 +225,7 @@ FMaterialPixelParameters MakeInitializedMaterialPixelParameters()
 	#if USE_INSTANCING
 	MPP.PerInstanceParams = 0;
 	#endif
-	MPP.SVPosition = 0;
+	MPP.SvPosition = 0;
 #endif
 	MPP.TangentToWorld = float3x3(1,0,0,0,1,0,0,0,1);
 	return MPP;
@@ -1505,11 +1502,11 @@ float3 GetMaterialHemisphereLightTransferFull(float3 DiffuseColor,FMaterialPixel
 	       LowerLighting * LowerColor;
 }
 
-void ClipLODTransition(float2 SVPosition, float DitherFactor)
+void ClipLODTransition(float2 SvPosition, float DitherFactor)
 {
 	if (abs(DitherFactor) > .001)
 	{
-		float RandCos = cos(dot(floor(SVPosition.xy), float2(347.83451793,3343.28371963)));
+		float RandCos = cos(dot(floor(SvPosition.xy), float2(347.83451793,3343.28371963)));
 		float RandomVal = frac(RandCos * 1000.0);
 		half RetVal = (DitherFactor < 0.0) ?
 			(DitherFactor + 1.0 > RandomVal) :
@@ -1520,7 +1517,7 @@ void ClipLODTransition(float2 SVPosition, float DitherFactor)
 
 void ClipLODTransition(FMaterialPixelParameters Parameters, float DitherFactor)
 {
-	ClipLODTransition(Parameters.SVPosition.xy, DitherFactor);
+	ClipLODTransition(Parameters.SvPosition.xy, DitherFactor);
 }
 
 
@@ -1540,18 +1537,18 @@ void ClipLODTransition(FMaterialPixelParameters Parameters)
 		ClipLODTransition(Parameters, NonInstancedDitherLODFactor);
 	}
 }
-void ClipLODTransition(float2 SVPosition)
+void ClipLODTransition(float2 SvPosition)
 {
 	if (NonInstancedDitherLODFactor != 0.0)
 	{
-		ClipLODTransition(SVPosition, NonInstancedDitherLODFactor);
+		ClipLODTransition(SvPosition, NonInstancedDitherLODFactor);
 	}
 }
 #else
 void ClipLODTransition(FMaterialPixelParameters Parameters)
 {
 }
-void ClipLODTransition(float2 SVPosition)
+void ClipLODTransition(float2 SvPosition)
 {
 }
 #endif
@@ -1594,12 +1591,13 @@ void GetMaterialClippingShadowDepth(FMaterialPixelParameters Parameters)
 #endif
 
 /** Initializes the subset of Parameters that was not set in GetMaterialPixelParameters. */
-void CalcMaterialParameters(
+void CalcMaterialParametersEx(
 	in out FMaterialPixelParameters Parameters,
 	float4 SvPosition,
+	float4 ScreenPosition,
 	FIsFrontFace bIsFrontFace,
-	float4 PixelPosition,
-	float4 PixelPositionExcludingShaderOffsets)
+	float3 PixelPosition,
+	float3 PixelPositionExcludingShaderOffsets)
 {
 	// Remove the pre view translation
 	Parameters.WorldPosition_CamRelative = PixelPosition.xyz;
@@ -1608,27 +1606,20 @@ void CalcMaterialParameters(
 	// If the material uses any non-offset world position expressions, calculate those parameters. If not, 
 	// the variables will have been initialised to 0 earlier.
 #if USE_WORLD_POSITION_EXCLUDING_SHADER_OFFSETS
-	Parameters.WorldPosition_NoOffsets_CamRelative = PixelPositionExcludingShaderOffsets.xyz;
-	Parameters.WorldPosition_NoOffsets = PixelPositionExcludingShaderOffsets.xyz - View.PreViewTranslation.xyz;
+	Parameters.WorldPosition_NoOffsets_CamRelative = PixelPositionExcludingShaderOffsets;
+	Parameters.WorldPosition_NoOffsets = PixelPositionExcludingShaderOffsets - View.PreViewTranslation.xyz;
 #endif
 
-	Parameters.SVPosition = SvPosition;
+	Parameters.SvPosition = SvPosition;
+	Parameters.ScreenPosition = ScreenPosition;
 
-	{
-		float2 PixelPos = SvPosition.xy - View.ViewRectMin.xy;	
-		Parameters.ScreenPosition.xy = ( PixelPos * View.ViewSizeAndInvSize.zw - 0.5 ) * float2(2,-2);
-		Parameters.ScreenPosition.zw = SvPosition.zw;
-		Parameters.ScreenPosition.xyz *= Parameters.ScreenPosition.w;
-	}
-
-	// PixelPosition is the world position translated to the camera position, which is just -CameraVector
 #if COMPILER_GLSL_ES2
 	// ES2 normalize isn't done accurately. This seems to fix it.
 	// Originally this was normalize(normalize(PixelPosition.xyz)) but tegra4 appears to optimize that out.
-	Parameters.CameraVector = normalize(-0.01 * PixelPosition.xyz);
+	Parameters.CameraVector = normalize(-0.01 * Parameters.WorldPosition_CamRelative.xyz);
 #else
 	// PixelPosition is the world position translated to the camera position, which is just -CameraVector
-	Parameters.CameraVector = normalize(-PixelPosition.xyz);
+	Parameters.CameraVector = normalize(-Parameters.WorldPosition_CamRelative.xyz);
 #endif
 
 	Parameters.LightVector = 0;
@@ -1669,13 +1660,39 @@ void CalcMaterialParameters(
 #endif // !PARTICLE_SPRITE_FACTORY
 }
 
+float4 ComputeScreenPosition(float4 SvPosition)
+{
+	float4 ScreenPosition;
+
+	float2 PixelPos = SvPosition.xy - View.ViewRectMin.xy;	
+	ScreenPosition.xy = ( PixelPos * View.ViewSizeAndInvSize.zw - 0.5 ) * float2(2,-2);
+	// clipspace
+	ScreenPosition.zw = SvPosition.zw;
+	// viewspace
+	ScreenPosition.xyz *= ScreenPosition.w;
+
+	return ScreenPosition;
+}
+
+// this will be optimized to not use ScreenPosition
+float3 ComputeWorldWorldPosCamRelative(float4 ScreenPosition)
+{
+	float4 TranslatedWorld = mul( ScreenPosition, View.ClipToTranslatedWorld );
+
+	return TranslatedWorld.xyz / TranslatedWorld.w;
+}
+
+// convenience function to setup CalcMaterialParameters assuming we don't support PixelPositionExcludingShaderOffsets
+// @param SvPosition from SV_Position when rendering the view, for other projections e.g. shadowmaps this function cannot be used and you need to call CalcMaterialParametersEx()
 void CalcMaterialParameters(
 	in out FMaterialPixelParameters Parameters,
 	float4 SvPosition,
-	FIsFrontFace bIsFrontFace,
-	float4 PixelPosition)
+	FIsFrontFace bIsFrontFace)
 {
-	CalcMaterialParameters(Parameters, SvPosition, bIsFrontFace, PixelPosition, PixelPosition);
+	float4 ScreenPosition = ComputeScreenPosition(SvPosition);
+	float3 PixelPosition = ComputeWorldWorldPosCamRelative(ScreenPosition);
+
+	CalcMaterialParametersEx(Parameters, SvPosition, ScreenPosition, bIsFrontFace, PixelPosition, PixelPosition);
 }
 
 /** Assemble the transform from tangent space into world space */
@@ -1733,7 +1750,7 @@ float ApplyPixelDepthOffsetToMaterialParameters(inout FMaterialPixelParameters M
 
 	// Update positions used for shading
 	MaterialParameters.ScreenPosition.w += PixelDepthOffset;
-	MaterialParameters.SVPosition.w = MaterialParameters.ScreenPosition.w;
+	MaterialParameters.SvPosition.w = MaterialParameters.ScreenPosition.w;
 	MaterialParameters.WorldPosition += MaterialParameters.CameraVector * PixelDepthOffset;
 
 	OutDepth = MaterialParameters.ScreenPosition.z / MaterialParameters.ScreenPosition.w;
diff --git a/Engine/Shaders/MeshParticleVertexFactory.usf b/Engine/Shaders/MeshParticleVertexFactory.usf
--- a/Engine/Shaders/MeshParticleVertexFactory.usf
+++ b/Engine/Shaders/MeshParticleVertexFactory.usf
@@ -143,7 +143,7 @@ float4x4 GetParticleTransform(FVertexFactoryInput Input)
 }
 
 /** Converts from vertex factory specific interpolants to a FMaterialPixelParameters, which is used by material inputs. */
-FMaterialPixelParameters GetMaterialPixelParameters(FVertexFactoryInterpolantsVSToPS Interpolants, float4 PixelPosition)
+FMaterialPixelParameters GetMaterialPixelParameters(FVertexFactoryInterpolantsVSToPS Interpolants, float4 SvPosition)
 {
 	// GetMaterialPixelParameters is responsible for fully initializing the result
 	FMaterialPixelParameters Result = MakeInitializedMaterialPixelParameters();
diff --git a/Engine/Shaders/ParticleBeamTrailVertexFactory.usf b/Engine/Shaders/ParticleBeamTrailVertexFactory.usf
--- a/Engine/Shaders/ParticleBeamTrailVertexFactory.usf
+++ b/Engine/Shaders/ParticleBeamTrailVertexFactory.usf
@@ -43,7 +43,7 @@ struct FVertexFactoryIntermediates
 };
 
 /** Converts from vertex factory specific interpolants to a FMaterialPixelParameters, which is used by material inputs. */
-FMaterialPixelParameters GetMaterialPixelParameters(FVertexFactoryInterpolantsVSToPS Interpolants, float4 PixelPosition)
+FMaterialPixelParameters GetMaterialPixelParameters(FVertexFactoryInterpolantsVSToPS Interpolants, float4 SvPosition)
 {
 	// GetMaterialPixelParameters is responsible for fully initializing the result
 	FMaterialPixelParameters Result = MakeInitializedMaterialPixelParameters();
diff --git a/Engine/Shaders/ParticleGPUSpriteVertexFactory.usf b/Engine/Shaders/ParticleGPUSpriteVertexFactory.usf
--- a/Engine/Shaders/ParticleGPUSpriteVertexFactory.usf
+++ b/Engine/Shaders/ParticleGPUSpriteVertexFactory.usf
@@ -342,7 +342,7 @@ void SetUV(inout FVertexFactoryInterpolantsVSToPS Interpolants, int UVIndex, flo
  * @param Interpolants - Attributes interpolated from the vertex shader.
  * @returns per-pixel material parameters.
  */
-FMaterialPixelParameters GetMaterialPixelParameters( FVertexFactoryInterpolantsVSToPS Interpolants, float4 PixelPosition )
+FMaterialPixelParameters GetMaterialPixelParameters( FVertexFactoryInterpolantsVSToPS Interpolants, float4 SvPosition )
 {
 	// GetMaterialPixelParameters is responsible for fully initializing the result
 	FMaterialPixelParameters Result = MakeInitializedMaterialPixelParameters();
@@ -374,7 +374,13 @@ FMaterialPixelParameters GetMaterialPixelParameters( FVertexFactoryInterpolantsV
 #endif
 
 #if USE_SPHERICAL_NORMALS
-	Result.TangentToWorld = GetSphericalParticleNormal(PixelPosition.xyz + View.WorldCameraOrigin, Interpolants.ParticlePositionAndSize.xyz, Interpolants.ParticlePositionAndSize.w);
+	// can be optimized
+	{
+		float4 ScreenPosition = ComputeScreenPosition(SvPosition);
+		float3 PixelPosition = ComputeWorldWorldPosCamRelative(ScreenPosition);
+
+		Result.TangentToWorld = GetSphericalParticleNormal(PixelPosition + View.WorldCameraOrigin, Interpolants.ParticlePositionAndSize.xyz, Interpolants.ParticlePositionAndSize.w);
+	}
 #else
 	half3 TangentToWorld0 = Interpolants.TangentToWorld0.xyz;
 	half4 TangentToWorld2 = Interpolants.TangentToWorld2;
diff --git a/Engine/Shaders/ParticleSpriteVertexFactory.usf b/Engine/Shaders/ParticleSpriteVertexFactory.usf
--- a/Engine/Shaders/ParticleSpriteVertexFactory.usf
+++ b/Engine/Shaders/ParticleSpriteVertexFactory.usf
@@ -149,7 +149,7 @@ void SetUV(inout FVertexFactoryInterpolantsVSToPS Interpolants, int UVIndex, flo
 #endif
 
 /** Converts from vertex factory specific interpolants to a FMaterialPixelParameters, which is used by material inputs. */
-FMaterialPixelParameters GetMaterialPixelParameters(FVertexFactoryInterpolantsVSToPS Interpolants, float4 PixelPosition)
+FMaterialPixelParameters GetMaterialPixelParameters(FVertexFactoryInterpolantsVSToPS Interpolants, float4 SvPosition)
 {
 	// GetMaterialPixelParameters is responsible for fully initializing the result
 	FMaterialPixelParameters Result = MakeInitializedMaterialPixelParameters();
@@ -170,7 +170,12 @@ FMaterialPixelParameters GetMaterialPixelParameters(FVertexFactoryInterpolantsVS
 	float SubImageLerp = Interpolants.TangentToWorld0AndInterp_Sizer.w;
 
 	#if GENERATE_SPHERICAL_PARTICLE_NORMALS && USE_PARTICLE_POSITION
-		Result.TangentToWorld = GetSphericalParticleNormal(PixelPosition.xyz + View.WorldCameraOrigin, Interpolants.ParticlePositionAndSize.xyz, Interpolants.ParticlePositionAndSize.w);
+		{	
+			// can be optimized
+			float4 ScreenPosition = ComputeScreenPosition(SvPosition);
+			float3 PixelPosition = ComputeWorldWorldPosCamRelative(ScreenPosition);
+			Result.TangentToWorld = GetSphericalParticleNormal(PixelPosition + View.WorldCameraOrigin, Interpolants.ParticlePositionAndSize.xyz, Interpolants.ParticlePositionAndSize.w);
+		}
 	#else
 		Result.TangentToWorld = AssembleTangentToWorld(TangentToWorld0.xyz, TangentToWorld2);
 	#endif
diff --git a/Engine/Shaders/PostProcessAmbient.usf b/Engine/Shaders/PostProcessAmbient.usf
--- a/Engine/Shaders/PostProcessAmbient.usf
+++ b/Engine/Shaders/PostProcessAmbient.usf
@@ -444,7 +444,7 @@ float3 ApproximateSpecularIBL( uint2 Random, float3 SpecularColor, float Roughne
 	return PrefilteredColor * ( SpecularColor * AB.x + AB.y );
 }
 
-void MainPS(in noperspective float4 UVAndScreenPos : TEXCOORD0, out float4 OutColor : SV_Target0)
+void MainPS(in noperspective float4 UVAndScreenPos : TEXCOORD0, float4 SvPosition : SV_POSITION, out float4 OutColor : SV_Target0)
 {
 	float2 UV = UVAndScreenPos.xy;
 
@@ -454,7 +454,7 @@ void MainPS(in noperspective float4 UVAndScreenPos : TEXCOORD0, out float4 OutCo
 
 	// screen position in [-1, 1] screen space
 	float2 ScreenSpacePos = UVAndScreenPos.zw;
-	int2 PixelPos = int2(UVAndScreenPos.zw * ScreenPosToPixel.xy + ScreenPosToPixel.zw + 0.5f);
+	int2 PixelPos = (int2)SvPosition.xy;
 
 	float3 ScreenVector = normalize(mul(float4(ScreenSpacePos, 1, 0), View.ScreenToWorld).xyz);
 
diff --git a/Engine/Shaders/PostProcessAmbientOcclusion.usf b/Engine/Shaders/PostProcessAmbientOcclusion.usf
--- a/Engine/Shaders/PostProcessAmbientOcclusion.usf
+++ b/Engine/Shaders/PostProcessAmbientOcclusion.usf
@@ -500,7 +500,7 @@ float3 ReconstructNormal(float2 In)
 }
 
 // the main pixel shader that computes ambient occlusion
-void MainPS(in noperspective float4 UVAndScreenPos : TEXCOORD0, out float4 OutColor : SV_Target0)
+void MainPS(in noperspective float4 UVAndScreenPos : TEXCOORD0, float4 SvPosition : SV_POSITION, out float4 OutColor : SV_Target0)
 {
 	OutColor = 0;
 
@@ -676,7 +676,7 @@ void MainPS(in noperspective float4 UVAndScreenPos : TEXCOORD0, out float4 OutCo
 		float4 dX = ddx_fine(CenterPixel);
 		float4 dY = ddy_fine(CenterPixel);
 
-		int2 Mod = (uint2)ComputePixelPosCenter(ScreenPos, true) % 2;
+		int2 Mod = (uint2)(SvPosition.xy) % 2;
 
 		float4 PixA = CenterPixel;
 		float4 PixB = CenterPixel - dX * (Mod.x * 2 - 1);
diff --git a/Engine/Shaders/PostProcessBloom.usf b/Engine/Shaders/PostProcessBloom.usf
--- a/Engine/Shaders/PostProcessBloom.usf
+++ b/Engine/Shaders/PostProcessBloom.usf
@@ -25,14 +25,15 @@ void MainPostprocessCommonVS(
 }
 
 void VisualizeBloomSetupPS(
-	noperspective float4 UVAndScreenPos : TEXCOORD0
-	, out float4 OutColor : SV_Target0)
+	noperspective float4 UVAndScreenPos : TEXCOORD0,
+	float4 SvPosition : SV_POSITION, 
+	out float4 OutColor : SV_Target0)
 {
 	float2 UV = UVAndScreenPos.xy;
 
 	OutColor = 0;
 
-	int2 PixelPos = (int2)ComputePixelPosCenter(UVAndScreenPos.zw, true);
+	int2 PixelPos = (int2)SvPosition.xy;
 
 	int2 ViewportCenter = (int2)(ViewportRect.xy + ViewportSize.xy / 2);
 
@@ -56,11 +57,12 @@ void VisualizeBloomSetupPS(
 }
 
 void VisualizeBloomOverlayPS(
-	noperspective float4 UVAndScreenPos : TEXCOORD0
-	, out float4 OutColor : SV_Target0)
+	noperspective float4 UVAndScreenPos : TEXCOORD0,
+	float4 SvPosition : SV_POSITION,
+	out float4 OutColor : SV_Target0)
 {
 	float2 UV = UVAndScreenPos.xy;
-	int2 PixelPos = (int2)ComputePixelPosCenter(UVAndScreenPos.zw, true);
+	int2 PixelPos = (int2)SvPosition.xy;
 
 	int2 ViewportCenter = (int2)(ViewportRect + ViewportSize / 2);
 
diff --git a/Engine/Shaders/PostProcessBokehDOF.usf b/Engine/Shaders/PostProcessBokehDOF.usf
--- a/Engine/Shaders/PostProcessBokehDOF.usf
+++ b/Engine/Shaders/PostProcessBokehDOF.usf
@@ -379,11 +379,12 @@ void MainPS(
 // pixel shader to combine the full res scene and the blurred images behind and in front of the the focal plane
 void MainRecombinePS(
 	noperspective in float4 UVAndScreenPos : TEXCOORD0,
+	float4 SvPosition : SV_POSITION,
 	out float4 OutColor : SV_Target0
 	)
 {
 	// SceneColor in full res
-	float2 PixelPosCenter = UVAndScreenPos.zw * ScreenPosToPixel.xy + ScreenPosToPixel.zw + 0.5f;
+	float2 PixelPosCenter = SvPosition.xy;
 
 	float2 FullResUV = PixelPosCenter * PostprocessInput0Size.zw;
 
diff --git a/Engine/Shaders/PostProcessCommon.usf b/Engine/Shaders/PostProcessCommon.usf
--- a/Engine/Shaders/PostProcessCommon.usf
+++ b/Engine/Shaders/PostProcessCommon.usf
@@ -62,7 +62,7 @@ SamplerState BilinearTextureSampler1;
 float4 ScreenPosToPixel;
 
 // uses ScreenPosToPixel which is setup in FPostProcessPassParameters
-// for integer output bPixelCenter should be set to true (to avoid rounding errors)
+// for integer output bPixelCenter should be set to true (to avoid rounding errors), but often it's better to use "int2 PixelPos = (int2)SvPosition.xy;" instead
 // @param ScreenPos -1 .. 1, usually from UVAndScreenPos.zw
 // @param bCenter false: left top of the pixel and almost no fractional value, true: added 0.5 to be pixel center
 float2 ComputePixelPosCenter(float2 ScreenPos, bool bPixelCenter)
diff --git a/Engine/Shaders/PostProcessDOF.usf b/Engine/Shaders/PostProcessDOF.usf
--- a/Engine/Shaders/PostProcessDOF.usf
+++ b/Engine/Shaders/PostProcessDOF.usf
@@ -135,12 +135,13 @@ float4 DepthOfFieldUVLimit;
 
 // pixel shader to combine the full res scene and the blurred images behind and in front of the the focal plane
 void MainRecombinePS(
-	in float4 UVAndScreenPos : TEXCOORD0,
+	noperspective in float4 UVAndScreenPos : TEXCOORD0,
+	float4 SvPosition : SV_POSITION,
 	out float4 OutColor : SV_Target0
 	)
 {
 	// SceneColor in full res
-	float2 PixelPosCenter = UVAndScreenPos.zw * ScreenPosToPixel.xy + ScreenPosToPixel.zw + 0.5f;
+	float2 PixelPosCenter = SvPosition.xy;
 
 	float2 FullResUV = PixelPosCenter * PostprocessInput0Size.zw;
 
diff --git a/Engine/Shaders/PostProcessDilation.usf b/Engine/Shaders/PostProcessDilation.usf
--- a/Engine/Shaders/PostProcessDilation.usf
+++ b/Engine/Shaders/PostProcessDilation.usf
@@ -9,12 +9,12 @@
 
 
 // pixel shader
-void MainPS(float4 UVAndScreenPos : TEXCOORD0, out float4 OutColor : SV_Target0)
+void MainPS(noperspective float4 UVAndScreenPos : TEXCOORD0, float4 SvPosition : SV_POSITION, out float4 OutColor : SV_Target0)
 {
 	float2 UV = UVAndScreenPos.xy;
 	// screen position in [-1, 1] screen space
 	float2 ScreenSpacePos = UVAndScreenPos.zw;
-	int3 PixelPos = int3((int2)(ScreenSpacePos * ScreenPosToPixel.xy + ScreenPosToPixel.zw + 0.5f), 0);
+	int3 PixelPos = int3((int2)(SvPosition.xy), 0);
 
 	OutColor = PostprocessInput0.Load(PixelPos);
 
diff --git a/Engine/Shaders/PostProcessGBufferHints.usf b/Engine/Shaders/PostProcessGBufferHints.usf
--- a/Engine/Shaders/PostProcessGBufferHints.usf
+++ b/Engine/Shaders/PostProcessGBufferHints.usf
@@ -97,10 +97,10 @@ void HighlightAdvice(inout float3 OutColor, float2 UV, uint2 PixelPos)
 	}
 }
 
-void MainPS(float4 UVAndScreenPos : TEXCOORD0, out float4 OutColor : SV_Target0)
+void MainPS(noperspective float4 UVAndScreenPos : TEXCOORD0, float4 SvPosition : SV_POSITION, out float4 OutColor : SV_Target0)
 {
 	float2 UV = UVAndScreenPos.xy;
- 	uint2 PixelPos = (uint2)(UVAndScreenPos.zw * ScreenPosToPixel.xy + ScreenPosToPixel.zw + 0.5f);
+	int2 PixelPos = (uint2)SvPosition.xy;
 	
 	// background is the scene color
 	OutColor = Texture2DSample(PostprocessInput0, PostprocessInput0Sampler, UV);
diff --git a/Engine/Shaders/PostProcessHistogramReduce.usf b/Engine/Shaders/PostProcessHistogramReduce.usf
--- a/Engine/Shaders/PostProcessHistogramReduce.usf
+++ b/Engine/Shaders/PostProcessHistogramReduce.usf
@@ -14,7 +14,7 @@ float4 EyeAdapationTemporalParams;
 // the last frames eye adaptation settings
 Texture2D EyeAdaptationTexture;
 
-void MainPS(float4 UVAndScreenPos : TEXCOORD0, float4 SVPos: SV_POSITION, out float4 OutColor : SV_Target0)
+void MainPS(noperspective float4 UVAndScreenPos : TEXCOORD0, float4 SvPosition: SV_POSITION, out float4 OutColor : SV_Target0)
 {
 	float2 UV = UVAndScreenPos.xy;
 
@@ -27,7 +27,7 @@ void MainPS(float4 UVAndScreenPos : TEXCOORD0, float4 SVPos: SV_POSITION, out fl
 		SceneColor += Texture2DSample(PostprocessInput0, PostprocessInput0Sampler, UV + float2(0, y * PostprocessInput0Size.w));
 	}
 
-	if(SVPos.y < 1.0f)
+	if(SvPosition.y < 1.0f)
 	{
 		// line 0: histogram
 		OutColor = SceneColor / LoopSize;
diff --git a/Engine/Shaders/PostProcessLpvIndirect.usf b/Engine/Shaders/PostProcessLpvIndirect.usf
--- a/Engine/Shaders/PostProcessLpvIndirect.usf
+++ b/Engine/Shaders/PostProcessLpvIndirect.usf
@@ -55,7 +55,7 @@ float3 LpvSpecularHighQuality( float3 WorldPosition, float3 R, float Roughness )
 	return specularOut;
 }
 
-void MainPS(in float4 UVAndScreenPos : TEXCOORD0, float4 SVPos: SV_POSITION, out float4 OutColor : SV_Target0
+void MainPS(in noperspective float4 UVAndScreenPos : TEXCOORD0, float4 SVPos: SV_POSITION, out float4 OutColor : SV_Target0
 #if APPLY_SEPARATE_SPECULAR_RT
 			,out float4 OutSpecular : SV_Target1
 #endif
@@ -152,7 +152,7 @@ void MainPS(in float4 UVAndScreenPos : TEXCOORD0, float4 SVPos: SV_POSITION, out
 
 
 
-void DirectionalOcclusionPS(in float4 UVAndScreenPos : TEXCOORD0, float4 SVPos: SV_POSITION, out float4 OutColor : SV_Target0 )
+void DirectionalOcclusionPS(in noperspective float4 UVAndScreenPos : TEXCOORD0, float4 SVPos: SV_POSITION, out float4 OutColor : SV_Target0 )
 {
 	OutColor = float4(1,1,1,1);
 #if ENABLE_DIRECTIONAL_OCCLUSION 
diff --git a/Engine/Shaders/PostProcessMaterialShaders.usf b/Engine/Shaders/PostProcessMaterialShaders.usf
--- a/Engine/Shaders/PostProcessMaterialShaders.usf
+++ b/Engine/Shaders/PostProcessMaterialShaders.usf
@@ -25,7 +25,8 @@ void MainPS(
 {
 	FMaterialPixelParameters Parameters = MakeInitializedMaterialPixelParameters();
 
-	float2 ScreenUV = ScreenAlignedPosition(ScreenPosition);
+	// can be optimized
+	float2 ScreenUV = ScreenAlignedPosition(ScreenPosition).xy;
 
 #if NUM_MATERIAL_TEXCOORDS
 	for(int CoordinateIndex = 0;CoordinateIndex < NUM_MATERIAL_TEXCOORDS;CoordinateIndex++)
@@ -35,15 +36,13 @@ void MainPS(
 #endif
 
 	Parameters.VertexColor = 1;
-
-	// Deproject quad pixel position into world
-	float SceneDepth = CalcSceneDepth(ScreenUV);
-	float3 ScreenVector = mul(float4(ScreenPosition.xy, 1, 0), View.ScreenToWorld).xyz;
-	float3 WorldPosition = ScreenVector * SceneDepth;
+	
+	SvPosition.z = LookupDeviceZ(ScreenUV);
+	// according to MaterialTemplate.usf, not sure about the .w
+	SvPosition.w = ConvertFromDeviceZ(SvPosition.z);
 
 	// fill out other related material parameters
-	CalcMaterialParameters(Parameters, SvPosition, true, float4(WorldPosition, 1));
-	Parameters.ScreenPosition = ScreenPosition;
+	CalcMaterialParameters(Parameters, SvPosition, true);
 
 	// Grab emissive colour as output
 	OutColor = float4(GetMaterialEmissive(Parameters), 1);
diff --git a/Engine/Shaders/PostProcessMotionBlur.usf b/Engine/Shaders/PostProcessMotionBlur.usf
--- a/Engine/Shaders/PostProcessMotionBlur.usf
+++ b/Engine/Shaders/PostProcessMotionBlur.usf
@@ -328,7 +328,7 @@ uint GetStepCountFromQuality()
 //   1: blurred quarter res velocity from MotionBlurSetup
 //   2: half res velocity from MotionBlurSetup
 // half resolution output: RGB: color, A:blend to full res factor
-void MainPS(float4 UVAndScreenPos : TEXCOORD0, out float4 OutColor : SV_Target0)
+void MainPS(noperspective float4 UVAndScreenPos : TEXCOORD0, float4 SvPosition : SV_POSITION, out float4 OutColor : SV_Target0)
 {
 	// Screen Quad UV 0..1
 	float2 UV = UVAndScreenPos.xy;
@@ -624,7 +624,7 @@ void MainPS(float4 UVAndScreenPos : TEXCOORD0, out float4 OutColor : SV_Target0)
 		const int GapBetweenBuffers = 8;
 		const int FadeRegion = 16;
 
-		float2 PixelPos = ComputePixelPosCenter(UVAndScreenPos.zw, true);
+		float2 PixelPos = SvPosition.xy;
 		int2 PixelPosInt = ((uint2)PixelPos - LeftTop) / Scale; 
 
 		// 0  and 1 are valid buffers
@@ -849,7 +849,6 @@ void MainNewPS(
 
 	const uint StepCount = GetStepCountFromQuality() / 2;
 
-	//uint2 PixelPos = uint2(UVAndScreenPos.zw * ScreenPosToPixel.xy + ScreenPosToPixel.zw + 0.5f);
 	uint2 PixelPos = SvPosition.xy;
 #if 0
 	float2 PosMod = float2( PixelPos & 1 );
diff --git a/Engine/Shaders/PostProcessNoiseBlur.usf b/Engine/Shaders/PostProcessNoiseBlur.usf
--- a/Engine/Shaders/PostProcessNoiseBlur.usf
+++ b/Engine/Shaders/PostProcessNoiseBlur.usf
@@ -20,13 +20,11 @@ float4 NoiseParams;
 	error
 #endif
 
-void MainPS(float4 UVAndScreenPos : TEXCOORD0, out float4 OutColor : SV_Target0)
+void MainPS(noperspective float4 UVAndScreenPos : TEXCOORD0, float4 SvPosition : SV_POSITION, out float4 OutColor : SV_Target0)
 {
 	float2 UV = UVAndScreenPos.xy;
 
-	float2 ScreenPos = UVAndScreenPos.zw;
-
-	float2 PixelPos = ScreenPos * ScreenPosToPixel.xy + ScreenPosToPixel.zw + 0.5f;
+	float2 PixelPos = SvPosition;
 	float Rnd = PseudoRandom(PixelPos); 
 
 	float2 Offset = PostprocessInput0Size.zw;
diff --git a/Engine/Shaders/PostProcessSelectionOutline.usf b/Engine/Shaders/PostProcessSelectionOutline.usf
--- a/Engine/Shaders/PostProcessSelectionOutline.usf
+++ b/Engine/Shaders/PostProcessSelectionOutline.usf
@@ -208,10 +208,10 @@ float4 PerSample(int2 PixelPos, int SampleID, float3 DeviceZPlane, float2 Device
 	return Ret;
 }
 
-void MainPS(float4 UVAndScreenPos : TEXCOORD0, out float4 OutColor : SV_Target0)
+void MainPS(noperspective float4 UVAndScreenPos : TEXCOORD0, float4 SvPosition : SV_POSITION, out float4 OutColor : SV_Target0)
 {
 	float2 UV = UVAndScreenPos.xy;
- 	int2 PixelPos = int2(UVAndScreenPos.zw * ScreenPosToPixel.xy + ScreenPosToPixel.zw + 0.5f);
+ 	int2 PixelPos = (int2)SvPosition.xy;
 
 	float SceneDepth = CalcSceneDepth(UV);
 
diff --git a/Engine/Shaders/PostProcessSubsurface.usf b/Engine/Shaders/PostProcessSubsurface.usf
--- a/Engine/Shaders/PostProcessSubsurface.usf
+++ b/Engine/Shaders/PostProcessSubsurface.usf
@@ -147,13 +147,13 @@ float ComputeFullResLerp(FScreenSpaceData ScreenSpaceData, float2 UVSceneColor,
 }
 
 // visualization (doesn't have to be fast)
-void VisualizePS(in noperspective float4 UVAndScreenPos : TEXCOORD0, out float4 OutColor : SV_Target0)
+void VisualizePS(in noperspective float4 UVAndScreenPos : TEXCOORD0, float4 SvPosition : SV_POSITION, out float4 OutColor : SV_Target0)
 {
 	float2 UV = UVAndScreenPos.xy;
 
 	OutColor = Texture2DSample(PostprocessInput0, PostprocessInput0Sampler, UV);
 
-	int2 PixelPos = (int2)(UVAndScreenPos.zw * ScreenPosToPixel.xy + ScreenPosToPixel.zw + 0.5f);
+	int2 PixelPos = (int2)SvPosition.xy;
 
 	float2 ViewLocalUV = (PixelPos - View.ViewRectMin.xy) * View.ViewSizeAndInvSize.zw;
 
@@ -389,15 +389,14 @@ void MainPS(noperspective float4 UVAndScreenPos : TEXCOORD0, out float4 OutColor
 
 // Recombines the half res Subsurface filtered lighting contribution (upsampled and renormalized with the alpha)
 // with the SceneColor.
-void SubsurfaceRecombinePS(noperspective float4 UVAndScreenPos : TEXCOORD0, out float4 OutColor : SV_Target0)
+void SubsurfaceRecombinePS(noperspective float4 UVAndScreenPos : TEXCOORD0, float4 SvPosition : SV_POSITION, out float4 OutColor : SV_Target0)
 {
 	float2 ViewportUV = UVAndScreenPos.xy;
 
 	// recombine with the scene color
 
 	// can be optimized
-	float2 PixelPos = UVAndScreenPos.zw * ScreenPosToPixel.xy + ScreenPosToPixel.zw;
-	float2 UVSceneColor = (PixelPos + 0.5f) * View.BufferSizeAndInvSize.zw;
+	float2 UVSceneColor = SvPosition.xy * View.BufferSizeAndInvSize.zw;
 
 	FScreenSpaceData ScreenSpaceData = GetScreenSpaceData(UVSceneColor);
 	
diff --git a/Engine/Shaders/PostProcessTestImage.usf b/Engine/Shaders/PostProcessTestImage.usf
--- a/Engine/Shaders/PostProcessTestImage.usf
+++ b/Engine/Shaders/PostProcessTestImage.usf
@@ -22,13 +22,13 @@ float Quantize(float x, float Count)
 	return ceil(x * Count - 1.0f) / (Count - 1.0f);
 }
 
-void MainPS(noperspective float4 UVAndScreenPos : TEXCOORD0, out float4 OutColor : SV_Target0)
+void MainPS(noperspective float4 UVAndScreenPos : TEXCOORD0, float4 SvPosition : SV_POSITION, out float4 OutColor : SV_Target0)
 {
 	OutColor = 0;
 
 	float2 UV = UVAndScreenPos.xy;
 	// full res
- 	int2 PixelPos = (int2)(UVAndScreenPos.zw * ScreenPosToPixel.xy + ScreenPosToPixel.zw + 0.5f);
+	int2 PixelPos = (int2)SvPosition.xy;
 	// 1/2 res
 	int2 PixelPos2 = PixelPos / 2;
 	// 1/32 res
diff --git a/Engine/Shaders/PostProcessVisualizeBuffer.usf b/Engine/Shaders/PostProcessVisualizeBuffer.usf
--- a/Engine/Shaders/PostProcessVisualizeBuffer.usf
+++ b/Engine/Shaders/PostProcessVisualizeBuffer.usf
@@ -10,12 +10,12 @@
 #include "PostProcessHistogramCommon.usf"
 
 /** Main pixel shader function */
-void MainPS(noperspective float4 UVAndScreenPos : TEXCOORD0, out float4 OutColor : SV_Target0)
+void MainPS(noperspective float4 UVAndScreenPos : TEXCOORD0, float4 SvPosition : SV_POSITION, out float4 OutColor : SV_Target0)
 {
 	float Alpha = 1;
 
 #if DRAWING_TILE
-	int2 PixelPos = (int2)(UVAndScreenPos.zw * ScreenPosToPixel.xy + ScreenPosToPixel.zw + 0.5f);
+	int2 PixelPos = (int2)SvPosition.xy;
 
 	// half viewport size
 	int2 HalfInsetSize = (int2)(ViewportSize.xy / 2);
diff --git a/Engine/Shaders/PostProcessVisualizeDOF.usf b/Engine/Shaders/PostProcessVisualizeDOF.usf
--- a/Engine/Shaders/PostProcessVisualizeDOF.usf
+++ b/Engine/Shaders/PostProcessVisualizeDOF.usf
@@ -26,12 +26,12 @@ int2 CursorPos;
 float4 VisualizeColors[2];
 
 // render in fullres to visualize the half res DOF input from the setup pass
-void VisualizeDOFPS(in float4 UVAndScreenPos : TEXCOORD0, out float4 OutColor : SV_Target0)
+void VisualizeDOFPS(in float4 UVAndScreenPos : TEXCOORD0, float4 SvPosition : SV_POSITION, out float4 OutColor : SV_Target0)
 {
 	OutColor = 0;
 
 	float2 UV = UVAndScreenPos.xy; 
-	int2 PixelPos = (int2)(UVAndScreenPos.zw * ScreenPosToPixel.xy + ScreenPosToPixel.zw + 0.5f);
+	int2 PixelPos = (int2)SvPosition.xy;
 	float2 ViewLocalUV = float2(UVAndScreenPos.z * 0.5f + 0.5f, 0.5f - 0.5f * UVAndScreenPos.w);
 	
 	float SceneDepth = CalcSceneDepth(UV);
diff --git a/Engine/Shaders/PostProcessVisualizeGBuffer.usf b/Engine/Shaders/PostProcessVisualizeGBuffer.usf
--- a/Engine/Shaders/PostProcessVisualizeGBuffer.usf
+++ b/Engine/Shaders/PostProcessVisualizeGBuffer.usf
@@ -96,10 +96,10 @@ float3 ComputeQuarterColor(float2 UV, int2 QuarterTileXY)
 }
 
 /** Main pixel shader function */
-void MainPS(noperspective float4 UVAndScreenPos : TEXCOORD0, out float4 OutColor : SV_Target0)
+void MainPS(noperspective float4 UVAndScreenPos : TEXCOORD0, float4 SvPosition : SV_POSITION, out float4 OutColor : SV_Target0)
 {
 #if GBUFFER_VISUALIZATION_MODE == GBUFFER_VISUALIZATION_MODE_ALL
-	int2 PixelPos = (int2)(UVAndScreenPos.zw * ScreenPosToPixel.xy + ScreenPosToPixel.zw + 0.5f);
+	int2 PixelPos = (int2)SvPosition.xy;
 
 	// half viewport size
 	int2 HalfInsetSize = (int2)(ViewportSize.xy / 2);
diff --git a/Engine/Shaders/PostProcessVisualizeHDR.usf b/Engine/Shaders/PostProcessVisualizeHDR.usf
--- a/Engine/Shaders/PostProcessVisualizeHDR.usf
+++ b/Engine/Shaders/PostProcessVisualizeHDR.usf
@@ -110,10 +110,10 @@ float3 Colorize(float x)
 }
 
 //
-void MainPS(noperspective float4 UVAndScreenPos : TEXCOORD0, out float4 OutColor : SV_Target0)
+void MainPS(noperspective float4 UVAndScreenPos : TEXCOORD0, float4 SvPosition : SV_POSITION, out float4 OutColor : SV_Target0)
 {
 	float2 UV = UVAndScreenPos.xy;
-	int2 PixelPos = (int2)(UVAndScreenPos.zw * ScreenPosToPixel.xy + ScreenPosToPixel.zw + 0.5f);
+	int2 PixelPos = (int2)SvPosition.xy;
 	float2 ViewLocalUV = float2(UVAndScreenPos.z * 0.5f + 0.5f, 0.5f - 0.5f * UVAndScreenPos.w);
 
 	int2 ViewportCenter = (int2)(ViewportRect.xy + ViewportSize.xy / 2);
diff --git a/Engine/Shaders/ScreenSpaceReflections.usf b/Engine/Shaders/ScreenSpaceReflections.usf
--- a/Engine/Shaders/ScreenSpaceReflections.usf
+++ b/Engine/Shaders/ScreenSpaceReflections.usf
@@ -93,11 +93,11 @@ void ScreenSpaceReflectionsStencilPS(in float4 UVAndScreenPos : TEXCOORD0, out f
 	OutColor = 0;
 }
 
-void ScreenSpaceReflectionsPS(in float4 UVAndScreenPos : TEXCOORD0, out float4 OutColor : SV_Target0)
+void ScreenSpaceReflectionsPS(in noperspective float4 UVAndScreenPos : TEXCOORD0, float4 SvPosition : SV_POSITION, out float4 OutColor : SV_Target0)
 {
 	float2 UV = UVAndScreenPos.xy;
 	float2 ScreenPos = UVAndScreenPos.zw;
-	int2 PixelPos = int2(UVAndScreenPos.zw * ScreenPosToPixel.xy + ScreenPosToPixel.zw + 0.5f);
+	int2 PixelPos = (int2)SvPosition.xy;
 
 	OutColor = 0;
 
@@ -149,7 +149,7 @@ void ScreenSpaceReflectionsPS(in float4 UVAndScreenPos : TEXCOORD0, out float4 O
 	const int NumSteps = 16;
 	
 	{
-		const float2 PixelPos = UVAndScreenPos.zw * ScreenPosToPixel.xy + ScreenPosToPixel.zw;
+		const float2 PixelPos = SvPosition.xy;
 		const float2 E = float2(PseudoRandom(PixelPos + int2(FrameRandom, 0)), PseudoRandom(PixelPos + int2(0, FrameRandom)));
 		const float StepOffset = PseudoRandom(PixelPos + FrameRandom);
 
diff --git a/Engine/Shaders/ShadowDepthCommon.usf b/Engine/Shaders/ShadowDepthCommon.usf
--- a/Engine/Shaders/ShadowDepthCommon.usf
+++ b/Engine/Shaders/ShadowDepthCommon.usf
@@ -37,6 +37,7 @@ struct FShadowDepthVSToPS
 #endif
 
 #if INTERPOLATE_POSITION
-	float4 PixelPosition	: TEXCOORD7;
+	// world space position needed in some cases, could be changed to be reconsructed from SvPosition (more tricky as this is not from camera perspective)
+	float3 PixelPosition	: TEXCOORD7;
 #endif
 };
diff --git a/Engine/Shaders/ShadowDepthPixelShader.usf b/Engine/Shaders/ShadowDepthPixelShader.usf
--- a/Engine/Shaders/ShadowDepthPixelShader.usf
+++ b/Engine/Shaders/ShadowDepthPixelShader.usf
@@ -56,14 +56,17 @@ void Main(
 {
 #if INTERPOLATE_VF_ATTRIBUTES
 	
-	FMaterialPixelParameters MaterialParameters = GetMaterialPixelParameters(Inputs.FactoryInterpolants, float4(0,0,.00001f,1));
+	FMaterialPixelParameters MaterialParameters = GetMaterialPixelParameters(Inputs.FactoryInterpolants, SvPosition);
 
 	#if INTERPOLATE_POSITION
-		CalcMaterialParameters(MaterialParameters, SvPosition, 1, Inputs.PixelPosition);
+		{
+			float4 ScreenPosition = ComputeScreenPosition(SvPosition);
+			CalcMaterialParametersEx(MaterialParameters, SvPosition, ScreenPosition, 1, ComputeWorldWorldPosCamRelative(ScreenPosition), SvPosition);
+		}
 	#else
-		// Note: Using default values for PixelPosition to reduce interpolator cost
-		CalcMaterialParameters(MaterialParameters, SvPosition, 1, float4(0,0,.00001f,1));
+		CalcMaterialParameters(MaterialParameters, SvPosition, 1);
 	#endif
+
 	// Evaluate the mask for masked materials
 	GetMaterialClippingShadowDepth(MaterialParameters);
 #else
diff --git a/Engine/Shaders/ShadowDepthVertexShader.usf b/Engine/Shaders/ShadowDepthVertexShader.usf
--- a/Engine/Shaders/ShadowDepthVertexShader.usf
+++ b/Engine/Shaders/ShadowDepthVertexShader.usf
@@ -104,7 +104,7 @@ void MainOnePassPointLightGS(triangle FPassSpecificVSToGS Input[3], inout Triang
 					#endif
 
 #if INTERPOLATE_POSITION
-					Output.PSInputs.PixelPosition = Input[VertexIndex].GSPosition;
+					Output.PSInputs.PixelPosition = Input[VertexIndex].GSPosition.xyz;
 #endif
 					OutStream.Append(Output);
 				}
@@ -210,7 +210,7 @@ void Main(
 	#endif
 
 	#if INTERPOLATE_POSITION
-		OutParameters.PixelPosition = WorldPos;
+		OutParameters.PixelPosition = WorldPos.xyz;
 	#endif
 
 #else
diff --git a/Engine/Shaders/SlateElementPixelShader.usf b/Engine/Shaders/SlateElementPixelShader.usf
--- a/Engine/Shaders/SlateElementPixelShader.usf
+++ b/Engine/Shaders/SlateElementPixelShader.usf
@@ -76,7 +76,7 @@ void FillTexCoords( inout FMaterialPixelParameters Parameters, VertexToPixelInte
 float4 GetMaterialColor( VertexToPixelInterpolants InVertex, float2 MeshUV )
 {
 	FMaterialPixelParameters Parameters = MakeInitializedMaterialPixelParameters();
-	CalcMaterialParameters( Parameters, InVertex.Position, true, InVertex.Position );
+	CalcMaterialParameters( Parameters, InVertex.Position, true );
 
 	Parameters.VertexColor = InVertex.Color;
 
diff --git a/Engine/Shaders/TranslucentShadowDepthShaders.usf b/Engine/Shaders/TranslucentShadowDepthShaders.usf
--- a/Engine/Shaders/TranslucentShadowDepthShaders.usf
+++ b/Engine/Shaders/TranslucentShadowDepthShaders.usf
@@ -79,8 +79,8 @@ void MainOpacityPS(
 	out float4 OutColor1 : SV_Target1
 	)
 {
-	FMaterialPixelParameters MaterialParameters = GetMaterialPixelParameters(Inputs.FactoryInterpolants, Inputs.PixelPosition);
-	CalcMaterialParameters(MaterialParameters, SvPosition, 1, Inputs.PixelPosition);
+	FMaterialPixelParameters MaterialParameters = GetMaterialPixelParameters(Inputs.FactoryInterpolants, SvPosition);
+	CalcMaterialParameters(MaterialParameters, SvPosition, 1);
 	
 	// Evaluate the mask for translucent materials
 	GetMaterialClippingShadowDepth(MaterialParameters);
diff --git a/Engine/Shaders/VectorFieldVisualizationVertexFactory.usf b/Engine/Shaders/VectorFieldVisualizationVertexFactory.usf
--- a/Engine/Shaders/VectorFieldVisualizationVertexFactory.usf
+++ b/Engine/Shaders/VectorFieldVisualizationVertexFactory.usf
@@ -76,7 +76,7 @@ FVertexFactoryIntermediates GetVertexFactoryIntermediates( FVertexFactoryInput I
  * @param Interpolants - Attributes interpolated from the vertex shader.
  * @returns per-pixel material parameters.
  */
-FMaterialPixelParameters GetMaterialPixelParameters( FVertexFactoryInterpolantsVSToPS Interpolants, float4 PixelPosition )
+FMaterialPixelParameters GetMaterialPixelParameters( FVertexFactoryInterpolantsVSToPS Interpolants, float4 SvPosition )
 {
 	// GetMaterialPixelParameters is responsible for fully initializing the result
 	FMaterialPixelParameters Result = MakeInitializedMaterialPixelParameters();
diff --git a/Engine/Shaders/VelocityShader.usf b/Engine/Shaders/VelocityShader.usf
--- a/Engine/Shaders/VelocityShader.usf
+++ b/Engine/Shaders/VelocityShader.usf
@@ -196,7 +196,7 @@ void MainPixelShader(
 {
 	// Manual clipping here (alpha-test, etc)
 	FMaterialPixelParameters MaterialParameters = GetMaterialPixelParameters(Input.FactoryInterpolants, SvPosition);
-	CalcMaterialParameters(MaterialParameters, SvPosition, bIsFrontFace, SvPosition);
+	CalcMaterialParameters(MaterialParameters, SvPosition, bIsFrontFace);
 	GetMaterialCoverageAndClipping(MaterialParameters);
 
 	// 2d velocity, includes camera an object motion
diff --git a/Engine/Shaders/VertexDensityShader.usf b/Engine/Shaders/VertexDensityShader.usf
--- a/Engine/Shaders/VertexDensityShader.usf
+++ b/Engine/Shaders/VertexDensityShader.usf
@@ -75,7 +75,7 @@ float4 MainPS(FVertexDensityVSToPS Input) : SV_Target0
 	else
 	{
 		FMaterialPixelParameters MaterialParameters = GetMaterialPixelParameters(Input.FactoryInterpolants, Input.WorldPosition);
-		CalcMaterialParameters(MaterialParameters, Input.Position, false, float4(0, 0, .00001f, 1));
+		CalcMaterialParameters(MaterialParameters, Input.Position, false);
 		GetMaterialCoverageAndClipping(MaterialParameters);
 		float3 Color = MaterialParameters.AOMaterialMask + MaterialParameters.VertexColor.xyz;
 		Color = normalize(Color) * 0.5 + float3(0, 0, 0.25);
