fixed build for ES shader compile

diff --git a/Engine/Shaders/Common.usf b/Engine/Shaders/Common.usf
--- a/Engine/Shaders/Common.usf
+++ b/Engine/Shaders/Common.usf
@@ -948,9 +948,11 @@ float LookupDeviceZ( float2 ScreenUV )
 	return Texture2DSampleLevel(SceneDepthTexture, SceneDepthTextureSampler, ScreenUV, 0).r;
 #else
 	#if COMPILER_GLSL_ES2
-		return ConvertToDeviceZ(CalcSceneDepth(ScreenUV));
+		// not yet supported
+		return 0;
 	#elif METAL_PROFILE
-		return ConvertToDeviceZ(CalcSceneDepth(ScreenUV));
+		// not yet supported
+		return 0;
 	#else
 		return Texture2DSampleLevel(SceneDepthTexture, SceneDepthTextureSampler, ScreenUV, 0).r;
 	#endif
