Improvements to stop collision response for GPU particles.

diff --git a/Engine/Shaders/ParticleSimulationShader.usf b/Engine/Shaders/ParticleSimulationShader.usf
--- a/Engine/Shaders/ParticleSimulationShader.usf
+++ b/Engine/Shaders/ParticleSimulationShader.usf
@@ -245,7 +245,8 @@ void CollideWithDepthBuffer(
 	in float3 InVelocity,
 	in float3 Acceleration,
 	in float CollisionRadius,
-	in float Resilience
+	in float Resilience,
+	inout float CollisionStop
 	)
 {
 	// Integration assuming no collision.
@@ -292,6 +293,14 @@ void CollideWithDepthBuffer(
 			//if (t_back >= 0 && t_front <= 1 && DistanceToPlane >= 0)
 			if (step(0, t_back) * step(t_front, 1) * step(0, DistanceToPlane))
 			{
+				//When the collision response is "Stop" the resilience is 0.0.
+				//In this case we mark the event so that all future updates of the particle motion are stopped.
+				//Prevents Stopped particles from continuing to move if they go off screen.
+				if (Resilience < 0.0001f)
+				{
+					CollisionStop = 1.0f;
+				}
+
 				// Separate velocity in to the components perpendicular and tangent to the collision plane.
 				float3 PerpVelocity = dot(MidVelocity,CollisionPlane.xyz) * CollisionPlane.xyz;
 				float3 TanVelocity = MidVelocity - PerpVelocity;
@@ -335,7 +344,8 @@ void CollideWithDistanceField(
 	in float3 InVelocity,
 	in float3 Acceleration,
 	in float CollisionRadius,
-	in float Resilience
+	in float Resilience,
+	out float CollisionStop
 	)
 {
 	// Integration assuming no collision.
@@ -368,6 +378,14 @@ void CollideWithDistanceField(
 		//if (t_back >= 0 && t_front <= 1 && DistanceToPlane >= 0)
 		if (step(0, t_back) * step(t_front, 1) * step(0, DistanceToPlane))
 		{
+			//When the collision response is "Stop" the resilience is 0.0.
+			//In this case we mark the event so that all future updates of the particle motion are stopped.
+			//Prevents Stopped particles from continuing to move if they go off screen.
+			if (Resilience < 0.0001f)
+			{
+				CollisionStop = 1.0f;
+			}
+
 			// Separate velocity in to the components perpendicular and tangent to the collision plane.
 			float3 PerpVelocity = dot(MidVelocity,CollisionPlane.xyz) * CollisionPlane.xyz;
 			float3 TanVelocity = MidVelocity - PerpVelocity;
@@ -398,6 +416,17 @@ void CollideWithDistanceField(
 
 #endif
 
+float PackRelativeTimeAndCollisionStop(in float RelativeTime, in float CollisionStop)
+{
+	return CollisionStop > 0.0f ? -RelativeTime : RelativeTime;
+}
+
+float UnpackRelativeTimeAndCollisionStop(in float PackedValue, out float CollisionStop)
+{
+	CollisionStop = PackedValue < 0.0f ? 1.0f : 0.0f;
+	return abs(PackedValue);
+}
+
 void PixelMain(
 	in FShaderInterpolants Interpolants,
 	out float4 OutPosition : SV_Target0,
@@ -419,7 +448,14 @@ void PixelMain(
 
 	// Position.w holds the relative time of the particle.
 	float3 Position = PositionSample.xyz;
-	float RelativeTime = PositionSample.w + DeltaSeconds * TimeScale;
+
+	float CollisionStop;
+	float RelativeTime = UnpackRelativeTimeAndCollisionStop(PositionSample.w, CollisionStop);
+
+	RelativeTime += DeltaSeconds * TimeScale;
+
+	//Everything other that particle aging should stop if we have collision stop.
+	DeltaSeconds = CollisionStop > 0.0f ? 0.0f : DeltaSeconds;
 
 	// Sample the attribute curve.
 	const float2 AttributeCurveTexCoord = Simulation.AttributeCurve.xy +
@@ -486,7 +522,8 @@ void PixelMain(
 		Velocity.xyz,
 		Acceleration,
 		CollisionRadius,
-		Resilience
+		Resilience,
+		CollisionStop
 		);
 #elif DISTANCE_FIELD_COLLISION
 
@@ -500,7 +537,8 @@ void PixelMain(
 		Velocity.xyz,
 		Acceleration,
 		CollisionRadius,
-		Resilience
+		Resilience,
+		CollisionStop
 		);
 
 #else 
@@ -516,7 +554,7 @@ void PixelMain(
 
 	// Store the new position, time, and velocity for the particle.
 	OutPosition.xyz = NewPosition + OrbitVelocity * DeltaSeconds + PositionOffsetAndAttractorStrength.xyz;
-	OutPosition.w = RelativeTime;
+	OutPosition.w = PackRelativeTimeAndCollisionStop(RelativeTime, CollisionStop);
 	OutVelocity.xyz = NewVelocity;
 	OutVelocity.w = TimeScale;
 }
