Merging from Orion

CL 2699552
  added SvPositionToTranslatedWorld for better quality and performance, some refactoring

diff --git a/Engine/Shaders/BasePassForForwardShadingPixelShader.usf b/Engine/Shaders/BasePassForForwardShadingPixelShader.usf
--- a/Engine/Shaders/BasePassForForwardShadingPixelShader.usf
+++ b/Engine/Shaders/BasePassForForwardShadingPixelShader.usf
@@ -144,7 +144,7 @@ void Main(
 
 	#if USE_WORLD_POSITION_EXCLUDING_SHADER_OFFSETS
 		{
-			float4 ScreenPosition = ComputeScreenPosition(SvPosition);
+			float4 ScreenPosition = SvPositionToScreenPosition(SvPosition);
 			CalcMaterialParametersEx(MaterialParameters, SvPosition, ScreenPosition, bIsFrontFace, ComputeWorldWorldPosCamRelative(ScreenPosition), BasePassInterpolants.PixelPositionExcludingWPO);
 		}
 	#else
diff --git a/Engine/Shaders/BasePassPixelShader.usf b/Engine/Shaders/BasePassPixelShader.usf
--- a/Engine/Shaders/BasePassPixelShader.usf
+++ b/Engine/Shaders/BasePassPixelShader.usf
@@ -653,7 +653,7 @@ void Main(
 	
 	#if USE_WORLD_POSITION_EXCLUDING_SHADER_OFFSETS
 		{
-			float4 ScreenPosition = ComputeScreenPosition(SvPosition);
+			float4 ScreenPosition = SvPositionToScreenPosition(SvPosition);
 			CalcMaterialParametersEx(MaterialParameters, SvPosition, ScreenPosition, bIsFrontFace, ComputeWorldWorldPosCamRelative(ScreenPosition), BasePassInterpolants.PixelPositionExcludingWPO);
 		}
 	#else
diff --git a/Engine/Shaders/Common.usf b/Engine/Shaders/Common.usf
--- a/Engine/Shaders/Common.usf
+++ b/Engine/Shaders/Common.usf
@@ -663,20 +663,6 @@ MaterialFloat3x3 GetLocalToWorld3x3()
 	return (MaterialFloat3x3)Primitive.LocalToWorld;
 }
 
-float4 PreviousLighting(float4 ScreenPosition)
-{
-	return Texture2DSample(SceneColorTexture, SceneColorTextureSampler,ScreenPosition.xy / ScreenPosition.w * View.ScreenPositionScaleBias.xy + View.ScreenPositionScaleBias.wz);
-}
-
-/** Fetches the values in the light attenuation buffer. */
-float4 GetLightAttenuation(float4 ScreenPosition)
-{
-	ScreenPosition.xy = float2(ScreenPosition.xy / ScreenPosition.w * View.ScreenPositionScaleBias.xy + View.ScreenPositionScaleBias.wz);
-	float4 Attenuation = Texture2DSample(LightAttenuationTexture, LightAttenuationTextureSampler, ScreenPosition.xy);
-	// Undo the sqrt that was used to encode
-	return Square(Attenuation);
-}
-
 /** Converts a linear input value into a value to be stored in the light attenuation buffer. */
 MaterialFloat EncodeLightAttenuation(MaterialFloat InColor)
 {
@@ -1006,12 +992,6 @@ float4 CalcSceneColorAndDepth( float2 ScreenUV )
 	return float4(CalcSceneColor(ScreenUV), CalcSceneDepth(ScreenUV));
 }
 
-float PreviousDepth(float4 ScreenPosition)
-{
-	ScreenPosition.xy = float2(ScreenPosition.xy / ScreenPosition.w * View.ScreenPositionScaleBias.xy + View.ScreenPositionScaleBias.wz);
-	return CalcSceneDepth(ScreenPosition.xy);
-}
-
 #if FEATURE_LEVEL >= FEATURE_LEVEL_SM4
 #if FEATURE_LEVEL == FEATURE_LEVEL_SM4
 	// SM4 requires texture size to be explicitly stated and expressed as the number of samples.
@@ -1032,20 +1012,54 @@ float PreviousDepth(float4 ScreenPosition)
 
 		return ConvertFromDeviceZ(DeviceZ);
 	}
-	float PreviousDepthMSAA(float4 ScreenPosition,uint SampleIndex)
-	{
-		ScreenPosition.xy = float2(ScreenPosition.xy / ScreenPosition.w * View.ScreenPositionScaleBias.xy + View.ScreenPositionScaleBias.wz);
-		return CalcSceneDepthMSAA(ScreenPosition.xy,SampleIndex);
-	}
 #endif
 	 
+// ----------------------------
+
+float2 ScreenPositionToBufferUV(float4 ScreenPosition)
+{
+	return float2(ScreenPosition.xy / ScreenPosition.w * View.ScreenPositionScaleBias.xy + View.ScreenPositionScaleBias.wz);
+}
+
+float2 SvPositionToBufferUV(float4 SvPosition)
+{
+	return SvPosition.xy * View.BufferSizeAndInvSize.zw;
+}
+	
+float3 SvPositionToWorld(float4 SvPosition)
+{
+	float4 HomWorldPos = mul(float4(SvPosition.xyz, 1), View.SVPositionToTranslatedWorld);
+
+	return HomWorldPos.xyz / HomWorldPos.w - View.PreViewTranslation;
+}
+
+// investigate: doesn't work for usage with View.ScreenToWorld, see SvPositionToScreenPosition2()
+float4 SvPositionToScreenPosition(float4 SvPosition)
+{
+	float2 PixelPos = SvPosition.xy - View.ViewRectMin.xy;	
+
+	// NDC (NormalizedDeviceCoordinates, after the perspective divide)
+	float3 NDCPos = float3( (PixelPos * View.ViewSizeAndInvSize.zw - 0.5f) * float2(2, -2), SvPosition.z);
+
+	// clipspace (before the perspective divide)
+//	return float4(NDCPos.xyz, 1) * SvPosition.w;
+
+	// so .w has the SceneDepth, some mobile code wants that
+//	Ret *= Ret.z;
+
+	// 
+	return float4(NDCPos.xyz, 1);
+}
+
+// ----------------------------
+
 /** 
  * aligns the clip space position so that it can be used as a texture coordinate
  * to properly align in screen space
  */
-MaterialFloat4 ScreenAlignedPosition( float4 ScreenPosition )
+MaterialFloat2 ScreenAlignedPosition( float4 ScreenPosition )
 {
-	return MaterialFloat4(ScreenPosition.xy / ScreenPosition.w * View.ScreenPositionScaleBias.xy + View.ScreenPositionScaleBias.wz, ScreenPosition.z/ScreenPosition.w,1);
+	return MaterialFloat2(ScreenPositionToBufferUV(ScreenPosition));
 }
 
 /** 
diff --git a/Engine/Shaders/DeferredDecal.usf b/Engine/Shaders/DeferredDecal.usf
--- a/Engine/Shaders/DeferredDecal.usf
+++ b/Engine/Shaders/DeferredDecal.usf
@@ -51,11 +51,10 @@ half3 FrameBufferDecalBlendOp(half4 Source)
 // decal vertex shader
 void MainVS(
 	in float4 InPosition : ATTRIBUTE0,
-	out float4 OutScreenPosition : TEXCOORD0,
 	out float4 OutPosition : SV_POSITION
 	)
 {
-	OutPosition = OutScreenPosition = mul(InPosition, FrustumComponentToClip);
+	OutPosition = mul(InPosition, FrustumComponentToClip);
 }
 
 float DistanceFunction(float4 ScreenPosition, float3 Pos)
@@ -135,6 +134,24 @@ float3 GradientNotNormalized(float4 ScreenPosition, float3 OSPos, float t)
 }
 
 
+float4 SvPositionToScreenPosition2(float4 SvPosition)
+{
+	// assumed SvPosition.w = 1
+
+	float4 Ret;
+
+	Ret.x = ((SvPosition.x - View.ViewRectMin.x) * View.ViewSizeAndInvSize.z) * 2 - 1;
+	Ret.y = ((SvPosition.y - View.ViewRectMin.y) * View.ViewSizeAndInvSize.w) * -2 + 1;
+	Ret.z = ConvertFromDeviceZ(SvPosition.z);
+	Ret.w = 1;
+
+	Ret.xy *= Ret.z;
+
+	// so .w has the SceneDepth, some mobile code wants that
+//	Ret *= Ret.z;
+
+	return Ret;
+}
 
 // decal pixel shader
 #if DECAL_PROJECTION
@@ -142,9 +159,7 @@ float3 GradientNotNormalized(float4 ScreenPosition, float3 OSPos, float t)
 // EARLYDEPTHSTENCIL for volumetric decals could work here, but only if we are ouside of the decal
 
 void MainPS(
-	in float4 ScreenPosition : TEXCOORD0			// is ClipPosition
-	
-	, in INPUT_POSITION_QUALIFIERS float4 SvPosition : SV_Position
+	in INPUT_POSITION_QUALIFIERS float4 SvPosition : SV_Position		// after all interpolators
 
 #if DECAL_RENDERTARGET_COUNT > 0
 	, out float4 OutTarget0 : SV_Target0
@@ -171,6 +186,14 @@ void MainPS(
 
 	)
 {
+	float2 ScreenUV = SvPositionToBufferUV(SvPosition);
+
+	// make SvPosition appear to be rasterized with the depth from the depth buffer
+	SvPosition.z = LookupDeviceZ(ScreenUV);
+//	SvPosition.w = 1;
+
+	float4 ScreenPosition = SvPositionToScreenPosition(SvPosition);
+
 #if DECAL_RENDERTARGET_COUNT > 4
 	// default to no shadows, good for volumetric decals
 	OutTarget4 = 1;
@@ -187,14 +210,14 @@ void MainPS(
 		// LineBoxIntersect want's to clip a long ray so we provide one, this could be done differently as it limites the draw distance and reduces precision 
 		float SceneW = 100000000;
 
-		float4 ProjectedScreenPosition = float4(ScreenPosition.xy / ScreenPosition.w * SceneW, SceneW, 1);		// is ProjectedClipPosition
-
-		ProjectedScreenPosition = float4(ScreenPosition.xy / ScreenPosition.w * SceneW, SceneW, 1);
-
-		float4 WorldPosHom = mul(ProjectedScreenPosition, View.ScreenToWorld);
+		float4 SvPositionInDistance = float4(SvPosition.xy, ConvertToDeviceZ(SceneW), 1);
+//		float4 ProjectedScreenPosition = float4(ScreenPosition.xy / ScreenPosition.w * SceneW, SceneW, 1);		// is ProjectedClipPosition
+//		ProjectedScreenPosition = float4(ScreenPosition.xy / ScreenPosition.w * SceneW, SceneW, 1);
+//		float4 WorldPosHom = mul(ProjectedScreenPosition, View.ScreenToWorld);
 
 		float3 WSStartRay = View.WorldCameraOrigin;
-		float3 WSEndRay = WorldPosHom.xyz;
+//		float3 WSEndRay = WorldPosHom.xyz;
+		float3 WSEndRay = SvPositionToWorld(SvPositionInDistance);
 
 		// in object/decal space
 		float4 OSStartRayHom = mul(float4(WSStartRay,1), WorldToDecal);
@@ -222,7 +245,9 @@ void MainPS(
 	}
 	else
 	{
-		half SceneW = PreviousDepth(ScreenPosition);
+		// can be optimized further
+
+		half SceneW = CalcSceneDepth(SvPositionToBufferUV(SvPosition));
 		float4 ProjectedScreenPosition = float4(ScreenPosition.xy / ScreenPosition.w * SceneW, SceneW, 1);	
 
 		float4 DecalVectorHom = mul(ProjectedScreenPosition, ScreenToDecal);
@@ -233,8 +258,46 @@ void MainPS(
 		clip(OSPosition.xyz + 1.0f);
 		clip(1.0f - OSPosition.xyz);
 
-		float4 WSPositionHom = mul(ProjectedScreenPosition, View.ScreenToWorld);
-		WSPosition = WSPositionHom.xyz / WSPositionHom.w;
+
+
+		// can be optimized
+		
+//		SvPosition.z = CalcSceneDepth(ScreenUV);
+//		SvPosition.w *= SvPosition.z;
+//		SvPosition.z = LookupDeviceZ(ScreenUV);
+//		SvPosition.xy /= SvPosition.w;
+	//	SvPosition.w = 1;
+
+		float4 Clip;
+
+	//	ProjectedScreenPosition = ComputeScreenPosition(SvPosition);
+		Clip.x = ((SvPosition.x-View.ViewRectMin.x) * View.ViewSizeAndInvSize.z) * 2 - 1;
+//		ProjectedScreenPosition.x = SvPosition.x * View.BufferSizeAndInvSize.z;
+		Clip.y = ((SvPosition.y-View.ViewRectMin.y) * View.ViewSizeAndInvSize.w) * -2 + 1;
+//		Clip.z = LookupDeviceZ(ScreenUV);
+
+// test
+		Clip.z = SvPosition.z;
+
+
+//		ProjectedScreenPosition.z = CalcSceneDepth(ScreenUV);
+//		ProjectedScreenPosition.xy*=ProjectedScreenPosition.z;
+//		ProjectedScreenPosition.w = 1;
+
+		Clip.w = 1;
+
+		SvPosition.z = LookupDeviceZ(ScreenUV);
+
+
+//		float4 Hom = mul(Clip, View.ClipToTranslatedWorld);WSPosition = Hom.xyz / Hom.w - View.PreViewTranslation;		// float4(-1..1,1..-1,LookupDeviceZ(ScreenUV), 1)
+//		float4 Hom = mul(ProjectedScreenPosition, View.ScreenToWorld);WSPosition = Hom.xyz / Hom.w;				// float4(-1..1*z,1..-1*z, CalcSceneDepth(ScreenUV),1)
+		float4 Hom = mul(float4(SvPosition.xyz - float3(View.ViewSizeAndInvSize.xy * 0.5f, 0),1), View.SVPositionToTranslatedWorld);WSPosition = Hom.xyz / Hom.w - View.PreViewTranslation;				// float4(PixelX,PixelY,DeviceZ,W)
+//		float4 Hom = mul(float4(SvPosition.xyz, 1), View.SVPositionToTranslatedWorld);WSPosition = Hom.xyz / Hom.w - View.PreViewTranslation;			// float4(PixelX,PixelY,DeviceZ,W)  perfect
+//		float4 Hom = mul(SvPositionToScreenPosition2(SvPosition), View.ScreenToWorld);WSPosition = Hom.xyz / Hom.w;	// works but why
+//		float4 Hom = mul(SvPositionToScreenPosition(SvPosition), View.ScreenToWorld);WSPosition = Hom.xyz / Hom.w;	// to verify SvPositionToScreenPosition, wrong?
+
+//		WSPosition = Hom.xyz / Hom.w - View.PreViewTranslation;
+//		WSPosition  = SvPosition.xyz / SvPosition.w;
 	}
 
 	float3 CameraVector = normalize(View.WorldCameraOrigin - WSPosition);
@@ -257,6 +320,7 @@ void MainPS(
 #endif
 	MaterialParameters.VertexColor = 1;
 	MaterialParameters.CameraVector = CameraVector;
+	MaterialParameters.SvPosition = SvPosition;
 	MaterialParameters.ScreenPosition = ScreenPosition;
 	MaterialParameters.LightVector = SwizzlePos;
 	MaterialParameters.WorldPosition = 
diff --git a/Engine/Shaders/LightFunctionPixelShader.usf b/Engine/Shaders/LightFunctionPixelShader.usf
--- a/Engine/Shaders/LightFunctionPixelShader.usf
+++ b/Engine/Shaders/LightFunctionPixelShader.usf
@@ -15,10 +15,11 @@ float3 LightFunctionParameters2;
 
 void Main(
 	in float4 ScreenPosition : TEXCOORD0,
+	in float4 SvPosition : SV_Position,		// after all interpolators
 	out float4 OutColor : SV_Target0
 	)
 {
-	half SceneW = PreviousDepth(ScreenPosition);
+	half SceneW = CalcSceneDepth(SvPositionToBufferUV(SvPosition));
 	float4 ProjectedScreenPosition = float4(ScreenPosition.xy / ScreenPosition.w * SceneW, SceneW, 1);
 	float4 LightVector = mul(ProjectedScreenPosition, ScreenToLight);
 	LightVector.xyz /= LightVector.w;
diff --git a/Engine/Shaders/MaterialTemplate.usf b/Engine/Shaders/MaterialTemplate.usf
--- a/Engine/Shaders/MaterialTemplate.usf
+++ b/Engine/Shaders/MaterialTemplate.usf
@@ -366,7 +366,7 @@ float4 GetScreenPosition(FMaterialPixelParameters Parameters)
 	 */
 	float2 ComputeDecalUVFromScreenPosition(float4 ScreenPosition)
 	{
-		half SceneW = PreviousDepth(ScreenPosition);
+		half SceneW = CalcSceneDepth(ScreenPositionToBufferUV(ScreenPosition));
 		float4 DecalVector = mul(float4(ScreenPosition.xy / ScreenPosition.w * SceneW, SceneW, 1), ScreenToDecal);
 		DecalVector.xyz /= DecalVector.w;
 		DecalVector = DecalVector * 0.5f + 0.5f;
@@ -817,7 +817,7 @@ float GetSphericalParticleOpacity(FMaterialPixelParameters Parameters, float Den
 		float NearDistance = Parameters.ScreenPosition.w - HemisphericalDistance;
 		float FarDistance = Parameters.ScreenPosition.w + HemisphericalDistance;
 
-		float SceneDepth = PreviousDepth(Parameters.ScreenPosition);
+		float SceneDepth = CalcSceneDepth(ScreenPositionToBufferUV(Parameters.ScreenPosition));
 		FarDistance = min(SceneDepth, FarDistance);
 
 		// Take into account opaque objects intersecting the sphere
@@ -1660,20 +1660,6 @@ void CalcMaterialParametersEx(
 #endif // !PARTICLE_SPRITE_FACTORY
 }
 
-float4 ComputeScreenPosition(float4 SvPosition)
-{
-	float4 ScreenPosition;
-
-	float2 PixelPos = SvPosition.xy - View.ViewRectMin.xy;	
-	ScreenPosition.xy = ( PixelPos * View.ViewSizeAndInvSize.zw - 0.5 ) * float2(2,-2);
-	// clipspace
-	ScreenPosition.zw = SvPosition.zw;
-	// viewspace
-	ScreenPosition.xyz *= ScreenPosition.w;
-
-	return ScreenPosition;
-}
-
 // this will be optimized to not use ScreenPosition
 float3 ComputeWorldWorldPosCamRelative(float4 ScreenPosition)
 {
@@ -1689,7 +1675,7 @@ void CalcMaterialParameters(
 	float4 SvPosition,
 	FIsFrontFace bIsFrontFace)
 {
-	float4 ScreenPosition = ComputeScreenPosition(SvPosition);
+	float4 ScreenPosition = SvPositionToScreenPosition(SvPosition);
 	float3 PixelPosition = ComputeWorldWorldPosCamRelative(ScreenPosition);
 
 	CalcMaterialParametersEx(Parameters, SvPosition, ScreenPosition, bIsFrontFace, PixelPosition, PixelPosition);
diff --git a/Engine/Shaders/ParticleGPUSpriteVertexFactory.usf b/Engine/Shaders/ParticleGPUSpriteVertexFactory.usf
--- a/Engine/Shaders/ParticleGPUSpriteVertexFactory.usf
+++ b/Engine/Shaders/ParticleGPUSpriteVertexFactory.usf
@@ -376,7 +376,7 @@ FMaterialPixelParameters GetMaterialPixelParameters( FVertexFactoryInterpolantsV
 #if USE_SPHERICAL_NORMALS
 	// can be optimized
 	{
-		float4 ScreenPosition = ComputeScreenPosition(SvPosition);
+		float4 ScreenPosition = SvPositionToScreenPosition(SvPosition);
 		float3 PixelPosition = ComputeWorldWorldPosCamRelative(ScreenPosition);
 
 		Result.TangentToWorld = GetSphericalParticleNormal(PixelPosition + View.WorldCameraOrigin, Interpolants.ParticlePositionAndSize.xyz, Interpolants.ParticlePositionAndSize.w);
diff --git a/Engine/Shaders/ParticleSpriteVertexFactory.usf b/Engine/Shaders/ParticleSpriteVertexFactory.usf
--- a/Engine/Shaders/ParticleSpriteVertexFactory.usf
+++ b/Engine/Shaders/ParticleSpriteVertexFactory.usf
@@ -172,7 +172,7 @@ FMaterialPixelParameters GetMaterialPixelParameters(FVertexFactoryInterpolantsVS
 	#if GENERATE_SPHERICAL_PARTICLE_NORMALS && USE_PARTICLE_POSITION
 		{	
 			// can be optimized
-			float4 ScreenPosition = ComputeScreenPosition(SvPosition);
+			float4 ScreenPosition = SvPositionToScreenPosition(SvPosition);
 			float3 PixelPosition = ComputeWorldWorldPosCamRelative(ScreenPosition);
 			Result.TangentToWorld = GetSphericalParticleNormal(PixelPosition + View.WorldCameraOrigin, Interpolants.ParticlePositionAndSize.xyz, Interpolants.ParticlePositionAndSize.w);
 		}
diff --git a/Engine/Shaders/PostProcessMaterialShaders.usf b/Engine/Shaders/PostProcessMaterialShaders.usf
--- a/Engine/Shaders/PostProcessMaterialShaders.usf
+++ b/Engine/Shaders/PostProcessMaterialShaders.usf
@@ -9,24 +9,21 @@
 
 void MainVS(
 	in float4 InPosition : ATTRIBUTE0,
-	out float4 OutScreenPosition : TEXCOORD0,
 	out float4 OutPosition : SV_POSITION
 	)
 {
 	DrawRectangle(InPosition,  OutPosition);
-	OutScreenPosition = OutPosition;
 }
 
 void MainPS(
-	in float4 ScreenPosition : TEXCOORD0,
-	in float4 SvPosition : SV_Position,
+	in float4 SvPosition : SV_Position,		// after all interpolators
 	out float4 OutColor : SV_Target0
 	)
 {
 	FMaterialPixelParameters Parameters = MakeInitializedMaterialPixelParameters();
 
 	// can be optimized
-	float2 ScreenUV = ScreenAlignedPosition(ScreenPosition).xy;
+	float2 ScreenUV = SvPositionToBufferUV(SvPosition);
 
 #if NUM_MATERIAL_TEXCOORDS
 	for(int CoordinateIndex = 0;CoordinateIndex < NUM_MATERIAL_TEXCOORDS;CoordinateIndex++)
@@ -38,8 +35,6 @@ void MainPS(
 	Parameters.VertexColor = 1;
 	
 	SvPosition.z = LookupDeviceZ(ScreenUV);
-	// according to MaterialTemplate.usf, not sure about the .w
-	SvPosition.w = ConvertFromDeviceZ(SvPosition.z);
 
 	// fill out other related material parameters
 	CalcMaterialParameters(Parameters, SvPosition, true);
diff --git a/Engine/Shaders/PostProcessMotionBlur.usf b/Engine/Shaders/PostProcessMotionBlur.usf
--- a/Engine/Shaders/PostProcessMotionBlur.usf
+++ b/Engine/Shaders/PostProcessMotionBlur.usf
@@ -834,7 +834,7 @@ float GetVelocityLengthPixels( float2 EncodedVelocity )
 
 void MainNewPS(
 	in float4 UVAndScreenPos : TEXCOORD0,
-	in float4 SvPosition : SV_Position,
+	in float4 SvPosition : SV_Position,		// after all interpolators
 	out float4 OutColor : SV_Target0
 	)
 {
diff --git a/Engine/Shaders/PostProcessTemporalAA.usf b/Engine/Shaders/PostProcessTemporalAA.usf
--- a/Engine/Shaders/PostProcessTemporalAA.usf
+++ b/Engine/Shaders/PostProcessTemporalAA.usf
@@ -302,7 +302,7 @@ void MainTemporalAAPS(
 	// Note: these are needed to match up with the vertex shader outputs which is FPostProcessTonemapVS, otherwise a d3d debug error is issued
 	in noperspective float4 Unused1 : TEXCOORD2,
 	in noperspective float4 Unused2 : TEXCOORD3,
-	in float4 SvPosition : SV_Position,
+	in float4 SvPosition : SV_Position,		// after all interpolators
 	out float4 OutColor : SV_Target0 )
 {
 #if ENABLE_TEMPORAL_AA
@@ -331,7 +331,7 @@ void MainTemporalAAPS(
 void MainFastTemporalAAPS(
 	in noperspective float4 UVAndScreenPos : TEXCOORD0,
 	in noperspective float3 InExposureScaleVignette : TEXCOORD1,
-	in float4 SvPosition : SV_Position,
+	in float4 SvPosition : SV_Position,		// after all interpolators
 	out float4 OutColor : SV_Target0 )
 {
 #if ENABLE_TEMPORAL_AA
diff --git a/Engine/Shaders/PostProcessTonemap.usf b/Engine/Shaders/PostProcessTonemap.usf
--- a/Engine/Shaders/PostProcessTonemap.usf
+++ b/Engine/Shaders/PostProcessTonemap.usf
@@ -348,7 +348,7 @@ void MainPS_ES2(
 	in float2 InVignette : TEXCOORD3,
 	in float2 InTexCoords[4] : TEXCOORD4,
 	#if USE_32BPP_HDR
-		in float4 SvPosition : SV_Position,
+		in float4 SvPosition : SV_Position,		// after all interpolators
 	#endif
 	out half4 OutColor : SV_Target0
 	)
diff --git a/Engine/Shaders/ShadowDepthPixelShader.usf b/Engine/Shaders/ShadowDepthPixelShader.usf
--- a/Engine/Shaders/ShadowDepthPixelShader.usf
+++ b/Engine/Shaders/ShadowDepthPixelShader.usf
@@ -44,7 +44,7 @@ float3 GetMaterialBounceColor(FMaterialPixelParameters MaterialParameters)
 
 void Main( 
 	FShadowDepthVSToPS Inputs,
-	in float4 SvPosition : SV_Position,
+	in float4 SvPosition : SV_Position,		// after all interpolators
 	out float4 OutColor : SV_Target0
 #if REFLECTIVE_SHADOW_MAP
 	,out float4 OutReflectiveShadowMapDiffuse : SV_Target1
@@ -60,7 +60,7 @@ void Main(
 
 	#if INTERPOLATE_POSITION
 		{
-			float4 ScreenPosition = ComputeScreenPosition(SvPosition);
+			float4 ScreenPosition = SvPositionToScreenPosition(SvPosition);
 			CalcMaterialParametersEx(MaterialParameters, SvPosition, ScreenPosition, 1, ComputeWorldWorldPosCamRelative(ScreenPosition), SvPosition);
 		}
 	#else
diff --git a/Engine/Shaders/StationaryLightOverlapShaders.usf b/Engine/Shaders/StationaryLightOverlapShaders.usf
--- a/Engine/Shaders/StationaryLightOverlapShaders.usf
+++ b/Engine/Shaders/StationaryLightOverlapShaders.usf
@@ -36,20 +36,20 @@ void OverlapDirectionalPixelMain(
 
 void OverlapRadialPixelMain(
 	float4 InScreenPosition : TEXCOORD0,
+	in float4 SvPosition : SV_Position,		// after all interpolators
 	out float4 OutColor : SV_Target0
 	)
 {
 	OutColor = 0;
-	float2 ScreenUV = InScreenPosition.xy / InScreenPosition.w * View.ScreenPositionScaleBias.xy + View.ScreenPositionScaleBias.wz;
+	float2 ScreenUV = SvPositionToBufferUV(SvPosition);
+	SvPosition.z = LookupDeviceZ(ScreenUV);
 	
-	FGBufferData GBuffer = GetGBufferData(ScreenUV);
+	FGBufferData GBuffer = GetGBufferData(ScreenUV); 
 	
 	// Only light pixels marked as using deferred shading
 	BRANCH if( GBuffer.ShadingModelID > 0 )
 	{
-		float SceneDepth = CalcSceneDepth(ScreenUV);
-		float4 HomogeneousWorldPosition = mul(float4(InScreenPosition.xy / InScreenPosition.w * SceneDepth, SceneDepth, 1), View.ScreenToWorld);
-		float3 WorldPosition = HomogeneousWorldPosition.xyz / HomogeneousWorldPosition.w;
+		float3 WorldPosition = SvPositionToWorld(SvPosition);
 
 		float3 ToLight = DeferredLightUniforms.NormalizedLightDirection;
 
diff --git a/Engine/Shaders/TranslucentShadowDepthShaders.usf b/Engine/Shaders/TranslucentShadowDepthShaders.usf
--- a/Engine/Shaders/TranslucentShadowDepthShaders.usf
+++ b/Engine/Shaders/TranslucentShadowDepthShaders.usf
@@ -74,7 +74,7 @@ float TranslInvMaxSubjectDepth;
 /** Pixel shader used to accumulate layer opacities in different channels based on the first translucent layer's depth. */
 void MainOpacityPS( 
 	FTranslucencyShadowDepthVSToPS Inputs,
-	in float4 SvPosition : SV_Position,
+	in float4 SvPosition : SV_Position,		// after all interpolators
 	out float4 OutColor0 : SV_Target0,
 	out float4 OutColor1 : SV_Target1
 	)
diff --git a/Engine/Shaders/VelocityCommon.usf b/Engine/Shaders/VelocityCommon.usf
--- a/Engine/Shaders/VelocityCommon.usf
+++ b/Engine/Shaders/VelocityCommon.usf
@@ -8,11 +8,11 @@
 
 float2 Calculate2DVelocity(float4 PackedVelocityA, float4 PackedVelocityC)
 {
-	float2 ScreenPosition = PackedVelocityA.xy / PackedVelocityA.w - View.ViewToClip[2].xy;
-	float2 PrevScreenPosition = PackedVelocityC.xy / PackedVelocityC.w - View.PrevProjection[2].xy;
+	float2 ScreenPos = PackedVelocityA.xy / PackedVelocityA.w - View.ViewToClip[2].xy;
+	float2 PrevScreenPos = PackedVelocityC.xy / PackedVelocityC.w - View.PrevProjection[2].xy;
 
 	// 2d velocity, includes camera an object motion
-	float2 Velocity = ScreenPosition - PrevScreenPosition;
+	float2 Velocity = ScreenPos - PrevScreenPos;
 
 	// Make sure not to touch 0,0 which is clear color
 	return Velocity;
diff --git a/Engine/Shaders/VelocityShader.usf b/Engine/Shaders/VelocityShader.usf
--- a/Engine/Shaders/VelocityShader.usf
+++ b/Engine/Shaders/VelocityShader.usf
@@ -26,9 +26,9 @@ static const float GNonTessellationDepthBias = 0.0f;
 
 struct FVelocityInterpsVSToPS
 {
-	// float4(ScreenPosition.xyzw)
+	// float4(ScreenPos.xyzw)
 	float4 PackedVelocityA : TEXCOORD6;
-	// float4(PrevScreenPosition.xyzw)
+	// float4(PrevScreenPos.xyzw)
 	float4 PackedVelocityC : TEXCOORD7;
 	FVertexFactoryInterpolantsVSToPS FactoryInterpolants;
 };
@@ -43,9 +43,9 @@ struct FVelocityVSToPS
 
 	struct FVelocityInterpsVSToDS
 	{
-		// float4(ScreenPosition.xyzw)
+		// float4(ScreenPos.xyzw)
 		float4 PackedVelocityA : TEXCOORD6;
-		// float4(PrevScreenPosition.xyzw)
+		// float4(PrevScreenPos.xyzw)
 		float4 PackedVelocityC : TEXCOORD8;
 	};
 
@@ -151,22 +151,22 @@ void MainVertexShader(
 #endif
 
 	{
-		float4 ScreenPosition;
+		float4 ScreenPos;
 		// compute the old screen pos with the old world position and the old camera matrix
-		float4 PrevScreenPositionObj = mul(float4(PrevTranslatedWorldPosition.xyz, 1), View.PrevTranslatedWorldToClip);
+		float4 PrevScreenPosObj = mul(float4(PrevTranslatedWorldPosition.xyz, 1), View.PrevTranslatedWorldToClip);
 
 #if !USING_TESSELLATION
 		ISOLATE
 		{
 #endif
-			ScreenPosition  = mul(float4(TranslatedWorldPosition.xyz, 1), View.TranslatedWorldToClip);
+			ScreenPos  = mul(float4(TranslatedWorldPosition.xyz, 1), View.TranslatedWorldToClip);
 #if !USING_TESSELLATION
-			Output.Position = ScreenPosition;
+			Output.Position = ScreenPos;
 		}
 #endif
 
-		Output.Interps.PackedVelocityA = ScreenPosition;
-		Output.Interps.PackedVelocityC = PrevScreenPositionObj;
+		Output.Interps.PackedVelocityA = ScreenPos;
+		Output.Interps.PackedVelocityC = PrevScreenPosObj;
 	}
 
 #if !USING_TESSELLATION
diff --git a/Engine/Source/Runtime/Engine/Private/Materials/HLSLMaterialTranslator.h b/Engine/Source/Runtime/Engine/Private/Materials/HLSLMaterialTranslator.h
--- a/Engine/Source/Runtime/Engine/Private/Materials/HLSLMaterialTranslator.h
+++ b/Engine/Source/Runtime/Engine/Private/Materials/HLSLMaterialTranslator.h
@@ -1949,7 +1949,7 @@ protected:
 			return Errorf(TEXT("Invalid node used in hull/domain shader input!"));
 		}
 
-		return AddCodeChunk(MCT_Float2,TEXT("ScreenAlignedPosition(GetScreenPosition(Parameters)).xy"));		
+		return AddCodeChunk(MCT_Float2,TEXT("ScreenAlignedPosition(GetScreenPosition(Parameters))"));		
 	}
 
 	virtual int32 ParticleMacroUV() override 
@@ -2525,7 +2525,7 @@ protected:
 		}
 		else
 		{
-			FString DefaultScreenAligned(TEXT("MaterialFloat2(ScreenAlignedPosition(GetScreenPosition(Parameters)).xy)"));
+			FString DefaultScreenAligned(TEXT("ScreenAlignedPosition(GetScreenPosition(Parameters))"));
 			FString CodeString = (UV != INDEX_NONE) ? CoerceParameter(UV,MCT_Float2) : DefaultScreenAligned;
 			return AddInlinedCodeChunk(MCT_Float2, *CodeString );
 		}
@@ -2573,7 +2573,7 @@ protected:
 
 		UseSceneTextureId(SceneTextureId, true);
 
-		FString DefaultScreenAligned(TEXT("MaterialFloat2(ScreenAlignedPosition(GetScreenPosition(Parameters)).xy)"));
+		FString DefaultScreenAligned(TEXT("ScreenAlignedPosition(GetScreenPosition(Parameters))"));
 		FString TexCoordCode((UV != INDEX_NONE) ? CoerceParameter(UV, MCT_Float2) : DefaultScreenAligned);
 
 		return AddCodeChunk(
diff --git a/Engine/Source/Runtime/Engine/Public/SceneView.h b/Engine/Source/Runtime/Engine/Public/SceneView.h
--- a/Engine/Source/Runtime/Engine/Public/SceneView.h
+++ b/Engine/Source/Runtime/Engine/Public/SceneView.h
@@ -362,6 +362,7 @@ BEGIN_UNIFORM_BUFFER_STRUCT_WITH_CONSTRUCTOR(FViewUniformShaderParameters,ENGINE
 	DECLARE_UNIFORM_BUFFER_STRUCT_MEMBER(FMatrix,ViewToClip)
 	DECLARE_UNIFORM_BUFFER_STRUCT_MEMBER(FMatrix,ClipToView)
 	DECLARE_UNIFORM_BUFFER_STRUCT_MEMBER(FMatrix,ClipToTranslatedWorld)
+	DECLARE_UNIFORM_BUFFER_STRUCT_MEMBER(FMatrix,SVPositionToTranslatedWorld)	// assumes input float4(SvPosition.xyz,1)
 	DECLARE_UNIFORM_BUFFER_STRUCT_MEMBER(FMatrix,ScreenToWorld)
 	DECLARE_UNIFORM_BUFFER_STRUCT_MEMBER(FMatrix,ScreenToTranslatedWorld)
 	DECLARE_UNIFORM_BUFFER_STRUCT_MEMBER_EX(FVector,ViewForward, EShaderPrecisionModifier::Half)
diff --git a/Engine/Source/Runtime/Renderer/Private/SceneRendering.cpp b/Engine/Source/Runtime/Renderer/Private/SceneRendering.cpp
--- a/Engine/Source/Runtime/Renderer/Private/SceneRendering.cpp
+++ b/Engine/Source/Runtime/Renderer/Private/SceneRendering.cpp
@@ -430,6 +430,7 @@ TUniformBufferRef<FViewUniformShaderParameters> FViewInfo::CreateUniformBuffer(
 	const FIntPoint BufferSize = SceneContext.GetBufferSizeXY();
 	const float InvBufferSizeX = 1.0f / BufferSize.X;
 	const float InvBufferSizeY = 1.0f / BufferSize.Y;
+	// to bring NDC (-1..1, 1..-1) into 0..1 UV for BufferSize textures
 	const FVector4 ScreenPositionScaleBias(
 		ViewRect.Width() * InvBufferSizeX / +2.0f,
 		ViewRect.Height() * InvBufferSizeY / (-2.0f * GProjectionSignY),
@@ -524,6 +525,27 @@ TUniformBufferRef<FViewUniformShaderParameters> FViewInfo::CreateUniformBuffer(
 	// can be optimized
 	ViewUniformShaderParameters.PrevInvViewProj = PrevViewProjMatrix.Inverse();
 
+	{
+		// setup a matrix to transform float4(SVPosition.xyz,1) directly to TranslatedWorld (quality, performance as we don't need to convert or use interpolator)
+
+		//	new_xy = (xy - ViewRectMin.xy) * ViewSizeAndInvSize.zw * float2(2,-2) + float2(-1, 1);
+
+		//  transformed into one MAD:  new_xy = xy * ViewSizeAndInvSize.zw * float2(2,-2)      +       (-ViewRectMin.xy) * ViewSizeAndInvSize.zw * float2(2,-2) + float2(-1, 1);
+
+		float Mx = 2.0f * ViewUniformShaderParameters.ViewSizeAndInvSize.Z;
+		float My = -2.0f * ViewUniformShaderParameters.ViewSizeAndInvSize.W;
+		float Ax = -1.0f - 2.0f * ViewRect.Min.X * ViewUniformShaderParameters.ViewSizeAndInvSize.Z;
+		float Ay = 1.0f + 2.0f * ViewRect.Min.Y * ViewUniformShaderParameters.ViewSizeAndInvSize.W;
+
+		// http://stackoverflow.com/questions/9010546/java-transformation-matrix-operations
+
+		ViewUniformShaderParameters.SVPositionToTranslatedWorld = 
+			FMatrix(FPlane(Mx,  0,  0,   0),
+					FPlane(0,  My,  0,   0),
+					FPlane(0,   0,  1,   0),
+					FPlane(Ax, Ay,  0,   1)) * ViewMatrices.InvTranslatedViewProjectionMatrix;
+	}
+
 	ViewUniformShaderParameters.ScreenToWorld = FMatrix(
 		FPlane(1,0,0,0),
 		FPlane(0,1,0,0),
