integrated from orion
fixed decal projection of recent checkin

diff --git a/Engine/Shaders/DeferredDecal.usf b/Engine/Shaders/DeferredDecal.usf
--- a/Engine/Shaders/DeferredDecal.usf
+++ b/Engine/Shaders/DeferredDecal.usf
@@ -266,8 +266,7 @@ void MainPS(
 //		SvPosition.w *= SvPosition.z;
 //		SvPosition.z = LookupDeviceZ(ScreenUV);
 //		SvPosition.xy /= SvPosition.w;
-	//	SvPosition.w = 1;
-
+	
 		float4 Clip;
 
 	//	ProjectedScreenPosition = ComputeScreenPosition(SvPosition);
@@ -291,8 +290,8 @@ void MainPS(
 
 //		float4 Hom = mul(Clip, View.ClipToTranslatedWorld);WSPosition = Hom.xyz / Hom.w - View.PreViewTranslation;		// float4(-1..1,1..-1,LookupDeviceZ(ScreenUV), 1)
 //		float4 Hom = mul(ProjectedScreenPosition, View.ScreenToWorld);WSPosition = Hom.xyz / Hom.w;				// float4(-1..1*z,1..-1*z, CalcSceneDepth(ScreenUV),1)
-		float4 Hom = mul(float4(SvPosition.xyz - float3(View.ViewSizeAndInvSize.xy * 0.5f, 0),1), View.SVPositionToTranslatedWorld);WSPosition = Hom.xyz / Hom.w - View.PreViewTranslation;				// float4(PixelX,PixelY,DeviceZ,W)
-//		float4 Hom = mul(float4(SvPosition.xyz, 1), View.SVPositionToTranslatedWorld);WSPosition = Hom.xyz / Hom.w - View.PreViewTranslation;			// float4(PixelX,PixelY,DeviceZ,W)  perfect
+//		float4 Hom = mul(float4(SvPosition.xyz - float3(View.ViewSizeAndInvSize.xy * 0.5f, 0),1), View.SVPositionToTranslatedWorld);WSPosition = Hom.xyz / Hom.w - View.PreViewTranslation;				// float4(PixelX,PixelY,DeviceZ,W)
+		float4 Hom = mul(float4(SvPosition.xyz, 1), View.SVPositionToTranslatedWorld);WSPosition = Hom.xyz / Hom.w - View.PreViewTranslation;			// float4(PixelX,PixelY,DeviceZ,W)  perfect
 //		float4 Hom = mul(SvPositionToScreenPosition2(SvPosition), View.ScreenToWorld);WSPosition = Hom.xyz / Hom.w;	// works but why
 //		float4 Hom = mul(SvPositionToScreenPosition(SvPosition), View.ScreenToWorld);WSPosition = Hom.xyz / Hom.w;	// to verify SvPositionToScreenPosition, wrong?
 
