Fix cook errors from trying to modify a uniform parameter.

diff --git a/Engine/Shaders/ParticleSimulationShader.usf b/Engine/Shaders/ParticleSimulationShader.usf
--- a/Engine/Shaders/ParticleSimulationShader.usf
+++ b/Engine/Shaders/ParticleSimulationShader.usf
@@ -455,7 +455,7 @@ void PixelMain(
 	RelativeTime += DeltaSeconds * TimeScale;
 
 	//Everything other that particle aging should stop if we have collision stop.
-	DeltaSeconds = CollisionStop > 0.0f ? 0.0f : DeltaSeconds;
+	float CollisionAdjustedDeltaSeconds = CollisionStop > 0.0f ? 0.0f : DeltaSeconds;
 
 	// Sample the attribute curve.
 	const float2 AttributeCurveTexCoord = Simulation.AttributeCurve.xy +
@@ -489,7 +489,7 @@ void PixelMain(
 	Force += ComputeAttractionForce(Position.xyz);
 
 	// Compute the acceleration to apply to the particle this frame.
-	float3 Acceleration = Force * DeltaSeconds;
+	float3 Acceleration = Force * CollisionAdjustedDeltaSeconds;
 
 #if DEPTH_BUFFER_COLLISION || DISTANCE_FIELD_COLLISION
 	// We need to look up render attributes for this particle to figure out how big it is.
