integrated from orion CL 2699770

renamed PS WorldPosition to AbsoluteWorldPosition as it's not translated
VS/DS/HS WorldPosition is actually TranslatedWorldPosition - we should change that as well
later we get rid of the AbsoluteWorldPosition where we can (for better precision)

diff --git a/Engine/Shaders/BasePassForForwardShadingPixelShader.usf b/Engine/Shaders/BasePassForForwardShadingPixelShader.usf
--- a/Engine/Shaders/BasePassForForwardShadingPixelShader.usf
+++ b/Engine/Shaders/BasePassForForwardShadingPixelShader.usf
@@ -310,7 +310,7 @@ void Main(
 			for (int i = 0; i < NUM_DYNAMIC_POINT_LIGHTS; i++)
 			{
 			#endif
-				float3 ToLight = LightPositionAndInvRadius[i].xyz - MaterialParameters.WorldPosition;
+				float3 ToLight = LightPositionAndInvRadius[i].xyz - MaterialParameters.AbsoluteWorldPosition;
 				float DistanceSqr = dot(ToLight, ToLight);
 				float3 L = ToLight * rsqrt(DistanceSqr);
 
diff --git a/Engine/Shaders/BasePassPixelShader.usf b/Engine/Shaders/BasePassPixelShader.usf
--- a/Engine/Shaders/BasePassPixelShader.usf
+++ b/Engine/Shaders/BasePassPixelShader.usf
@@ -170,8 +170,8 @@ float3 GetTranslucencyLighting(FMaterialPixelParameters MaterialParameters, floa
 	float3 InterpolatedLighting = 0;
 
 	// Apply a stable offset to the world position used for lighting, which breaks up artifacts from using a low res volume texture
-	float3 InnerVolumeUVs = (MaterialParameters.WorldPosition + MaterialParameters.LightingPositionOffset - View.TranslucencyLightingVolumeMin[0].xyz) * View.TranslucencyLightingVolumeInvSize[0].xyz;
-	float3 OuterVolumeUVs = (MaterialParameters.WorldPosition + MaterialParameters.LightingPositionOffset - View.TranslucencyLightingVolumeMin[1].xyz) * View.TranslucencyLightingVolumeInvSize[1].xyz;
+	float3 InnerVolumeUVs = (MaterialParameters.AbsoluteWorldPosition + MaterialParameters.LightingPositionOffset - View.TranslucencyLightingVolumeMin[0].xyz) * View.TranslucencyLightingVolumeInvSize[0].xyz;
+	float3 OuterVolumeUVs = (MaterialParameters.AbsoluteWorldPosition + MaterialParameters.LightingPositionOffset - View.TranslucencyLightingVolumeMin[1].xyz) * View.TranslucencyLightingVolumeInvSize[1].xyz;
 
 	// Controls how fast the lerp between the inner and outer cascades happens
 	// Larger values result in a shorter transition distance
@@ -259,7 +259,7 @@ float3 GetTranslucencyLighting(FMaterialPixelParameters MaterialParameters, floa
 
 			// Determine the shadow space position
 			// Apply a stable offset to the world position used for shadowing, which blurs out high frequency details in the shadowmap with many layers
-			float4 HomogeneousShadowPosition = mul(float4(MaterialParameters.WorldPosition + MaterialParameters.LightingPositionOffset, 1), WorldToShadowMatrix);
+			float4 HomogeneousShadowPosition = mul(float4(MaterialParameters.AbsoluteWorldPosition + MaterialParameters.LightingPositionOffset, 1), WorldToShadowMatrix);
 			float2 ShadowUVs = HomogeneousShadowPosition.xy / HomogeneousShadowPosition.w;
 			// Lookup the shadow density at the point being shaded
 			float3 ShadowDensity = CalculateTranslucencyShadowingDensity(ShadowUVs, HomogeneousShadowPosition.z) / GetMaterialTranslucentMultipleScatteringExtinction();
@@ -419,7 +419,7 @@ void GetPrecomputedIndirectLightingAndSkyLight(
 	#if CACHED_VOLUME_INDIRECT_LIGHTING
 	
 		// Compute volume teture UVs from world position
-		float3 VolumeUVs = MaterialParameters.WorldPosition * IndirectlightingCachePrimitiveScale + IndirectlightingCachePrimitiveAdd;
+		float3 VolumeUVs = MaterialParameters.AbsoluteWorldPosition * IndirectlightingCachePrimitiveScale + IndirectlightingCachePrimitiveAdd;
 		// Clamp UV to be within the valid region
 		// Pixels outside of the object's bounding box would read garbage otherwise
 		VolumeUVs = clamp(VolumeUVs, IndirectlightingCacheMinUV, IndirectlightingCacheMaxUV);
@@ -801,9 +801,9 @@ void Main(
 		BRANCH
 		if (View.OutOfBoundsMask > 0)
 		{
-			if (any(abs(MaterialParameters.WorldPosition - Primitive.ObjectWorldPositionAndRadius.xyz) > Primitive.ObjectBounds + 1))
+			if (any(abs(MaterialParameters.AbsoluteWorldPosition - Primitive.ObjectWorldPositionAndRadius.xyz) > Primitive.ObjectBounds + 1))
 			{
-				float Gradient = frac(dot(MaterialParameters.WorldPosition, float3(.577f, .577f, .577f)) / 500.0f);
+				float Gradient = frac(dot(MaterialParameters.AbsoluteWorldPosition, float3(.577f, .577f, .577f)) / 500.0f);
 				Emissive = lerp(float3(1,1,0), float3(0,1,1), Gradient.xxx > .5f);
 				Opacity = 1;
 			}
diff --git a/Engine/Shaders/DeferredDecal.usf b/Engine/Shaders/DeferredDecal.usf
--- a/Engine/Shaders/DeferredDecal.usf
+++ b/Engine/Shaders/DeferredDecal.usf
@@ -74,11 +74,11 @@ float DistanceFunction(float4 ScreenPosition, float3 Pos)
 	MaterialParameters.VertexColor = 1;
 	MaterialParameters.ScreenPosition = ScreenPosition;
 	MaterialParameters.LightVector = SwizzlePos * 0.5f + 0.5f;		// to be compatible with decals but *0.5+0.5 seems wrong
-	MaterialParameters.WorldPosition = 
+	MaterialParameters.AbsoluteWorldPosition = 
 	MaterialParameters.WorldPosition_CamRelative = 
 	MaterialParameters.WorldPosition_NoOffsets = 
 	MaterialParameters.WorldPosition_NoOffsets_CamRelative = mul(DecalToWorld, float4(Pos, 1)).xyz;
-	MaterialParameters.CameraVector = normalize(View.WorldCameraOrigin - MaterialParameters.WorldPosition);
+	MaterialParameters.CameraVector = normalize(View.WorldCameraOrigin - MaterialParameters.AbsoluteWorldPosition);
 
 	// the material needs to output the distance from "Opacity" depending on LightVector (in 0..1 range)
 	return GetMaterialMaskInputRaw(MaterialParameters);
@@ -322,7 +322,7 @@ void MainPS(
 	MaterialParameters.SvPosition = SvPosition;
 	MaterialParameters.ScreenPosition = ScreenPosition;
 	MaterialParameters.LightVector = SwizzlePos;
-	MaterialParameters.WorldPosition = 
+	MaterialParameters.AbsoluteWorldPosition = 
 		MaterialParameters.WorldPosition_CamRelative = 
 		MaterialParameters.WorldPosition_NoOffsets = 
 		MaterialParameters.WorldPosition_NoOffsets_CamRelative = WSPosition;
diff --git a/Engine/Shaders/ForEachLight.usf b/Engine/Shaders/ForEachLight.usf
--- a/Engine/Shaders/ForEachLight.usf
+++ b/Engine/Shaders/ForEachLight.usf
@@ -17,7 +17,7 @@ float3 SimplePointLightDiffuse( FScreenSpaceData ScreenSpaceData )
 
 float3 LightingForOneLight(FMaterialPixelParameters Parameters, uint LightIndex, FScreenSpaceData ScreenSpaceData)
 {
-	float3 CameraVector = normalize(Parameters.WorldPosition - View.WorldCameraOrigin);
+	float3 CameraVector = normalize(Parameters.AbsoluteWorldPosition - View.WorldCameraOrigin);
 	
 	FDeferredLightData LightData = (FDeferredLightData)0;
 	LightData.LightPositionAndInvRadius = ForwardLightData.LightPositionAndInvRadius[LightIndex];
@@ -39,7 +39,7 @@ float3 LightingForOneLight(FMaterialPixelParameters Parameters, uint LightIndex,
 	
 	float2 ScreenUV = Parameters.SvPosition.xy * View.BufferSizeAndInvSize.zw;
 
-	return GetDynamicLighting(Parameters.WorldPosition, CameraVector, ScreenUV, ScreenSpaceData, ScreenSpaceData.GBuffer.ShadingModelID, LightData, float4(1, 1, 1, 1), uint2(0,0)).rgb;
+	return GetDynamicLighting(Parameters.AbsoluteWorldPosition, CameraVector, ScreenUV, ScreenSpaceData, ScreenSpaceData.GBuffer.ShadingModelID, LightData, float4(1, 1, 1, 1), uint2(0,0)).rgb;
 }
 
 // @param WorldNormal is not assumed to be normalized
@@ -85,4 +85,4 @@ float3 AccumulateForwardLights(FMaterialPixelParameters Parameters, float3 Diffu
 	return Ret;
 }
 
-#endif // __FOR_EACH_LIGHT_COMMON__
\ No newline at end of file
+#endif // __FOR_EACH_LIGHT_COMMON__
diff --git a/Engine/Shaders/LightFunctionCommon.usf b/Engine/Shaders/LightFunctionCommon.usf
--- a/Engine/Shaders/LightFunctionCommon.usf
+++ b/Engine/Shaders/LightFunctionCommon.usf
@@ -12,7 +12,7 @@ float4x4 LightFunctionWorldToLight;
  * Calculates the light function color with the given light vector. 
  * LightVector is the vector from the light to the position being shaded, in the light's coordinate space.
  */
-float3 GetLightFunctionColor(float3 LightVector, float3 WorldPosition)
+float3 GetLightFunctionColor(float3 LightVector, float3 AbsoluteWorldPosition)
 {
 	// Swizzle so that LightVector.xy are perpendicular to the light direction and LightVector.z is distance along the light direction
 	LightVector.xyz = LightVector.zyx;
@@ -45,7 +45,7 @@ float3 GetLightFunctionColor(float3 LightVector, float3 WorldPosition)
 	MaterialParameters.CameraVector = LightVector.xyz;
 	MaterialParameters.ReflectionVector = LightVector.xyz;
 	MaterialParameters.LightVector = LightVector.xyz;
-	MaterialParameters.WorldPosition = WorldPosition;
+	MaterialParameters.AbsoluteWorldPosition = AbsoluteWorldPosition;
 
 #if COMPUTESHADER
 	return GetMaterialEmissiveForCS(MaterialParameters);
diff --git a/Engine/Shaders/LightFunctionPixelShader.usf b/Engine/Shaders/LightFunctionPixelShader.usf
--- a/Engine/Shaders/LightFunctionPixelShader.usf
+++ b/Engine/Shaders/LightFunctionPixelShader.usf
@@ -25,11 +25,11 @@ void Main(
 	LightVector.xyz /= LightVector.w;
 
 	float4 HomogeneousWorldPosition = mul(ProjectedScreenPosition, View.ScreenToWorld);
-	float3 WorldPosition = HomogeneousWorldPosition.xyz / HomogeneousWorldPosition.w;
+	float3 AbsoluteWorldPosition = HomogeneousWorldPosition.xyz / HomogeneousWorldPosition.w;
 	// Calculate radial view distance for stable fading
-	float ViewDistance = length(View.WorldCameraOrigin - WorldPosition);
+	float ViewDistance = length(View.WorldCameraOrigin - AbsoluteWorldPosition);
 
-	float3 Color = GetLightFunctionColor(LightVector.xyz, WorldPosition);
+	float3 Color = GetLightFunctionColor(LightVector.xyz, AbsoluteWorldPosition);
 
 	float DistanceFadeAlpha = saturate((LightFunctionParameters2.x - ViewDistance) / (LightFunctionParameters2.x * .2f));
 	// Fade to disabled based on LightFunctionFadeDistance
diff --git a/Engine/Shaders/MaterialTemplate.usf b/Engine/Shaders/MaterialTemplate.usf
--- a/Engine/Shaders/MaterialTemplate.usf
+++ b/Engine/Shaders/MaterialTemplate.usf
@@ -135,9 +135,10 @@ struct FMaterialPixelParameters
 	half3x3 TangentToWorld;
 
 	/** 
-	 * Interpolated worldspace position of this pixel 
+	 * Interpolated worldspace position of this pixel
+	 * todo: Make this TranslatedWorldPosition and also rename the VS/DS/HS WorldPosition to be TranslatedWorldPosition
 	 */
-	float3 WorldPosition;
+	float3 AbsoluteWorldPosition;
 
 	/** 
 	 * Interpolated worldspace position of this pixel, centered around the camera
@@ -200,7 +201,7 @@ FMaterialPixelParameters MakeInitializedMaterialPixelParameters()
 	MPP.ScreenPosition = 0;
 	MPP.UnMirrored = 0;
 	MPP.TwoSidedSign = 0;
-	MPP.WorldPosition = 0;
+	MPP.AbsoluteWorldPosition = 0;
 	MPP.WorldPosition_NoOffsets = 0;
 	MPP.WorldPosition_CamRelative = 0;
 	MPP.WorldPosition_NoOffsets_CamRelative = 0;
@@ -323,7 +324,7 @@ float3 GetTranslatedWorldPosition(FMaterialTessellationParameters Parameters)
 
 float3 GetWorldPosition(FMaterialPixelParameters Parameters)
 {
-	return Parameters.WorldPosition;
+	return Parameters.AbsoluteWorldPosition;
 }
 
 float3 GetWorldPosition_NoMaterialOffsets(FMaterialPixelParameters Parameters)
@@ -642,11 +643,11 @@ float3 MaterialExpressionBlackBody( float Temp )
 	return mul( XYZtoRGB, float3( X, Y, Z ) ) * pow( 0.0004 * Temp, 4 );
 }
 
-float4 MaterialExpressionAtmosphericFog(FMaterialPixelParameters Parameters, float3 WorldPosition)
+float4 MaterialExpressionAtmosphericFog(FMaterialPixelParameters Parameters, float3 AbsoluteWorldPosition)
 {
 #if MATERIAL_ATMOSPHERIC_FOG
-	// WorldPosition default value is Parameters.WorldPosition if not overridden by the user
-	float3 ViewVector = WorldPosition - View.WorldCameraOrigin;
+	// WorldPosition default value is Parameters.AbsoluteWorldPosition if not overridden by the user
+	float3 ViewVector = AbsoluteWorldPosition - View.WorldCameraOrigin;
 	float SceneDepth = length(ViewVector);
 	return GetAtmosphericFog(View.WorldCameraOrigin, ViewVector, SceneDepth, float3(0.f, 0.f, 0.f));
 #else
@@ -801,7 +802,7 @@ float GetSphericalParticleOpacity(FMaterialPixelParameters Parameters, float Den
 	float RescaledDensity = Density / ParticleRadius;
 
 	// Distance from point being shaded to particle center
-	float DistanceToParticle = length(Parameters.WorldPosition - ParticlePosition);
+	float DistanceToParticle = length(Parameters.AbsoluteWorldPosition - ParticlePosition);
 
 	FLATTEN
 	if (DistanceToParticle < ParticleRadius) 
@@ -1601,7 +1602,7 @@ void CalcMaterialParametersEx(
 {
 	// Remove the pre view translation
 	Parameters.WorldPosition_CamRelative = PixelPosition.xyz;
-	Parameters.WorldPosition = PixelPosition.xyz - View.PreViewTranslation.xyz;
+	Parameters.AbsoluteWorldPosition = PixelPosition.xyz - View.PreViewTranslation.xyz;
 
 	// If the material uses any non-offset world position expressions, calculate those parameters. If not, 
 	// the variables will have been initialised to 0 earlier.
@@ -1737,7 +1738,7 @@ float ApplyPixelDepthOffsetToMaterialParameters(inout FMaterialPixelParameters M
 	// Update positions used for shading
 	MaterialParameters.ScreenPosition.w += PixelDepthOffset;
 	MaterialParameters.SvPosition.w = MaterialParameters.ScreenPosition.w;
-	MaterialParameters.WorldPosition += MaterialParameters.CameraVector * PixelDepthOffset;
+	MaterialParameters.AbsoluteWorldPosition += MaterialParameters.CameraVector * PixelDepthOffset;
 
 	OutDepth = MaterialParameters.ScreenPosition.z / MaterialParameters.ScreenPosition.w;
 
diff --git a/Engine/Shaders/ShadowDepthPixelShader.usf b/Engine/Shaders/ShadowDepthPixelShader.usf
--- a/Engine/Shaders/ShadowDepthPixelShader.usf
+++ b/Engine/Shaders/ShadowDepthPixelShader.usf
@@ -105,7 +105,7 @@ void Main(
 
 #if SECONDARY_OCCLUSION || MATERIAL_INJECT_EMISSIVE_INTO_LPV
 	// Optimisation: only add to GV/VPL if the fragment is within the grid
-	float3 grid = WorldToGrid( MaterialParameters.WorldPosition.xyz );
+	float3 grid = WorldToGrid( MaterialParameters.AbsoluteWorldPosition.xyz );
 	float minGrid = min( grid.x, min(grid.y,grid.z ) );
 	float maxGrid = max( grid.x, max(grid.y,grid.z ) );
 
@@ -113,7 +113,7 @@ void Main(
 	if ( minGrid > -0.5f && maxGrid <= 32.5f ) // Account for 0.5 cell offset
 	{
 #if SECONDARY_OCCLUSION && (MATERIALBLENDING_MASKED || MATERIALBLENDING_SOLID) && !MATERIAL_BLOCK_GI
-		float3 GvInjectPosition = MaterialParameters.WorldPosition - normal * LpvWrite.GeometryVolumeInjectionBias;
+		float3 GvInjectPosition = MaterialParameters.AbsoluteWorldPosition - normal * LpvWrite.GeometryVolumeInjectionBias;
 #if LPV_GV_BOOST
 		if ( Primitive.LpvBiasMultiplier < 0.25 ) 
 		{
@@ -129,7 +129,7 @@ void Main(
 #endif
 
 #if MATERIAL_INJECT_EMISSIVE_INTO_LPV && !MATERIAL_BLOCK_GI
-		AddToVplList( MaterialParameters.WorldPosition, emissive * LpvWrite.EmissiveInjectionMultiplier, normal, true );
+		AddToVplList( MaterialParameters.AbsoluteWorldPosition, emissive * LpvWrite.EmissiveInjectionMultiplier, normal, true );
 #endif
 	}
 	else
diff --git a/Engine/Shaders/TranslucentLightInjectionShaders.usf b/Engine/Shaders/TranslucentLightInjectionShaders.usf
--- a/Engine/Shaders/TranslucentLightInjectionShaders.usf
+++ b/Engine/Shaders/TranslucentLightInjectionShaders.usf
@@ -110,15 +110,15 @@ SamplerState StaticShadowDepthTextureSampler;
 /** Transform used for static shadowing by spot and directional lights. */
 float4x4 WorldToStaticShadowMatrix;
 
-float GetLightFunctionShadowFactor(float3 WorldPositionForLighting)
+float GetLightFunctionShadowFactor(float3 AbsoluteWorldPositionForLighting)
 {
 	float ShadowFactor = 1;
 
 	// Apply light function after edge fading, so that a black light function at the edges can cause distant translucency to also be black
 	#if APPLY_LIGHT_FUNCTION
-		float4 LightVector = mul(float4(WorldPositionForLighting, 1),LightFunctionWorldToLight);
+		float4 LightVector = mul(float4(AbsoluteWorldPositionForLighting, 1),LightFunctionWorldToLight);
 		LightVector.xyz /= LightVector.w;
-		float3 LightFunction = GetLightFunctionColor(LightVector.xyz, WorldPositionForLighting);
+		float3 LightFunction = GetLightFunctionColor(LightVector.xyz, AbsoluteWorldPositionForLighting);
 
 		// We only suport monochrome light functions
 		ShadowFactor = dot(LightFunction, .3333f).x;
diff --git a/Engine/Source/Runtime/Engine/Private/Materials/HLSLMaterialTranslator.h b/Engine/Source/Runtime/Engine/Private/Materials/HLSLMaterialTranslator.h
--- a/Engine/Source/Runtime/Engine/Private/Materials/HLSLMaterialTranslator.h
+++ b/Engine/Source/Runtime/Engine/Private/Materials/HLSLMaterialTranslator.h
@@ -3785,7 +3785,7 @@ protected:
 		bUsesAtmosphericFog = true;
 		if( WorldPosition == INDEX_NONE )
 		{
-			return AddCodeChunk( MCT_Float4, TEXT("MaterialExpressionAtmosphericFog(Parameters, Parameters.WorldPosition)"));
+			return AddCodeChunk( MCT_Float4, TEXT("MaterialExpressionAtmosphericFog(Parameters, Parameters.AbsoluteWorldPosition)"));
 		}
 		else
 		{
