Integrate from Orion CL 2699859

	removed usage of ComputeWorldWorldPosCamRelative() where possible for better quality and performance, used SvPositionToTranslatedWorld() instead

diff --git a/Engine/Shaders/BasePassForForwardShadingPixelShader.usf b/Engine/Shaders/BasePassForForwardShadingPixelShader.usf
--- a/Engine/Shaders/BasePassForForwardShadingPixelShader.usf
+++ b/Engine/Shaders/BasePassForForwardShadingPixelShader.usf
@@ -145,7 +145,7 @@ void Main(
 	#if USE_WORLD_POSITION_EXCLUDING_SHADER_OFFSETS
 		{
 			float4 ScreenPosition = SvPositionToScreenPosition(SvPosition);
-			CalcMaterialParametersEx(MaterialParameters, SvPosition, ScreenPosition, bIsFrontFace, ComputeWorldWorldPosCamRelative(ScreenPosition), BasePassInterpolants.PixelPositionExcludingWPO);
+			CalcMaterialParametersEx(MaterialParameters, SvPosition, ScreenPosition, bIsFrontFace, SvPositionToTranslatedWorld(SvPosition), BasePassInterpolants.PixelPositionExcludingWPO);
 		}
 	#else
 		CalcMaterialParameters(MaterialParameters, SvPosition, bIsFrontFace);
diff --git a/Engine/Shaders/BasePassPixelShader.usf b/Engine/Shaders/BasePassPixelShader.usf
--- a/Engine/Shaders/BasePassPixelShader.usf
+++ b/Engine/Shaders/BasePassPixelShader.usf
@@ -654,7 +654,7 @@ void Main(
 	#if USE_WORLD_POSITION_EXCLUDING_SHADER_OFFSETS
 		{
 			float4 ScreenPosition = SvPositionToScreenPosition(SvPosition);
-			CalcMaterialParametersEx(MaterialParameters, SvPosition, ScreenPosition, bIsFrontFace, ComputeWorldWorldPosCamRelative(ScreenPosition), BasePassInterpolants.PixelPositionExcludingWPO);
+			CalcMaterialParametersEx(MaterialParameters, SvPosition, ScreenPosition, bIsFrontFace, SvPositionToTranslatedWorld(SvPosition), BasePassInterpolants.PixelPositionExcludingWPO);
 		}
 	#else
 		CalcMaterialParameters(MaterialParameters, SvPosition, bIsFrontFace);
diff --git a/Engine/Shaders/Common.usf b/Engine/Shaders/Common.usf
--- a/Engine/Shaders/Common.usf
+++ b/Engine/Shaders/Common.usf
@@ -1026,11 +1026,17 @@ float2 SvPositionToBufferUV(float4 SvPosition)
 	return SvPosition.xy * View.BufferSizeAndInvSize.zw;
 }
 	
-float3 SvPositionToWorld(float4 SvPosition)
+float3 SvPositionToTranslatedWorld(float4 SvPosition)
 {
 	float4 HomWorldPos = mul(float4(SvPosition.xyz, 1), View.SVPositionToTranslatedWorld);
 
-	return HomWorldPos.xyz / HomWorldPos.w - View.PreViewTranslation;
+	return HomWorldPos.xyz / HomWorldPos.w;
+}
+
+// prefer to use SvPositionToTranslatedWorld() for better quality
+float3 SvPositionToWorld(float4 SvPosition)
+{
+	return SvPositionToTranslatedWorld(SvPosition) - View.PreViewTranslation;
 }
 
 // investigate: doesn't work for usage with View.ScreenToWorld, see SvPositionToScreenPosition2()
diff --git a/Engine/Shaders/MaterialTemplate.usf b/Engine/Shaders/MaterialTemplate.usf
--- a/Engine/Shaders/MaterialTemplate.usf
+++ b/Engine/Shaders/MaterialTemplate.usf
@@ -818,7 +818,7 @@ float GetSphericalParticleOpacity(FMaterialPixelParameters Parameters, float Den
 		float NearDistance = Parameters.ScreenPosition.w - HemisphericalDistance;
 		float FarDistance = Parameters.ScreenPosition.w + HemisphericalDistance;
 
-		float SceneDepth = CalcSceneDepth(ScreenPositionToBufferUV(Parameters.ScreenPosition));
+		float SceneDepth = CalcSceneDepth(SvPositionToBufferUV(Parameters.SvPosition));
 		FarDistance = min(SceneDepth, FarDistance);
 
 		// Take into account opaque objects intersecting the sphere
@@ -1597,18 +1597,18 @@ void CalcMaterialParametersEx(
 	float4 SvPosition,
 	float4 ScreenPosition,
 	FIsFrontFace bIsFrontFace,
-	float3 PixelPosition,
-	float3 PixelPositionExcludingShaderOffsets)
+	float3 TranslatedWorldPosition,
+	float3 TranslatedWorldPositionExcludingShaderOffsets)
 {
 	// Remove the pre view translation
-	Parameters.WorldPosition_CamRelative = PixelPosition.xyz;
-	Parameters.AbsoluteWorldPosition = PixelPosition.xyz - View.PreViewTranslation.xyz;
+	Parameters.WorldPosition_CamRelative = TranslatedWorldPosition.xyz;
+	Parameters.AbsoluteWorldPosition = TranslatedWorldPosition.xyz - View.PreViewTranslation.xyz;
 
 	// If the material uses any non-offset world position expressions, calculate those parameters. If not, 
 	// the variables will have been initialised to 0 earlier.
 #if USE_WORLD_POSITION_EXCLUDING_SHADER_OFFSETS
-	Parameters.WorldPosition_NoOffsets_CamRelative = PixelPositionExcludingShaderOffsets;
-	Parameters.WorldPosition_NoOffsets = PixelPositionExcludingShaderOffsets - View.PreViewTranslation.xyz;
+	Parameters.WorldPosition_NoOffsets_CamRelative = TranslatedWorldPositionExcludingShaderOffsets;
+	Parameters.WorldPosition_NoOffsets = TranslatedWorldPositionExcludingShaderOffsets - View.PreViewTranslation.xyz;
 #endif
 
 	Parameters.SvPosition = SvPosition;
@@ -1616,10 +1616,10 @@ void CalcMaterialParametersEx(
 
 #if COMPILER_GLSL_ES2
 	// ES2 normalize isn't done accurately. This seems to fix it.
-	// Originally this was normalize(normalize(PixelPosition.xyz)) but tegra4 appears to optimize that out.
+	// Originally this was normalize(normalize(TranslatedWorldPosition.xyz)) but tegra4 appears to optimize that out.
 	Parameters.CameraVector = normalize(-0.01 * Parameters.WorldPosition_CamRelative.xyz);
 #else
-	// PixelPosition is the world position translated to the camera position, which is just -CameraVector
+	// TranslatedWorldPosition is the world position translated to the camera position, which is just -CameraVector
 	Parameters.CameraVector = normalize(-Parameters.WorldPosition_CamRelative.xyz);
 #endif
 
@@ -1661,15 +1661,7 @@ void CalcMaterialParametersEx(
 #endif // !PARTICLE_SPRITE_FACTORY
 }
 
-// this will be optimized to not use ScreenPosition
-float3 ComputeWorldWorldPosCamRelative(float4 ScreenPosition)
-{
-	float4 TranslatedWorld = mul( ScreenPosition, View.ClipToTranslatedWorld );
-
-	return TranslatedWorld.xyz / TranslatedWorld.w;
-}
-
-// convenience function to setup CalcMaterialParameters assuming we don't support PixelPositionExcludingShaderOffsets
+// convenience function to setup CalcMaterialParameters assuming we don't support TranslatedWorldPositionExcludingShaderOffsets
 // @param SvPosition from SV_Position when rendering the view, for other projections e.g. shadowmaps this function cannot be used and you need to call CalcMaterialParametersEx()
 void CalcMaterialParameters(
 	in out FMaterialPixelParameters Parameters,
@@ -1677,9 +1669,9 @@ void CalcMaterialParameters(
 	FIsFrontFace bIsFrontFace)
 {
 	float4 ScreenPosition = SvPositionToScreenPosition(SvPosition);
-	float3 PixelPosition = ComputeWorldWorldPosCamRelative(ScreenPosition);
+	float3 TranslatedWorldPosition = SvPositionToTranslatedWorld(SvPosition);
 
-	CalcMaterialParametersEx(Parameters, SvPosition, ScreenPosition, bIsFrontFace, PixelPosition, PixelPosition);
+	CalcMaterialParametersEx(Parameters, SvPosition, ScreenPosition, bIsFrontFace, TranslatedWorldPosition, TranslatedWorldPosition);
 }
 
 /** Assemble the transform from tangent space into world space */
diff --git a/Engine/Shaders/ParticleGPUSpriteVertexFactory.usf b/Engine/Shaders/ParticleGPUSpriteVertexFactory.usf
--- a/Engine/Shaders/ParticleGPUSpriteVertexFactory.usf
+++ b/Engine/Shaders/ParticleGPUSpriteVertexFactory.usf
@@ -377,9 +377,9 @@ FMaterialPixelParameters GetMaterialPixelParameters( FVertexFactoryInterpolantsV
 	// can be optimized
 	{
 		float4 ScreenPosition = SvPositionToScreenPosition(SvPosition);
-		float3 PixelPosition = ComputeWorldWorldPosCamRelative(ScreenPosition);
+		float3 TranslatedWorldPosition = SvPositionToTranslatedWorld(SvPosition);
 
-		Result.TangentToWorld = GetSphericalParticleNormal(PixelPosition + View.WorldCameraOrigin, Interpolants.ParticlePositionAndSize.xyz, Interpolants.ParticlePositionAndSize.w);
+		Result.TangentToWorld = GetSphericalParticleNormal(TranslatedWorldPosition + View.WorldCameraOrigin, Interpolants.ParticlePositionAndSize.xyz, Interpolants.ParticlePositionAndSize.w);
 	}
 #else
 	half3 TangentToWorld0 = Interpolants.TangentToWorld0.xyz;
diff --git a/Engine/Shaders/ParticleSpriteVertexFactory.usf b/Engine/Shaders/ParticleSpriteVertexFactory.usf
--- a/Engine/Shaders/ParticleSpriteVertexFactory.usf
+++ b/Engine/Shaders/ParticleSpriteVertexFactory.usf
@@ -173,8 +173,8 @@ FMaterialPixelParameters GetMaterialPixelParameters(FVertexFactoryInterpolantsVS
 		{	
 			// can be optimized
 			float4 ScreenPosition = SvPositionToScreenPosition(SvPosition);
-			float3 PixelPosition = ComputeWorldWorldPosCamRelative(ScreenPosition);
-			Result.TangentToWorld = GetSphericalParticleNormal(PixelPosition + View.WorldCameraOrigin, Interpolants.ParticlePositionAndSize.xyz, Interpolants.ParticlePositionAndSize.w);
+			float3 TranslatedWorldPosition = SvPositionToTranslatedWorld(SvPosition);
+			Result.TangentToWorld = GetSphericalParticleNormal(TranslatedWorldPosition + View.WorldCameraOrigin, Interpolants.ParticlePositionAndSize.xyz, Interpolants.ParticlePositionAndSize.w);
 		}
 	#else
 		Result.TangentToWorld = AssembleTangentToWorld(TangentToWorld0.xyz, TangentToWorld2);
diff --git a/Engine/Shaders/ShadowDepthPixelShader.usf b/Engine/Shaders/ShadowDepthPixelShader.usf
--- a/Engine/Shaders/ShadowDepthPixelShader.usf
+++ b/Engine/Shaders/ShadowDepthPixelShader.usf
@@ -61,7 +61,7 @@ void Main(
 	#if INTERPOLATE_POSITION
 		{
 			float4 ScreenPosition = SvPositionToScreenPosition(SvPosition);
-			CalcMaterialParametersEx(MaterialParameters, SvPosition, ScreenPosition, 1, ComputeWorldWorldPosCamRelative(ScreenPosition), SvPosition);
+			CalcMaterialParametersEx(MaterialParameters, SvPosition, ScreenPosition, 1, SvPositionToTranslatedWorld(SvPosition), SvPosition);
 		}
 	#else
 		CalcMaterialParameters(MaterialParameters, SvPosition, 1);
