integrated from Orion

fixed depth fade material node

diff --git a/Engine/Shaders/Common.usf b/Engine/Shaders/Common.usf
--- a/Engine/Shaders/Common.usf
+++ b/Engine/Shaders/Common.usf
@@ -1042,19 +1042,16 @@ float3 SvPositionToWorld(float4 SvPosition)
 // investigate: doesn't work for usage with View.ScreenToWorld, see SvPositionToScreenPosition2()
 float4 SvPositionToScreenPosition(float4 SvPosition)
 {
+	// todo: is already in .w or needs to be reconstructed like this:
+//	SvPosition.w = ConvertFromDeviceZ(SvPosition.z);
+
 	float2 PixelPos = SvPosition.xy - View.ViewRectMin.xy;	
 
 	// NDC (NormalizedDeviceCoordinates, after the perspective divide)
 	float3 NDCPos = float3( (PixelPos * View.ViewSizeAndInvSize.zw - 0.5f) * float2(2, -2), SvPosition.z);
 
-	// clipspace (before the perspective divide)
-//	return float4(NDCPos.xyz, 1) * SvPosition.w;
-
-	// so .w has the SceneDepth, some mobile code wants that
-//	Ret *= Ret.z;
-
-	// 
-	return float4(NDCPos.xyz, 1);
+	// SvPosition.w: so .w has the SceneDepth, some mobile code and the DepthFade material expression wants that
+	return float4(NDCPos.xyz, 1) * SvPosition.w;
 }
 
 // ----------------------------
