New options to enable selective outputs in the base pass for opaque geometry.
Optionnal outputs are : CustomData, PreshadowFactors and Velocity (when enabled)
Controlled by ini file variable : r.SelectiveBasePassOutputs=true

diff --git a/Engine/Config/BaseEditor.ini b/Engine/Config/BaseEditor.ini
--- a/Engine/Config/BaseEditor.ini
+++ b/Engine/Config/BaseEditor.ini
@@ -433,6 +433,7 @@ bUseFullGCAssetClassNames=true
 +VersionedIntRValues=r.AllowStaticLighting
 +VersionedIntRValues=r.GBuffer
 +VersionedIntRValues=r.BasePassOutputsVelocity
++VersionedIntRValues=r.SelectiveBasePassOutputs
 +VersionedIntRValues=r.DBuffer
 +VersionedIntRValues=r.Shaders.KeepDebugInfo
 +VersionedIntRValues=r.Shaders.Optimize
diff --git a/Engine/Shaders/BasePassCommon.usf b/Engine/Shaders/BasePassCommon.usf
--- a/Engine/Shaders/BasePassCommon.usf
+++ b/Engine/Shaders/BasePassCommon.usf
@@ -13,7 +13,18 @@
 #define NEEDS_BASEPASS_FOGGING			(ENABLE_TRANSLUCENCY_VERTEX_FOG && ((MATERIALBLENDING_TRANSLUCENT || MATERIALBLENDING_ADDITIVE || MATERIALBLENDING_MODULATE) && !MATERIAL_USES_SCENE_COLOR_COPY))
 #define NEEDS_LIGHTMAP					(NEEDS_LIGHTMAP_COORDINATE)
 
-#define WRITES_VELOCITY_TO_GBUFFER		(OUTPUT_GBUFFER_VELOCITY && (MATERIALBLENDING_SOLID || MATERIALBLENDING_MASKED))
+#define USES_GBUFFER						(FEATURE_LEVEL >= FEATURE_LEVEL_SM4 && (MATERIALBLENDING_SOLID || MATERIALBLENDING_MASKED) && !NO_GBUFFER)
+
+// Only some shader models actually need custom data.
+#define WRITES_CUSTOMDATA_TO_GBUFFER		(USES_GBUFFER && (!SELECTIVE_BASEPASS_OUTPUTS || MATERIAL_SHADINGMODEL_SUBSURFACE || MATERIAL_SHADINGMODEL_PREINTEGRATED_SKIN || MATERIAL_SHADINGMODEL_SUBSURFACE_PROFILE || MATERIAL_SHADINGMODEL_CLEAR_COAT))
+
+// Based on GetPrecomputedShadowMasks()
+#define WRITES_PRECSHADOWFACTOR_TO_GBUFFER	(USES_GBUFFER && ALLOW_STATIC_LIGHTING && (!SELECTIVE_BASEPASS_OUTPUTS || (STATICLIGHTING_TEXTUREMASK && STATICLIGHTING_SIGNEDDISTANCEFIELD)))
+#define WRITES_PRECSHADOWFACTOR_ZERO		(USES_GBUFFER && ALLOW_STATIC_LIGHTING && (!SELECTIVE_BASEPASS_OUTPUTS || HQ_TEXTURE_LIGHTMAP || LQ_TEXTURE_LIGHTMAP))
+
+// If a primitive has static lighting, we assume it is not moving. If it is, it will be rerendered in an extra renderpass.
+#define WRITES_VELOCITY_TO_GBUFFER			(USES_GBUFFER && OUTPUT_GBUFFER_VELOCITY && (!SELECTIVE_BASEPASS_OUTPUTS || !(STATICLIGHTING_TEXTUREMASK || STATICLIGHTING_SIGNEDDISTANCEFIELD || HQ_TEXTURE_LIGHTMAP || LQ_TEXTURE_LIGHTMAP)))
+
 // Use an extra interpolator for position
 #define WRITES_VELOCITY_TO_GBUFFER_USE_POS_INTERPOLATOR		(WRITES_VELOCITY_TO_GBUFFER && USING_TESSELLATION)
 
diff --git a/Engine/Shaders/BasePassPixelShader.usf b/Engine/Shaders/BasePassPixelShader.usf
--- a/Engine/Shaders/BasePassPixelShader.usf
+++ b/Engine/Shaders/BasePassPixelShader.usf
@@ -57,7 +57,6 @@ float4 PointSkyBentNormal;
 
 #define EDITOR_ALPHA2COVERAGE (USE_EDITOR_COMPOSITING && FEATURE_LEVEL >= FEATURE_LEVEL_SM5 && !COMPILER_GLSL)
 
-
 #if TRANSLUCENCY_LIGHTING_SURFACE || TRANSLUCENCY_LIGHTING_SURFACE_PERPIXEL
 
 #if FEATURE_LEVEL >= FEATURE_LEVEL_SM5
@@ -614,29 +613,63 @@ void ApplyPixelDepthOffsetForBasePass(inout FMaterialPixelParameters MaterialPar
 #endif
 }
 
+#if USES_GBUFFER
+
+// The selective output mask can only depend on defines, since the shadow will not export the data.
+uint GetSelectiveOutputMask()
+{
+	uint Mask = 0;
+#if !WRITES_CUSTOMDATA_TO_GBUFFER
+	Mask |= SKIP_CUSTOMDATA_MASK;
+#endif
+#if !WRITES_PRECSHADOWFACTOR_TO_GBUFFER
+	Mask |= SKIP_PRECSHADOW_MASK;
+#endif
+#if WRITES_PRECSHADOWFACTOR_TO_ZERO
+	Mask |= ZERO_PRECSHADOW_MASK;
+#endif
+#if !WRITES_VELOCITY_TO_GBUFFER
+	Mask |= SKIP_VELOCITY_MASK;
+#endif
+	return Mask;
+}
+#endif // USES_GBUFFER
+
+#if OUTPUT_GBUFFER_VELOCITY
+#define GBUFFER_D_TARGET SV_Target5
+#define GBUFFER_E_TARGET SV_Target6 // Only used with ALLOW_STATIC_LIGHTING
+#else
+#define GBUFFER_D_TARGET SV_Target4
+#define GBUFFER_E_TARGET SV_Target5
+#endif
+
 void Main(
 	FVertexFactoryInterpolantsVSToPS Interpolants,
 	FBasePassInterpolantsVSToPS BasePassInterpolants,
 	in INPUT_POSITION_QUALIFIERS float4 SvPosition : SV_Position
 	OPTIONAL_IsFrontFace
 	,out float4 OutColor		: SV_Target0
-#if FEATURE_LEVEL >= FEATURE_LEVEL_SM4 && (MATERIALBLENDING_SOLID || MATERIALBLENDING_MASKED) && !NO_GBUFFER
+
+// Two options control the render target bindings : OUTPUT_GBUFFER_VELOCITY and ALLOW_STATIC_LIGHTING
+#if USES_GBUFFER
 	,out float4 OutGBufferA : SV_Target1 
 	,out float4 OutGBufferB : SV_Target2
 	,out float4 OutGBufferC : SV_Target3
-	,out float4 OutGBufferD : SV_Target4
-	#if ALLOW_STATIC_LIGHTING
-		,out float4 OutGBufferE : SV_Target5
-		#if WRITES_VELOCITY_TO_GBUFFER
-			, out float4 OutGBufferVelocity : SV_Target6
-		#endif
-	#else
-		#if WRITES_VELOCITY_TO_GBUFFER
-			, out float4 OutGBufferVelocity : SV_Target5
-		#endif
-	#endif
+
+// Big issue here on PC with AMD hardware : putting the velocity buffer after an optionnal RT fails!!!
+#if WRITES_VELOCITY_TO_GBUFFER
+	,out float4 OutGBufferVelocity : SV_Target4
+#endif
+
+#if WRITES_CUSTOMDATA_TO_GBUFFER
+	,out float4 OutGBufferD : GBUFFER_D_TARGET
 #endif
 
+#if WRITES_PRECSHADOWFACTOR_TO_GBUFFER
+	,out float4 OutGBufferE : GBUFFER_E_TARGET
+#endif
+#endif // USES_GBUFFER
+
 	OPTIONAL_OutDepthConservative
 
 #if EDITOR_ALPHA2COVERAGE != 0
@@ -645,6 +678,16 @@ void Main(
 #endif
    )
 {
+#if !WRITES_CUSTOMDATA_TO_GBUFFER
+	float4 OutGBufferD = 0;
+#endif
+#if !WRITES_PRECSHADOWFACTOR_TO_GBUFFER
+	float4 OutGBufferE = 0;
+#endif
+#if !WRITES_VELOCITY_TO_GBUFFER
+	float4 OutGBufferVelocity = 0;
+#endif
+
 #if EDITOR_ALPHA2COVERAGE != 0
 	OutCoverage = InCoverage;
 #endif
@@ -836,7 +879,7 @@ void Main(
 
 	float4 PrecomputedShadowFactors = GetPrecomputedShadowMasks(Interpolants);
 
-	#if FEATURE_LEVEL >= FEATURE_LEVEL_SM4 && (MATERIALBLENDING_SOLID || MATERIALBLENDING_MASKED) && !NO_GBUFFER
+	#if USES_GBUFFER
 
 		FGBufferData GBuffer = (FGBufferData)0;
 
@@ -850,82 +893,79 @@ void Main(
 		GBuffer.GBufferAO = MaterialAO;
 		GBuffer.PerObjectGBufferData = Primitive.PerObjectGBufferData;
 
-		#if MATERIAL_SHADINGMODEL_UNLIT
-			GBuffer.ShadingModelID = SHADINGMODELID_UNLIT;
-		#elif MATERIAL_SHADINGMODEL_DEFAULT_LIT
-			GBuffer.ShadingModelID = SHADINGMODELID_DEFAULT_LIT;
-		#elif MATERIAL_SHADINGMODEL_SUBSURFACE
-			GBuffer.ShadingModelID = SHADINGMODELID_SUBSURFACE;
-			GBuffer.CustomData.rgb = EncodeSubsurfaceColor(SubsurfaceColor);
-			GBuffer.CustomData.a = Opacity;
-		#elif MATERIAL_SHADINGMODEL_PREINTEGRATED_SKIN
-			GBuffer.ShadingModelID = SHADINGMODELID_PREINTEGRATED_SKIN;
-			GBuffer.CustomData.rgb = EncodeSubsurfaceColor(SubsurfaceColor);
-			GBuffer.CustomData.a = Opacity;
-		#elif MATERIAL_SHADINGMODEL_SUBSURFACE_PROFILE
-			GBuffer.ShadingModelID = SHADINGMODELID_SUBSURFACE_PROFILE;
-			GBuffer.CustomData.rgb = EncodeSubsurfaceProfile(SubsurfaceProfile);
-			GBuffer.CustomData.a = Opacity;
-		#elif MATERIAL_SHADINGMODEL_CLEAR_COAT
-			GBuffer.ShadingModelID = SHADINGMODELID_CLEAR_COAT;
-
-			float ClearCoat				= saturate( GetMaterialCustomData0(MaterialParameters) );
-			float ClearCoatRoughness	= saturate( GetMaterialCustomData1(MaterialParameters) );
-			float MetalSpec = 0.9;
-
-			float NoV = saturate( dot( MaterialParameters.WorldNormal, MaterialParameters.CameraVector ) );
-
-			// Approximation of refraction's effect on EnvBRDF
-			float RefractionScale = ( (NoV * 0.5 + 0.5) * NoV - 1 ) * saturate( 1.25 - 1.25 * Roughness ) + 1;
+	#if MATERIAL_SHADINGMODEL_UNLIT
+		GBuffer.ShadingModelID = SHADINGMODELID_UNLIT;
+	#elif MATERIAL_SHADINGMODEL_DEFAULT_LIT
+		GBuffer.ShadingModelID = SHADINGMODELID_DEFAULT_LIT;
+	#elif MATERIAL_SHADINGMODEL_SUBSURFACE
+		GBuffer.ShadingModelID = SHADINGMODELID_SUBSURFACE;
+		GBuffer.CustomData.rgb = EncodeSubsurfaceColor(SubsurfaceColor);
+		GBuffer.CustomData.a = Opacity;
+	#elif MATERIAL_SHADINGMODEL_PREINTEGRATED_SKIN
+		GBuffer.ShadingModelID = SHADINGMODELID_PREINTEGRATED_SKIN;
+		GBuffer.CustomData.rgb = EncodeSubsurfaceColor(SubsurfaceColor);
+		GBuffer.CustomData.a = Opacity;
+	#elif MATERIAL_SHADINGMODEL_SUBSURFACE_PROFILE
+		GBuffer.ShadingModelID = SHADINGMODELID_SUBSURFACE_PROFILE;
+		GBuffer.CustomData.rgb = EncodeSubsurfaceProfile(SubsurfaceProfile);
+		GBuffer.CustomData.a = Opacity;
+	#elif MATERIAL_SHADINGMODEL_CLEAR_COAT
+		GBuffer.ShadingModelID = SHADINGMODELID_CLEAR_COAT;
+
+		float ClearCoat				= saturate( GetMaterialCustomData0(MaterialParameters) );
+		float ClearCoatRoughness	= saturate( GetMaterialCustomData1(MaterialParameters) );
+		float MetalSpec = 0.9;
+
+		float NoV = saturate( dot( MaterialParameters.WorldNormal, MaterialParameters.CameraVector ) );
+
+		// Approximation of refraction's effect on EnvBRDF
+		float RefractionScale = ( (NoV * 0.5 + 0.5) * NoV - 1 ) * saturate( 1.25 - 1.25 * Roughness ) + 1;
 		
-			// Approximation of absorption integral, tuned for Roughness=0.4
-			float3 AbsorptionColor = BaseColor * (1 / MetalSpec);
-			float3 Absorption = AbsorptionColor * ( (NoV - 1) * 0.85 * ( 1 - lerp( AbsorptionColor, Square(AbsorptionColor), -0.78 ) ) + 1 );
-
-			GBuffer.BaseColor = lerp( BaseColor, MetalSpec * Absorption * lerp( 1, RefractionScale, Metallic ), ClearCoat );
-			GBuffer.Specular *= Luminance( Absorption ) * RefractionScale;
-
-			GBuffer.CustomData.x = ClearCoat;
-			GBuffer.CustomData.y = ClearCoatRoughness;
-		#elif MATERIAL_SHADINGMODEL_TWOSIDED_FOLIAGE
-			GBuffer.ShadingModelID = SHADINGMODELID_TWOSIDED_FOLIAGE;
-			GBuffer.CustomData.rgb = EncodeSubsurfaceColor(SubsurfaceColor);
-			GBuffer.CustomData.a = Opacity;
-		#elif MATERIAL_SHADINGMODEL_HAIR
-			GBuffer.ShadingModelID = SHADINGMODELID_HAIR;
-			#if NUM_MATERIAL_OUTPUTS_GETTANGENTOUTPUT > 0
-				float3 Tangent = GetTangentOutput0(MaterialParameters);
-				Tangent = normalize( TransformTangentVectorToWorld( MaterialParameters.TangentToWorld, Tangent ) );
-				GBuffer.WorldNormal = Tangent;
-			#endif
-			GBuffer.CustomData.xy = UnitVectorToOctahedron( MaterialParameters.WorldNormal ) * 0.5 + 0.5;
-			GBuffer.CustomData.z = saturate( GetMaterialCustomData1(MaterialParameters) );	// Roughness2
-			GBuffer.CustomData.w = saturate( GetMaterialCustomData0(MaterialParameters) );	// Opacity
-		#else
-			// missing shading model, compiler should report ShadingModelID is not set
-		#endif
-
-		#if !ALLOW_STATIC_LIGHTING
-			float4 OutGBufferE = 0;
+		// Approximation of absorption integral, tuned for Roughness=0.4
+		float3 AbsorptionColor = BaseColor * (1 / MetalSpec);
+		float3 Absorption = AbsorptionColor * ( (NoV - 1) * 0.85 * ( 1 - lerp( AbsorptionColor, Square(AbsorptionColor), -0.78 ) ) + 1 );
+
+		GBuffer.BaseColor = lerp( BaseColor, MetalSpec * Absorption * lerp( 1, RefractionScale, Metallic ), ClearCoat );
+		GBuffer.Specular *= Luminance( Absorption ) * RefractionScale;
+
+		GBuffer.CustomData.x = ClearCoat;
+		GBuffer.CustomData.y = ClearCoatRoughness;
+	#elif MATERIAL_SHADINGMODEL_TWOSIDED_FOLIAGE
+		GBuffer.ShadingModelID = SHADINGMODELID_TWOSIDED_FOLIAGE;
+		GBuffer.CustomData.rgb = EncodeSubsurfaceColor(SubsurfaceColor);
+		GBuffer.CustomData.a = Opacity;
+	#elif MATERIAL_SHADINGMODEL_HAIR
+		GBuffer.ShadingModelID = SHADINGMODELID_HAIR;
+		#if NUM_MATERIAL_OUTPUTS_GETTANGENTOUTPUT > 0
+			float3 Tangent = GetTangentOutput0(MaterialParameters);
+			Tangent = normalize( TransformTangentVectorToWorld( MaterialParameters.TangentToWorld, Tangent ) );
+			GBuffer.WorldNormal = Tangent;
 		#endif
+		GBuffer.CustomData.xy = UnitVectorToOctahedron( MaterialParameters.WorldNormal ) * 0.5 + 0.5;
+		GBuffer.CustomData.z = saturate( GetMaterialCustomData1(MaterialParameters) );	// Roughness2
+		GBuffer.CustomData.w = saturate( GetMaterialCustomData0(MaterialParameters) );	// Opacity
+	#else
+		// missing shading model, compiler should report ShadingModelID is not set
+	#endif
 
-		#if WRITES_VELOCITY_TO_GBUFFER
-			{
-				// 2d velocity, includes camera an object motion
-			#if WRITES_VELOCITY_TO_GBUFFER_USE_POS_INTERPOLATOR
-				float2 Velocity = Calculate2DVelocity(BasePassInterpolants.VelocityScreenPosition, BasePassInterpolants.VelocityPrevScreenPosition);
-			#else
-				float2 Velocity = Calculate2DVelocity(MaterialParameters.ScreenPosition, BasePassInterpolants.VelocityPrevScreenPosition);
-			#endif
+		GBuffer.SelectiveOutputMask = GetSelectiveOutputMask();
 
-				// Make sure not to touch 0,0 which is clear color
-				GBuffer.Velocity = float4(EncodeVelocityToTexture(Velocity), 0, 0) * BasePassInterpolants.VelocityPrevScreenPosition.z;
-			}
+	#if WRITES_VELOCITY_TO_GBUFFER
+		{
+			// 2d velocity, includes camera an object motion
+		#if WRITES_VELOCITY_TO_GBUFFER_USE_POS_INTERPOLATOR
+			float2 Velocity = Calculate2DVelocity(BasePassInterpolants.VelocityScreenPosition, BasePassInterpolants.VelocityPrevScreenPosition);
 		#else
-			GBuffer.Velocity = 0;
-			float4 OutGBufferVelocity = 0;
+			float2 Velocity = Calculate2DVelocity(MaterialParameters.ScreenPosition, BasePassInterpolants.VelocityPrevScreenPosition);
 		#endif
 
+			// Make sure not to touch 0,0 which is clear color
+			GBuffer.Velocity = float4(EncodeVelocityToTexture(Velocity), 0, 0) * BasePassInterpolants.VelocityPrevScreenPosition.z;
+		}
+	#else
+		GBuffer.Velocity = 0;
+	#endif
+
 		// -0.5 .. 0.5, could be optimzed as lower quality noise would be sufficient
 		float QuantizationBias = PseudoRandom( MaterialParameters.SvPosition.xy ) - 0.5f;
 		EncodeGBuffer(GBuffer, OutGBufferA, OutGBufferB, OutGBufferC, OutGBufferD, OutGBufferE, OutGBufferVelocity, QuantizationBias);
diff --git a/Engine/Shaders/DeferredShadingCommon.usf b/Engine/Shaders/DeferredShadingCommon.usf
--- a/Engine/Shaders/DeferredShadingCommon.usf
+++ b/Engine/Shaders/DeferredShadingCommon.usf
@@ -212,16 +212,30 @@ float Decode71(float Scalar, out uint Mask)
 #define SHADINGMODELID_SUBSURFACE_PROFILE	5
 #define SHADINGMODELID_TWOSIDED_FOLIAGE		6
 #define SHADINGMODELID_HAIR					7
-#define SHADINGMODELID_NUM					8
 
-float EncodeShadingModelId(uint ShadingModelId)
+#define SHADINGMODELID_MASK					0xF					
+
+// The flags are defined so that 0 value has no effect!
+#define SKIP_CUSTOMDATA_MASK			(1 << 4)
+#define SKIP_PRECSHADOW_MASK			(1 << 5)
+#define ZERO_PRECSHADOW_MASK			(1 << 6)
+#define SKIP_VELOCITY_MASK				(1 << 7)
+
+
+float EncodeShadingModelIdAndSelectiveOutputMask(uint ShadingModelId, uint SelectiveOutputMask)
 {
-	return (ShadingModelId + .5f) / (float)SHADINGMODELID_NUM;
+	uint Value = (ShadingModelId & SHADINGMODELID_MASK) | SelectiveOutputMask;
+	return (float)Value / (float)0xFF;
 }
 
 uint DecodeShadingModelId(float InPackedChannel)
 {
-	return (uint)( InPackedChannel * SHADINGMODELID_NUM );
+	return ((uint)round(InPackedChannel * (float)0xFF)) & SHADINGMODELID_MASK;
+}
+
+uint DecodeSelectiveOutputMask(float InPackedChannel)
+{
+	return ((uint)round(InPackedChannel * (float)0xFF)) & ~SHADINGMODELID_MASK;
 }
 
 bool IsSubsurfaceModel(int ShadingModel)
@@ -261,6 +275,8 @@ struct FGBufferData
 	float GBufferAO;
 	// 0..255 
 	uint ShadingModelID;
+	// 0..255 
+	uint SelectiveOutputMask;
 	// 0..1, 2 bits, use HasDistanceFieldRepresentation(GBuffer) or HasHeightfieldRepresentation(GBuffer) to extract
 	float PerObjectGBufferData;
 	// in world units
@@ -360,7 +376,7 @@ void EncodeGBuffer(
 		OutGBufferD = GBuffer.CustomData;
 	}
 
-	OutGBufferB.a = EncodeShadingModelId(GBuffer.ShadingModelID);
+	OutGBufferB.a = EncodeShadingModelIdAndSelectiveOutputMask(GBuffer.ShadingModelID, GBuffer.SelectiveOutputMask);
 
 	OutGBufferE = GBuffer.PrecomputedShadowFactors;
 
@@ -416,12 +432,13 @@ FGBufferData DecodeGBufferData(
 	FGBufferData GBuffer;
 
 	GBuffer.WorldNormal = DecodeNormal( InGBufferA.xyz );
-
 	if(bGetNormalizedNormal)
 	{
 		GBuffer.WorldNormal = normalize(GBuffer.WorldNormal);
 	}
 
+	GBuffer.SelectiveOutputMask = DecodeSelectiveOutputMask(InGBufferB.a);
+
 	GBuffer.PerObjectGBufferData = InGBufferA.a;
 	GBuffer.ShadingModelID = DecodeShadingModelId(InGBufferB.a);
 	
@@ -433,7 +450,8 @@ FGBufferData DecodeGBufferData(
 	GBuffer.Specular = GBuffer.StoredSpecular.g;
 	
 	GBuffer.Roughness = InGBufferB.b;
-	GBuffer.CustomData = InGBufferD;
+	GBuffer.CustomData = !(GBuffer.SelectiveOutputMask & SKIP_CUSTOMDATA_MASK) ? InGBufferD : 0;
+
 #if ALLOW_STATIC_LIGHTING
 	GBuffer.GBufferAO = 1;
 	GBuffer.IndirectIrradiance = DecodeIndirectIrradiance(InGBufferC.a);
@@ -441,7 +459,8 @@ FGBufferData DecodeGBufferData(
 	GBuffer.GBufferAO = InGBufferC.a;
 	GBuffer.IndirectIrradiance = 1;
 #endif
-	GBuffer.PrecomputedShadowFactors = InGBufferE;
+
+	GBuffer.PrecomputedShadowFactors = !(GBuffer.SelectiveOutputMask & SKIP_PRECSHADOW_MASK) ? InGBufferE :  ((GBuffer.SelectiveOutputMask & ZERO_PRECSHADOW_MASK) ? 0 :  1);
 	GBuffer.CustomDepth = ConvertFromDeviceZ(CustomNativeDepth);
 	GBuffer.CustomStencil = CustomStencil;
 	GBuffer.Depth = SceneDepth;
@@ -467,7 +486,7 @@ FGBufferData DecodeGBufferData(
 		#endif //USE_DEVELOPMENT_SHADERS
 	}
 
-	GBuffer.Velocity = InGBufferVelocity;
+	GBuffer.Velocity = !(GBuffer.SelectiveOutputMask & SKIP_VELOCITY_MASK) ? InGBufferVelocity : 0;
 
 	return GBuffer;
 }
diff --git a/Engine/Source/Runtime/Engine/Classes/Engine/RendererSettings.h b/Engine/Source/Runtime/Engine/Classes/Engine/RendererSettings.h
--- a/Engine/Source/Runtime/Engine/Classes/Engine/RendererSettings.h
+++ b/Engine/Source/Runtime/Engine/Classes/Engine/RendererSettings.h
@@ -255,6 +255,12 @@ class ENGINE_API URendererSettings : public UDeveloperSettings
 		ConfigRestartRequired=true))
 	uint32 bBasePassOutputsVelocity:1;
 
+	UPROPERTY(config, EditAnywhere, Category=Optimizations, meta=(
+		ConsoleVariable="r.SelectiveBasePassOutputs", DisplayName="Selectively output to the GBuffer rendertargets",
+		ToolTip="Enables not exporting to the GBuffer rendertargets that are not relevant. Changing this setting requires restarting the editor.",
+		ConfigRestartRequired=true))
+	uint32 bSelectiveBasePassOutputs:1;
+
 	UPROPERTY(config, EditAnywhere, Category=Editor, meta=(
 		ConsoleVariable="r.WireframeCullThreshold",DisplayName="Wireframe Cull Threshold",
 		ToolTip="Screen radius at which wireframe objects are culled. Larger values can improve performance when viewing a scene in wireframe."))
diff --git a/Engine/Source/Runtime/Engine/Private/ShaderCompiler/ShaderCompiler.cpp b/Engine/Source/Runtime/Engine/Private/ShaderCompiler/ShaderCompiler.cpp
--- a/Engine/Source/Runtime/Engine/Private/ShaderCompiler/ShaderCompiler.cpp
+++ b/Engine/Source/Runtime/Engine/Private/ShaderCompiler/ShaderCompiler.cpp
@@ -1939,6 +1939,11 @@ void GlobalBeginCompileShader(
 		Input.Environment.SetDefine(TEXT("OUTPUT_GBUFFER_VELOCITY"), CVar ? (CVar->GetValueOnGameThread() != 0) : 0);
 	}
 
+	{
+		static const auto CVar = IConsoleManager::Get().FindTConsoleVariableDataInt(TEXT("r.SelectiveBasePassOutputs"));
+		Input.Environment.SetDefine(TEXT("SELECTIVE_BASEPASS_OUTPUTS"), CVar ? (CVar->GetValueOnGameThread() != 0) : 0);
+	}
+
 	{
 		static const auto CVar = IConsoleManager::Get().FindTConsoleVariableDataInt(TEXT("r.GBuffer"));
 		Input.Environment.SetDefine(TEXT("NO_GBUFFER"), CVar ? (CVar->GetValueOnGameThread() == 0) : 0);
diff --git a/Engine/Source/Runtime/Renderer/Private/BasePassRendering.cpp b/Engine/Source/Runtime/Renderer/Private/BasePassRendering.cpp
--- a/Engine/Source/Runtime/Renderer/Private/BasePassRendering.cpp
+++ b/Engine/Source/Runtime/Renderer/Private/BasePassRendering.cpp
@@ -7,6 +7,21 @@
 #include "RendererPrivate.h"
 #include "ScenePrivate.h"
 
+// Changing this causes a full shader recompile
+static TAutoConsoleVariable<int32> CVarSelectiveBasePassOutputs(
+	TEXT("r.SelectiveBasePassOutputs"),
+	0,
+	TEXT("Enables shaders to only export to relevant rendertargets.\n") \
+	TEXT(" 0: Export in all rendertargets.\n") \
+	TEXT(" 1: Export only into relevant rendertarget.\n"),
+	ECVF_ReadOnly | ECVF_RenderThreadSafe);
+
+
+bool UseSelectiveBasePassOutputs()
+{
+	return CVarSelectiveBasePassOutputs.GetValueOnAnyThread() == 1;
+}
+
 /** Whether to replace lightmap textures with solid colors to visualize the mip-levels. */
 bool GVisualizeMipLevels = false;
 
diff --git a/Engine/Source/Runtime/Renderer/Private/BasePassRendering.h b/Engine/Source/Runtime/Renderer/Private/BasePassRendering.h
--- a/Engine/Source/Runtime/Renderer/Private/BasePassRendering.h
+++ b/Engine/Source/Runtime/Renderer/Private/BasePassRendering.h
@@ -15,6 +15,9 @@
 /** Whether to allow the indirect lighting cache to be applied to dynamic objects. */
 extern int32 GIndirectLightingCache;
 
+/** Whether some GBuffer targets are optionnal. */
+bool UseSelectiveBasePassOutputs();
+
 /**
  * The base shader type for vertex shaders that render the emissive color, and light-mapped/ambient lighting of a mesh.
  * The base type is shared between the versions with and without atmospheric fog.
diff --git a/Engine/Source/Runtime/Renderer/Private/DeferredShadingRenderer.cpp b/Engine/Source/Runtime/Renderer/Private/DeferredShadingRenderer.cpp
--- a/Engine/Source/Runtime/Renderer/Private/DeferredShadingRenderer.cpp
+++ b/Engine/Source/Runtime/Renderer/Private/DeferredShadingRenderer.cpp
@@ -972,6 +972,7 @@ void FDeferredShadingSceneRenderer::Render(FRHICommandListImmediate& RHICmdList)
 	ServiceLocalQueue();	
 	const bool bShouldRenderVelocities = ShouldRenderVelocities();
 	const bool bUseVelocityGBuffer = FVelocityRendering::OutputsToGBuffer();
+	const bool bUseSelectiveBasePassOutputs = UseSelectiveBasePassOutputs();
 
 	{
 		static bool bOnce = false;
@@ -989,7 +990,7 @@ void FDeferredShadingSceneRenderer::Render(FRHICommandListImmediate& RHICmdList)
 
 	{
 		QUICK_SCOPE_CYCLE_COUNTER(STAT_FDeferredShadingSceneRenderer_AllocGBufferTargets);
-		SceneContext.PreallocGBufferTargets(bShouldRenderVelocities && bUseVelocityGBuffer);
+		SceneContext.PreallocGBufferTargets(bUseVelocityGBuffer); // Even if !bShouldRenderVelocities, the velocity buffer must be bound because it's a compile time option for the shader.
 		SceneContext.AllocGBufferTargets(RHICmdList);
 	}
 	
@@ -1123,7 +1124,8 @@ void FDeferredShadingSceneRenderer::Render(FRHICommandListImmediate& RHICmdList)
 	{
 		VelocityRT = SceneContext.GetGBufferVelocityRT();
 	}
-	else if (bShouldRenderVelocities)
+	
+	if (bShouldRenderVelocities && (!bUseVelocityGBuffer || bUseSelectiveBasePassOutputs))
 	{
 		// Render the velocities of movable objects
 		RHICmdList.SetCurrentStat(GET_STATID(STAT_CLM_Velocity));
@@ -1707,10 +1709,10 @@ bool FDeferredShadingSceneRenderer::RenderBasePass(FRHICommandListImmediate& RHI
 				FViewInfo& View = Views[ViewIndex];
 				RenderBasePassViewParallel(View, RHICmdList);
 			}
-				bDirty = true; // assume dirty since we are not going to wait
-				if (FVelocityRendering::OutputsToGBuffer())
-				{
-					GPrevPerBoneMotionBlur.EndAppendFence(RHICmdList);
+			bDirty = true; // assume dirty since we are not going to wait
+			if (FVelocityRendering::OutputsToGBuffer())
+			{
+				GPrevPerBoneMotionBlur.EndAppendFence(RHICmdList);
 			}
 		}
 		else
diff --git a/Engine/Source/Runtime/Renderer/Private/PostProcess/SceneRenderTargets.cpp b/Engine/Source/Runtime/Renderer/Private/PostProcess/SceneRenderTargets.cpp
--- a/Engine/Source/Runtime/Renderer/Private/PostProcess/SceneRenderTargets.cpp
+++ b/Engine/Source/Runtime/Renderer/Private/PostProcess/SceneRenderTargets.cpp
@@ -456,32 +456,33 @@ void FSceneRenderTargets::BeginRenderingSceneColor(FRHICommandList& RHICmdList,
 
 int32 FSceneRenderTargets::GetGBufferRenderTargets(ERenderTargetLoadAction ColorLoadAction, FRHIRenderTargetView OutRenderTargets[MaxSimultaneousRenderTargets], int32& OutVelocityRTIndex)
 {
-	OutRenderTargets[0] = FRHIRenderTargetView(GetSceneColorSurface(), 0, -1, ColorLoadAction, ERenderTargetStoreAction::EStore);
-	OutRenderTargets[1] = FRHIRenderTargetView(GBufferA->GetRenderTargetItem().TargetableTexture, 0, -1, ColorLoadAction, ERenderTargetStoreAction::EStore);
-	OutRenderTargets[2] = FRHIRenderTargetView(GBufferB->GetRenderTargetItem().TargetableTexture, 0, -1, ColorLoadAction, ERenderTargetStoreAction::EStore);
-	OutRenderTargets[3] = FRHIRenderTargetView(GBufferC->GetRenderTargetItem().TargetableTexture, 0, -1, ColorLoadAction, ERenderTargetStoreAction::EStore);
-	OutRenderTargets[4] = FRHIRenderTargetView(GBufferD->GetRenderTargetItem().TargetableTexture, 0, -1, ColorLoadAction, ERenderTargetStoreAction::EStore);
-	int32 MRTCount = 5;	// Derived from previously used RenderTargets[]; would have been nice to keep a pointer while filling it but it's more confusing to use
+	int32 MRTCount = 0;
+	OutRenderTargets[MRTCount++] = FRHIRenderTargetView(GetSceneColorSurface(), 0, -1, ColorLoadAction, ERenderTargetStoreAction::EStore);
+	OutRenderTargets[MRTCount++] = FRHIRenderTargetView(GBufferA->GetRenderTargetItem().TargetableTexture, 0, -1, ColorLoadAction, ERenderTargetStoreAction::EStore);
+	OutRenderTargets[MRTCount++] = FRHIRenderTargetView(GBufferB->GetRenderTargetItem().TargetableTexture, 0, -1, ColorLoadAction, ERenderTargetStoreAction::EStore);
+	OutRenderTargets[MRTCount++] = FRHIRenderTargetView(GBufferC->GetRenderTargetItem().TargetableTexture, 0, -1, ColorLoadAction, ERenderTargetStoreAction::EStore);
 
-	OutVelocityRTIndex = -1;
-
-	if (bAllowStaticLighting)
+	if (bAllocateVelocityGBuffer)
+	{
+		OutVelocityRTIndex = MRTCount;
+		// check(OutVelocityRTIndex == GetGBufferVelocityIndex());
+		OutRenderTargets[MRTCount++] = FRHIRenderTargetView(GBufferVelocity->GetRenderTargetItem().TargetableTexture, 0, -1, ColorLoadAction, ERenderTargetStoreAction::EStore);
+	}
+	else
 	{
-		check(MRTCount == GetGBufferEIndex());
-		OutRenderTargets[MRTCount] = FRHIRenderTargetView(GBufferE->GetRenderTargetItem().TargetableTexture, 0, -1, ColorLoadAction, ERenderTargetStoreAction::EStore);
-		++MRTCount;
+		OutVelocityRTIndex = -1;
 	}
 
-	if (bAllocateVelocityGBuffer)
+	OutRenderTargets[MRTCount++] = FRHIRenderTargetView(GBufferD->GetRenderTargetItem().TargetableTexture, 0, -1, ColorLoadAction, ERenderTargetStoreAction::EStore);
+
+
+	if (bAllowStaticLighting)
 	{
-		OutVelocityRTIndex = MRTCount;
-		++MRTCount;
-		check(OutVelocityRTIndex == GetGBufferVelocityIndex());
-		OutRenderTargets[OutVelocityRTIndex] = FRHIRenderTargetView(GBufferVelocity->GetRenderTargetItem().TargetableTexture, 0, -1, ColorLoadAction, ERenderTargetStoreAction::EStore);
+		// check(MRTCount == GetGBufferEIndex());
+		OutRenderTargets[MRTCount++] = FRHIRenderTargetView(GBufferE->GetRenderTargetItem().TargetableTexture, 0, -1, ColorLoadAction, ERenderTargetStoreAction::EStore);
 	}
 
 	check(MRTCount <= MaxSimultaneousRenderTargets);
-
 	return MRTCount;
 }
 
diff --git a/Engine/Source/Runtime/Renderer/Private/PostProcess/SceneRenderTargets.h b/Engine/Source/Runtime/Renderer/Private/PostProcess/SceneRenderTargets.h
--- a/Engine/Source/Runtime/Renderer/Private/PostProcess/SceneRenderTargets.h
+++ b/Engine/Source/Runtime/Renderer/Private/PostProcess/SceneRenderTargets.h
@@ -392,21 +392,6 @@ public:
 
 	IPooledRenderTarget* GetGBufferVelocityRT();
 
-	int32 GetGBufferEIndex() const
-	{
-		return bAllowStaticLighting ? 5 : -1;
-	}
-
-	int32 GetGBufferVelocityIndex() const
-	{
-		if (bAllocateVelocityGBuffer)
-		{
-			return bAllowStaticLighting ? 6 : 5;
-		}
-
-		return -1;
-	}
-
 	// @return can be 0 if the feature is disabled
 	IPooledRenderTarget* RequestCustomDepth(FRHICommandListImmediate& RHICmdList, bool bPrimitives);
 
diff --git a/Engine/Source/Runtime/Renderer/Private/PrimitiveSceneInfo.cpp b/Engine/Source/Runtime/Renderer/Private/PrimitiveSceneInfo.cpp
--- a/Engine/Source/Runtime/Renderer/Private/PrimitiveSceneInfo.cpp
+++ b/Engine/Source/Runtime/Renderer/Private/PrimitiveSceneInfo.cpp
@@ -494,6 +494,12 @@ bool FPrimitiveSceneInfo::ShouldRenderVelocity(const FViewInfo& View, bool bChec
 		return false;
 	}
 
+	// If the base pass is allowed to render velocity in the GBuffer, only mesh with static lighting need the velocity pass.
+	if (FVelocityRendering::OutputsToGBuffer() && (!UseSelectiveBasePassOutputs() || !Proxy->HasStaticLighting()))
+	{
+		return false;
+	}
+
 	return true;
 }
 
diff --git a/Engine/Source/Runtime/Renderer/Private/VelocityRendering.cpp b/Engine/Source/Runtime/Renderer/Private/VelocityRendering.cpp
--- a/Engine/Source/Runtime/Renderer/Private/VelocityRendering.cpp
+++ b/Engine/Source/Runtime/Renderer/Private/VelocityRendering.cpp
@@ -67,7 +67,8 @@ public:
 			|| (Material->IsTwoSided() && !IsTranslucentBlendMode(Material->GetBlendMode()))
 			// or if the material modifies meshes
 			|| Material->MaterialMayModifyMeshPosition()))
-			&& IsFeatureLevelSupported(Platform, ERHIFeatureLevel::SM4) && !FVelocityRendering::OutputsToGBuffer();
+			&& IsFeatureLevelSupported(Platform, ERHIFeatureLevel::SM4) 
+			&& !FVelocityRendering::OutputsOnlyToGBuffer(VertexFactoryType->SupportsStaticLighting());
 	}
 
 protected:
@@ -155,7 +156,8 @@ public:
 			|| (Material->IsTwoSided() && !IsTranslucentBlendMode(Material->GetBlendMode()))
 			// or if the material modifies meshes
 			|| Material->MaterialMayModifyMeshPosition()))
-			&& IsFeatureLevelSupported(Platform, ERHIFeatureLevel::SM4) && !FVelocityRendering::OutputsToGBuffer();
+			&& IsFeatureLevelSupported(Platform, ERHIFeatureLevel::SM4) && 
+			!FVelocityRendering::OutputsOnlyToGBuffer(VertexFactoryType->SupportsStaticLighting());
 	}
 
 	static void ModifyCompilationEnvironment( EShaderPlatform Platform, const FMaterial* Material, FShaderCompilerEnvironment& OutEnvironment )
@@ -411,8 +413,11 @@ void FVelocityDrawingPolicyFactory::AddStaticMesh(FScene* Scene, FStaticMesh* St
 	const FMaterialRenderProxy* MaterialRenderProxy = StaticMesh->MaterialRenderProxy;
 	const FMaterial* Material = MaterialRenderProxy->GetMaterial(FeatureLevel);
 
+	// When selective outputs are enable, only primitive with no static lighting output velocity in GBuffer.
+	const bool bVelocityInGBuffer = FVelocityRendering::OutputsToGBuffer() && (!UseSelectiveBasePassOutputs() || !StaticMesh->PrimitiveSceneInfo->Proxy->HasStaticLighting());
+
 	// Velocity only needs to be directly rendered for movable meshes
-	if (StaticMesh->PrimitiveSceneInfo->Proxy->IsMovable())
+	if (StaticMesh->PrimitiveSceneInfo->Proxy->IsMovable() && !bVelocityInGBuffer)
 	{
 	    EBlendMode BlendMode = Material->GetBlendMode();
 	    if(BlendMode == BLEND_Opaque || BlendMode == BLEND_Masked)
@@ -749,9 +754,18 @@ void FDeferredShadingSceneRenderer::RenderVelocities(FRHICommandListImmediate& R
 
 	{
 		FSceneRenderTargets::Get(RHICmdList).SetVelocityPass(true);
-		GPrevPerBoneMotionBlur.StartAppend(RHICmdList, ViewFamily.bWorldIsPaused);
 
-		BeginVelocityRendering(RHICmdList, VelocityRT, true);
+
+		if (FVelocityRendering::OutputsToGBuffer() && UseSelectiveBasePassOutputs())
+		{
+			// In this case, basepass also outputs some of the velocities, so append is already started, and don't clear the buffer.
+			BeginVelocityRendering(RHICmdList, VelocityRT, false);
+		}
+		else
+		{
+			GPrevPerBoneMotionBlur.StartAppend(RHICmdList, ViewFamily.bWorldIsPaused);
+			BeginVelocityRendering(RHICmdList, VelocityRT, true);
+		}
 
 		if (IsParallelVelocity())
 		{
@@ -786,3 +800,11 @@ bool FVelocityRendering::OutputsToGBuffer()
 {
 	return CVarBasePassOutputsVelocity.GetValueOnAnyThread() == 1;
 }
+
+bool FVelocityRendering::OutputsOnlyToGBuffer(bool bSupportsStaticLighting)
+{
+	// With selective outputs, only primitive that have static lighting are rendered in the velocity pass.
+	// If the vertex factory does not support static lighting, then it must be rendered in the velocity pass.
+	return CVarBasePassOutputsVelocity.GetValueOnAnyThread() == 1 &&
+		   (!UseSelectiveBasePassOutputs() || !bSupportsStaticLighting);
+}
diff --git a/Engine/Source/Runtime/Renderer/Private/VelocityRendering.h b/Engine/Source/Runtime/Renderer/Private/VelocityRendering.h
--- a/Engine/Source/Runtime/Renderer/Private/VelocityRendering.h
+++ b/Engine/Source/Runtime/Renderer/Private/VelocityRendering.h
@@ -89,4 +89,5 @@ struct FVelocityRendering
 	static FPooledRenderTargetDesc GetRenderTargetDesc();
 
 	static bool OutputsToGBuffer();
+	static bool OutputsOnlyToGBuffer(bool bSupportsStaticLighting);
 };
diff --git a/Engine/Source/Runtime/ShaderCore/Private/Shader.cpp b/Engine/Source/Runtime/ShaderCore/Private/Shader.cpp
--- a/Engine/Source/Runtime/ShaderCore/Private/Shader.cpp
+++ b/Engine/Source/Runtime/ShaderCore/Private/Shader.cpp
@@ -1125,6 +1125,14 @@ void ShaderMapAppendKeyString(EShaderPlatform Platform, FString& KeyString)
 		}
 	}
 
+	{
+		static const auto CVar = IConsoleManager::Get().FindTConsoleVariableDataInt(TEXT("r.SelectiveBasePassOutputs"));
+		if (CVar && CVar->GetValueOnGameThread() != 0)
+		{
+			KeyString += TEXT("_SO");
+		}
+	}
+
 	{
 		static const auto CVar = IConsoleManager::Get().FindTConsoleVariableDataInt(TEXT("r.GBuffer"));
 		if (CVar ? CVar->GetValueOnAnyThread() == 0 : false)
