Merging using UE4_to_UE4-Orion

2702788 Fix tint color on brushes not applied to materials in Slate.

diff --git a/Engine/Shaders/SlateElementPixelShader.usf b/Engine/Shaders/SlateElementPixelShader.usf
--- a/Engine/Shaders/SlateElementPixelShader.usf
+++ b/Engine/Shaders/SlateElementPixelShader.usf
@@ -124,9 +124,16 @@ float4 GetColor( VertexToPixelInterpolants InVertex, float2 UV )
 	float4 FinalColor;
 	
 #if USE_MATERIALS
-	FinalColor = GetMaterialColor( InVertex, UV );
+	float4 BaseColor = GetMaterialColor( InVertex, UV );
 #else
-	FinalColor = GetTextureColor( InVertex, UV )*InVertex.Color;
+	float4 BaseColor = GetTextureColor( InVertex, UV );
+#endif
+
+#if SHADER_TYPE == ST_Custom
+	// Make no assumptionn of what vertex color does in a custom material shader
+	FinalColor = BaseColor;
+#else
+	FinalColor = BaseColor*InVertex.Color;
 #endif
 
 	return FinalColor;
