Back out of bad change to GPU partilce collision.

diff --git a/Engine/Shaders/ParticleSimulationShader.usf b/Engine/Shaders/ParticleSimulationShader.usf
--- a/Engine/Shaders/ParticleSimulationShader.usf
+++ b/Engine/Shaders/ParticleSimulationShader.usf
@@ -245,8 +245,7 @@ void CollideWithDepthBuffer(
 	in float3 InVelocity,
 	in float3 Acceleration,
 	in float CollisionRadius,
-	in float Resilience,
-	inout float CollisionStop
+	in float Resilience
 	)
 {
 	// Integration assuming no collision.
@@ -293,14 +292,6 @@ void CollideWithDepthBuffer(
 			//if (t_back >= 0 && t_front <= 1 && DistanceToPlane >= 0)
 			if (step(0, t_back) * step(t_front, 1) * step(0, DistanceToPlane))
 			{
-				//When the collision response is "Stop" the resilience is 0.0.
-				//In this case we mark the event so that all future updates of the particle motion are stopped.
-				//Prevents Stopped particles from continuing to move if they go off screen.
-				if (Resilience < 0.0001f)
-				{
-					CollisionStop = 1.0f;
-				}
-
 				// Separate velocity in to the components perpendicular and tangent to the collision plane.
 				float3 PerpVelocity = dot(MidVelocity,CollisionPlane.xyz) * CollisionPlane.xyz;
 				float3 TanVelocity = MidVelocity - PerpVelocity;
@@ -344,8 +335,7 @@ void CollideWithDistanceField(
 	in float3 InVelocity,
 	in float3 Acceleration,
 	in float CollisionRadius,
-	in float Resilience,
-	out float CollisionStop
+	in float Resilience
 	)
 {
 	// Integration assuming no collision.
@@ -378,14 +368,6 @@ void CollideWithDistanceField(
 		//if (t_back >= 0 && t_front <= 1 && DistanceToPlane >= 0)
 		if (step(0, t_back) * step(t_front, 1) * step(0, DistanceToPlane))
 		{
-			//When the collision response is "Stop" the resilience is 0.0.
-			//In this case we mark the event so that all future updates of the particle motion are stopped.
-			//Prevents Stopped particles from continuing to move if they go off screen.
-			if (Resilience < 0.0001f)
-			{
-				CollisionStop = 1.0f;
-			}
-
 			// Separate velocity in to the components perpendicular and tangent to the collision plane.
 			float3 PerpVelocity = dot(MidVelocity,CollisionPlane.xyz) * CollisionPlane.xyz;
 			float3 TanVelocity = MidVelocity - PerpVelocity;
@@ -416,17 +398,6 @@ void CollideWithDistanceField(
 
 #endif
 
-float PackRelativeTimeAndCollisionStop(in float RelativeTime, in float CollisionStop)
-{
-	return CollisionStop > 0.0f ? -RelativeTime : RelativeTime;
-}
-
-float UnpackRelativeTimeAndCollisionStop(in float PackedValue, out float CollisionStop)
-{
-	CollisionStop = PackedValue < 0.0f ? 1.0f : 0.0f;
-	return abs(PackedValue);
-}
-
 void PixelMain(
 	in FShaderInterpolants Interpolants,
 	out float4 OutPosition : SV_Target0,
@@ -448,14 +419,7 @@ void PixelMain(
 
 	// Position.w holds the relative time of the particle.
 	float3 Position = PositionSample.xyz;
-
-	float CollisionStop;
-	float RelativeTime = UnpackRelativeTimeAndCollisionStop(PositionSample.w, CollisionStop);
-
-	RelativeTime += DeltaSeconds * TimeScale;
-
-	//Everything other that particle aging should stop if we have collision stop.
-	float CollisionAdjustedDeltaSeconds = CollisionStop > 0.0f ? 0.0f : DeltaSeconds;
+	float RelativeTime = PositionSample.w + DeltaSeconds * TimeScale;
 
 	// Sample the attribute curve.
 	const float2 AttributeCurveTexCoord = Simulation.AttributeCurve.xy +
@@ -489,7 +453,7 @@ void PixelMain(
 	Force += ComputeAttractionForce(Position.xyz);
 
 	// Compute the acceleration to apply to the particle this frame.
-	float3 Acceleration = Force * CollisionAdjustedDeltaSeconds;
+	float3 Acceleration = Force * DeltaSeconds;
 
 #if DEPTH_BUFFER_COLLISION || DISTANCE_FIELD_COLLISION
 	// We need to look up render attributes for this particle to figure out how big it is.
@@ -522,8 +486,7 @@ void PixelMain(
 		Velocity.xyz,
 		Acceleration,
 		CollisionRadius,
-		Resilience,
-		CollisionStop
+		Resilience
 		);
 #elif DISTANCE_FIELD_COLLISION
 
@@ -537,8 +500,7 @@ void PixelMain(
 		Velocity.xyz,
 		Acceleration,
 		CollisionRadius,
-		Resilience,
-		CollisionStop
+		Resilience
 		);
 
 #else 
@@ -554,7 +516,7 @@ void PixelMain(
 
 	// Store the new position, time, and velocity for the particle.
 	OutPosition.xyz = NewPosition + OrbitVelocity * DeltaSeconds + PositionOffsetAndAttractorStrength.xyz;
-	OutPosition.w = PackRelativeTimeAndCollisionStop(RelativeTime, CollisionStop);
+	OutPosition.w = RelativeTime;
 	OutVelocity.xyz = NewVelocity;
 	OutVelocity.w = TimeScale;
 }
