Fixed skylights on translucency in projects with static lighting disabled

diff --git a/Engine/Shaders/BasePassPixelShader.usf b/Engine/Shaders/BasePassPixelShader.usf
--- a/Engine/Shaders/BasePassPixelShader.usf
+++ b/Engine/Shaders/BasePassPixelShader.usf
@@ -514,7 +514,7 @@ void GetPrecomputedIndirectLightingAndSkyLight(
 	OutSubsurfaceLighting += SkySubsurfaceLighting;
 
 	BRANCH
-	if( View.UseLightmaps > 0 )
+	if( View.UseLightmaps > 0 || !ALLOW_STATIC_LIGHTING)
 	{
 		// Sky lighting must contribute to IndirectIrradiance for ReflectionEnvironment lightmap mixing
 		OutDiffuseLighting += SkyDiffuseLighting;
