Merging
Using SvPosition instead of ScreenPosition for Decal projection for better quality and performance, also used on mobile but there is some extra cost we could get rid of

diff --git a/Engine/Shaders/Common.usf b/Engine/Shaders/Common.usf
--- a/Engine/Shaders/Common.usf
+++ b/Engine/Shaders/Common.usf
@@ -927,29 +927,11 @@ float ConvertToDeviceZ(float SceneDepth)
 	return 1.0f / ((SceneDepth + View.InvDeviceZToWorldZTransform[3]) * View.InvDeviceZToWorldZTransform[2]);
 }
 
-/** Returns clip space W, which is world space distance along the View Z axis. */
-float LookupDeviceZ( float2 ScreenUV )
-{
-#if FEATURE_LEVEL > FEATURE_LEVEL_ES3_1
-	return Texture2DSampleLevel(SceneDepthTexture, SceneDepthTextureSampler, ScreenUV, 0).r;
-#else
-	#if COMPILER_GLSL_ES2
-		// not yet supported
-		return 0;
-	#elif METAL_PROFILE
-		// not yet supported
-		return 0;
-	#else
-		return Texture2DSampleLevel(SceneDepthTexture, SceneDepthTextureSampler, ScreenUV, 0).r;
-	#endif
-#endif
-}
-
-/** Returns clip space W, which is world space distance along the View Z axis. */
+/** Returns clip space W, which is world space distance along the View Z axis. Note if you need DeviceZ LookupDeviceZ() is the faster option */
 float CalcSceneDepth( float2 ScreenUV )
 {
 #if FEATURE_LEVEL > FEATURE_LEVEL_ES3_1
-	return ConvertFromDeviceZ(LookupDeviceZ(ScreenUV));
+	return ConvertFromDeviceZ(Texture2DSampleLevel(SceneDepthTexture, SceneDepthTextureSampler, ScreenUV, 0).r);
 #else
 	#if COMPILER_GLSL_ES2
 		#if IOS
@@ -964,7 +946,27 @@ float CalcSceneDepth( float2 ScreenUV )
 	#elif METAL_PROFILE
 			return FramebufferFetchES2().w;
 	#else
-		return ConvertFromDeviceZ(LookupDeviceZ(ScreenUV));
+		return ConvertFromDeviceZ(Texture2DSampleLevel(SceneDepthTexture, SceneDepthTextureSampler, ScreenUV, 0).r);
+	#endif
+#endif
+}
+
+/** Returns DeviceZ which is the z value storesd in the depth buffer. Perfer this over CalcSceneDepth() on most hardware. */
+float LookupDeviceZ( float2 ScreenUV )
+{
+#if FEATURE_LEVEL > FEATURE_LEVEL_ES3_1
+	// native Depth buffer lookup
+	return Texture2DSampleLevel(SceneDepthTexture, SceneDepthTextureSampler, ScreenUV, 0).r;
+#else
+	#if COMPILER_GLSL_ES2
+		// todo: can be optimized
+		return ConvertToDeviceZ(CalcSceneDepth(ScreenUV));
+	#elif METAL_PROFILE
+		// todo: can be optimized
+		return ConvertToDeviceZ(CalcSceneDepth(ScreenUV));
+	#else
+		// native Depth buffer lookup
+		return Texture2DSampleLevel(SceneDepthTexture, SceneDepthTextureSampler, ScreenUV, 0).r;
 	#endif
 #endif
 }
diff --git a/Engine/Shaders/DeferredDecal.usf b/Engine/Shaders/DeferredDecal.usf
--- a/Engine/Shaders/DeferredDecal.usf
+++ b/Engine/Shaders/DeferredDecal.usf
@@ -10,8 +10,8 @@
 
 #include "Common.usf"
 
-// from screen space position to decal space (for position)
-float4x4 ScreenToDecal;
+// from SvPosition to decal space (for position)
+float4x4 SvPositionToDecal;
 // to transform position from decal space to world space, for normals use transpose of inverse
 float4x4 DecalToWorld;
 // to transform position from world space to decal space, for normals use transpose of inverse
@@ -190,7 +190,6 @@ void MainPS(
 
 	// make SvPosition appear to be rasterized with the depth from the depth buffer
 	SvPosition.z = LookupDeviceZ(ScreenUV);
-//	SvPosition.w = 1;
 
 	float4 ScreenPosition = SvPositionToScreenPosition(SvPosition);
 
@@ -245,12 +244,7 @@ void MainPS(
 	}
 	else
 	{
-		// can be optimized further
-
-		half SceneW = CalcSceneDepth(SvPositionToBufferUV(SvPosition));
-		float4 ProjectedScreenPosition = float4(ScreenPosition.xy / ScreenPosition.w * SceneW, SceneW, 1);	
-
-		float4 DecalVectorHom = mul(ProjectedScreenPosition, ScreenToDecal);
+		float4 DecalVectorHom = mul(float4(SvPosition.xyz,1), SvPositionToDecal);
 		OSPosition = DecalVectorHom.xyz / DecalVectorHom.w;
 
 		// clip content outside the decal
@@ -258,45 +252,9 @@ void MainPS(
 		clip(OSPosition.xyz + 1.0f);
 		clip(1.0f - OSPosition.xyz);
 
-
-
-		// can be optimized
+		float4 Hom = mul(float4(SvPosition.xyz, 1), View.SVPositionToTranslatedWorld);
 		
-//		SvPosition.z = CalcSceneDepth(ScreenUV);
-//		SvPosition.w *= SvPosition.z;
-//		SvPosition.z = LookupDeviceZ(ScreenUV);
-//		SvPosition.xy /= SvPosition.w;
-	
-		float4 Clip;
-
-	//	ProjectedScreenPosition = ComputeScreenPosition(SvPosition);
-		Clip.x = ((SvPosition.x-View.ViewRectMin.x) * View.ViewSizeAndInvSize.z) * 2 - 1;
-//		ProjectedScreenPosition.x = SvPosition.x * View.BufferSizeAndInvSize.z;
-		Clip.y = ((SvPosition.y-View.ViewRectMin.y) * View.ViewSizeAndInvSize.w) * -2 + 1;
-//		Clip.z = LookupDeviceZ(ScreenUV);
-
-// test
-		Clip.z = SvPosition.z;
-
-
-//		ProjectedScreenPosition.z = CalcSceneDepth(ScreenUV);
-//		ProjectedScreenPosition.xy*=ProjectedScreenPosition.z;
-//		ProjectedScreenPosition.w = 1;
-
-		Clip.w = 1;
-
-		SvPosition.z = LookupDeviceZ(ScreenUV);
-
-
-//		float4 Hom = mul(Clip, View.ClipToTranslatedWorld);WSPosition = Hom.xyz / Hom.w - View.PreViewTranslation;		// float4(-1..1,1..-1,LookupDeviceZ(ScreenUV), 1)
-//		float4 Hom = mul(ProjectedScreenPosition, View.ScreenToWorld);WSPosition = Hom.xyz / Hom.w;				// float4(-1..1*z,1..-1*z, CalcSceneDepth(ScreenUV),1)
-//		float4 Hom = mul(float4(SvPosition.xyz - float3(View.ViewSizeAndInvSize.xy * 0.5f, 0),1), View.SVPositionToTranslatedWorld);WSPosition = Hom.xyz / Hom.w - View.PreViewTranslation;				// float4(PixelX,PixelY,DeviceZ,W)
-		float4 Hom = mul(float4(SvPosition.xyz, 1), View.SVPositionToTranslatedWorld);WSPosition = Hom.xyz / Hom.w - View.PreViewTranslation;			// float4(PixelX,PixelY,DeviceZ,W)  perfect
-//		float4 Hom = mul(SvPositionToScreenPosition2(SvPosition), View.ScreenToWorld);WSPosition = Hom.xyz / Hom.w;	// works but why
-//		float4 Hom = mul(SvPositionToScreenPosition(SvPosition), View.ScreenToWorld);WSPosition = Hom.xyz / Hom.w;	// to verify SvPositionToScreenPosition, wrong?
-
-//		WSPosition = Hom.xyz / Hom.w - View.PreViewTranslation;
-//		WSPosition  = SvPosition.xyz / SvPosition.w;
+		WSPosition = Hom.xyz / Hom.w - View.PreViewTranslation;
 	}
 
 	float3 CameraVector = normalize(View.WorldCameraOrigin - WSPosition);
diff --git a/Engine/Shaders/MaterialTemplate.usf b/Engine/Shaders/MaterialTemplate.usf
--- a/Engine/Shaders/MaterialTemplate.usf
+++ b/Engine/Shaders/MaterialTemplate.usf
@@ -365,10 +365,13 @@ float4 GetScreenPosition(FMaterialPixelParameters Parameters)
 	 * mesh of the same group of pixel, then one of the delta might be very big, leading to choosing a low mipmap level for this
 	 * group of 4 pixels, causing the artefacts.
 	 */
-	float2 ComputeDecalUVFromScreenPosition(float4 ScreenPosition)
+	float2 ComputeDecalUVFromSvPosition(float4 SvPosition)
 	{
-		half SceneW = CalcSceneDepth(ScreenPositionToBufferUV(ScreenPosition));
-		float4 DecalVector = mul(float4(ScreenPosition.xy / ScreenPosition.w * SceneW, SceneW, 1), ScreenToDecal);
+		half DeviceZ = LookupDeviceZ(SvPositionToBufferUV(SvPosition));
+
+		SvPosition.z = DeviceZ;
+
+		float4 DecalVector = mul(float4(SvPosition.xyz,1), SvPositionToDecal);
 		DecalVector.xyz /= DecalVector.w;
 		DecalVector = DecalVector * 0.5f + 0.5f;
 		DecalVector.xyz = DecalVector.zyx;
@@ -381,9 +384,9 @@ float4 GetScreenPosition(FMaterialPixelParameters Parameters)
 		 * Assuming where in a pixel shader invocation, then we compute manualy compute two d(uv)/d(x)
 		 * with the pixels's left and right neighbours.
 		 */
-		float4 ScreenDeltaX = ddx(Parameters.ScreenPosition);
-		float2 UvDiffX0 = Parameters.TexCoords[0] - ComputeDecalUVFromScreenPosition(Parameters.ScreenPosition - ScreenDeltaX);
-		float2 UvDiffX1 = ComputeDecalUVFromScreenPosition(Parameters.ScreenPosition + ScreenDeltaX) - Parameters.TexCoords[0];
+		float4 ScreenDeltaX = float4(1, 0, 0, 0);
+		float2 UvDiffX0 = Parameters.TexCoords[0] - ComputeDecalUVFromSvPosition(Parameters.SvPosition - ScreenDeltaX);
+		float2 UvDiffX1 = ComputeDecalUVFromSvPosition(Parameters.SvPosition + ScreenDeltaX) - Parameters.TexCoords[0];
 
 		/*
 		 * So we have two diff on the X axis, we want the one that has the smallest length
@@ -395,9 +398,9 @@ float4 GetScreenPosition(FMaterialPixelParameters Parameters)
 	float2 ComputeDecalDDY(FMaterialPixelParameters Parameters)
 	{
 		// do same for the Y axis
-		float4 ScreenDeltaY = ddy(Parameters.ScreenPosition);
-		float2 UvDiffY0 = Parameters.TexCoords[0] - ComputeDecalUVFromScreenPosition(Parameters.ScreenPosition - ScreenDeltaY);
-		float2 UvDiffY1 = ComputeDecalUVFromScreenPosition(Parameters.ScreenPosition + ScreenDeltaY) - Parameters.TexCoords[0];
+		float4 ScreenDeltaY = float4(0, 1, 0, 0);
+		float2 UvDiffY0 = Parameters.TexCoords[0] - ComputeDecalUVFromSvPosition(Parameters.SvPosition - ScreenDeltaY);
+		float2 UvDiffY1 = ComputeDecalUVFromSvPosition(Parameters.SvPosition + ScreenDeltaY) - Parameters.TexCoords[0];
 
 		return dot(UvDiffY0, UvDiffY0) < dot(UvDiffY1, UvDiffY1) ? UvDiffY0 : UvDiffY1;
 	}
diff --git a/Engine/Source/Runtime/Renderer/Private/DecalRenderingShared.cpp b/Engine/Source/Runtime/Renderer/Private/DecalRenderingShared.cpp
--- a/Engine/Source/Runtime/Renderer/Private/DecalRenderingShared.cpp
+++ b/Engine/Source/Runtime/Renderer/Private/DecalRenderingShared.cpp
@@ -149,7 +149,7 @@ public:
 	FDeferredDecalPS(const ShaderMetaType::CompiledShaderInitializerType& Initializer):
 		FMaterialShader(Initializer)
 	{
-		ScreenToDecal.Bind(Initializer.ParameterMap,TEXT("ScreenToDecal"));
+		SvPositionToDecal.Bind(Initializer.ParameterMap,TEXT("SvPositionToDecal"));
 		DecalToWorld.Bind(Initializer.ParameterMap,TEXT("DecalToWorld"));
 		FadeAlpha.Bind(Initializer.ParameterMap, TEXT("FadeAlpha"));
 		WorldToDecal.Bind(Initializer.ParameterMap,TEXT("WorldToDecal"));
@@ -166,17 +166,31 @@ public:
 		FMatrix WorldToComponent = ComponentTrans.ToInverseMatrixWithScale();
 
 		// Set the transform from screen space to light space.
-		if(ScreenToDecal.IsBound())
+		if(SvPositionToDecal.IsBound())
 		{
-			const FMatrix ScreenToDecalValue = 
+			FVector2D InvViewSize = FVector2D(1.0f / View.ViewRect.Width(), 1.0f / View.ViewRect.Height());
+
+			// setup a matrix to transform float4(SVPosition.xyz,1) directly to Decal (quality, performance as we don't need to convert or use interpolator)
+
+			//	new_xy = (xy - ViewRectMin.xy) * ViewSizeAndInvSize.zw * float2(2,-2) + float2(-1, 1);
+
+			//  transformed into one MAD:  new_xy = xy * ViewSizeAndInvSize.zw * float2(2,-2)      +       (-ViewRectMin.xy) * ViewSizeAndInvSize.zw * float2(2,-2) + float2(-1, 1);
+
+			float Mx = 2.0f * InvViewSize.X;
+			float My = -2.0f * InvViewSize.Y;
+			float Ax = -1.0f - 2.0f * View.ViewRect.Min.X * InvViewSize.X;
+			float Ay = 1.0f + 2.0f * View.ViewRect.Min.Y * InvViewSize.Y;
+
+			// todo: we could use InvTranslatedViewProjectionMatrix and TranslatedWorldToComponent for better quality
+			const FMatrix SvPositionToDecalValue = 
 				FMatrix(
-					FPlane(1,0,0,0),
-					FPlane(0,1,0,0),
-					FPlane(0,0,View.ViewMatrices.ProjMatrix.M[2][2],1),
-					FPlane(0,0,View.ViewMatrices.ProjMatrix.M[3][2],0)
+					FPlane(Mx,  0,   0,  0),
+					FPlane( 0, My,   0,  0),
+					FPlane( 0,  0,   1,  0),
+					FPlane(Ax, Ay,   0,  1)
 				) * View.InvViewProjectionMatrix * WorldToComponent;
 
-			SetShaderValue(RHICmdList, ShaderRHI, ScreenToDecal, ScreenToDecalValue);
+			SetShaderValue(RHICmdList, ShaderRHI, SvPositionToDecal, SvPositionToDecalValue);
 		}
 
 		// Set the transform from light space to world space
@@ -194,12 +208,12 @@ public:
 	virtual bool Serialize(FArchive& Ar) override
 	{
 		bool bShaderHasOutdatedParameters = FMaterialShader::Serialize(Ar);
-		Ar << ScreenToDecal << DecalToWorld << WorldToDecal << FadeAlpha;
+		Ar << SvPositionToDecal << DecalToWorld << WorldToDecal << FadeAlpha;
 		return bShaderHasOutdatedParameters;
 	}
 
 private:
-	FShaderParameter ScreenToDecal;
+	FShaderParameter SvPositionToDecal;
 	FShaderParameter DecalToWorld;
 	FShaderParameter FadeAlpha;
 	FShaderParameter WorldToDecal;
