Merging
Use SvPosition for ShadowProjection for quality and performance

diff --git a/Engine/Shaders/BasePassForForwardShadingPixelShader.usf b/Engine/Shaders/BasePassForForwardShadingPixelShader.usf
--- a/Engine/Shaders/BasePassForForwardShadingPixelShader.usf
+++ b/Engine/Shaders/BasePassForForwardShadingPixelShader.usf
@@ -257,7 +257,7 @@ void Main(
 				if (i == (MAX_FORWARD_SHADOWCASCADES - 1) // ignore test for last shadow map.
 					|| MaterialParameters.ScreenPosition.w < View.DirectionalLightShadowDistances[i])
 				{
-					ShadowPosition = mul(float4(MaterialParameters.ScreenPosition.xyw, 1), View.DirectionalLightScreenToShadow[i]);
+					ShadowPosition = mul(float4(MaterialParameters.SvPosition.xyz, 1), View.DirectionalLightSvPositionToShadow[i]);
 					break; // position found.
 				}
 			}
diff --git a/Engine/Shaders/LightFunctionPixelShader.usf b/Engine/Shaders/LightFunctionPixelShader.usf
--- a/Engine/Shaders/LightFunctionPixelShader.usf
+++ b/Engine/Shaders/LightFunctionPixelShader.usf
@@ -8,6 +8,7 @@
 #include "Material.usf"
 #include "LightFunctionCommon.usf"
 
+// used like this: mul(float4(SvPosition.xyz, 1), SvPositionToLight);
 float4x4 SvPositionToLight;
 
 /** Fade distance in x, disabled brightness in y, output for preview shadows mask in z. */
diff --git a/Engine/Shaders/ShadowProjectionPixelShader.usf b/Engine/Shaders/ShadowProjectionPixelShader.usf
--- a/Engine/Shaders/ShadowProjectionPixelShader.usf
+++ b/Engine/Shaders/ShadowProjectionPixelShader.usf
@@ -28,7 +28,8 @@
 
 float ShadowFadeFraction;
 float ShadowSharpen;
-float4x4 ScreenToShadowMatrix;
+// used like this: mul(float4(SvPosition.xyz, 1), SvPositionToShadow);
+float4x4 SvPositionToShadow;
 // .x:DepthBias, y: MaxSubjectZ - MinSubjectZ
 float2 ProjectionDepthBiasParameters;
 
@@ -44,22 +45,29 @@ float4 ModulatedShadowColor;
 */
 EARLYDEPTHSTENCIL
 void Main(
-	in float4 SVPos : SV_POSITION,
+	in float4 SvPosition : SV_POSITION,
 	out float4 OutColor : SV_Target0
 	)
 {
-	float2 ScreenUV = float2( SVPos.xy * View.BufferSizeAndInvSize.zw );
-	float SceneW = CalcSceneDepth( ScreenUV );
-	float2 ScreenPosition = ( ScreenUV.xy - View.ScreenPositionScaleBias.wz ) / View.ScreenPositionScaleBias.xy;
-	float4 ShadowPosition = mul(float4(ScreenPosition.xy * SceneW,SceneW,1), ScreenToShadowMatrix);
+	float2 ScreenUV = SvPositionToBufferUV(SvPosition);
+
+	// make SvPosition appear to be rasterized with the depth from the depth buffer
+	SvPosition.z = LookupDeviceZ(ScreenUV);
+
+	float3 ShadowPosition;
+	{
+		float4 Hom = mul(float4(SvPosition.xyz, 1), SvPositionToShadow);
+
+		ShadowPosition = Hom.xyz / Hom.w;
+	}
 
-	ShadowPosition.xy /= ShadowPosition.w;
 	// Clamp pixel depth in light space for shadowing opaque, because areas of the shadow depth buffer that weren't rendered to will have been cleared to 1
 	// We want to force the shadow comparison to result in 'unshadowed' in that case, regardless of whether the pixel being shaded is in front or behind that plane
 	float LightSpacePixelDepthForOpaque = min(ShadowPosition.z, 0.99999f);
 	// Must not clamp for SSS shadowing, the subsurface gradient must continue past the far plane
 	float LightSpacePixelDepthForSSS = ShadowPosition.z;
 
+
 	float Shadow = 1;
 	float SSSTransmission = 1;
 	
@@ -172,16 +180,16 @@ float2 PointLightDepthBiasParameters;
 
 /** Pixel shader for projecting a one pass point light shadow from a cube map. */
 void MainOnePassPointLightPS(
-	in float4 SVPos : SV_POSITION,
+	in float4 SvPosition : SV_POSITION,
 	out float4 OutColor : SV_Target0
 	)
 {
+	float2 ScreenUV = SvPositionToBufferUV(SvPosition);
 
-	float2 ScreenUV = float2( SVPos.xy * View.BufferSizeAndInvSize.zw );
-	float SceneW = CalcSceneDepth( ScreenUV );
-	float2 ScreenPosition = ( ScreenUV.xy - View.ScreenPositionScaleBias.wz ) / View.ScreenPositionScaleBias.xy;
-	float4 HomogeneousWorldPosition = mul(float4(ScreenPosition.xy * SceneW, SceneW, 1), View.ScreenToWorld);
- 	float3 WorldPosition = HomogeneousWorldPosition.xyz / HomogeneousWorldPosition.w;
+	// make SvPosition appear to be rasterized with the depth from the depth buffer
+	SvPosition.z = LookupDeviceZ(ScreenUV);
+	
+ 	float3 WorldPosition = SvPositionToWorld(SvPosition);
 
  	float Shadow = CubemapHardwarePCF(WorldPosition, LightPositionAndInvRadius.xyz, LightPositionAndInvRadius.w, PointLightDepthBiasParameters.x);
 
diff --git a/Engine/Source/Runtime/Engine/Public/SceneView.h b/Engine/Source/Runtime/Engine/Public/SceneView.h
--- a/Engine/Source/Runtime/Engine/Public/SceneView.h
+++ b/Engine/Source/Runtime/Engine/Public/SceneView.h
@@ -404,7 +404,7 @@ BEGIN_UNIFORM_BUFFER_STRUCT_WITH_CONSTRUCTOR(FViewUniformShaderParameters,ENGINE
 	DECLARE_UNIFORM_BUFFER_STRUCT_MEMBER_EX(FVector,DirectionalLightDirection, EShaderPrecisionModifier::Half)
 	DECLARE_UNIFORM_BUFFER_STRUCT_MEMBER_EX(float, DirectionalLightShadowTransition, EShaderPrecisionModifier::Half)			
 	DECLARE_UNIFORM_BUFFER_STRUCT_MEMBER_EX(FVector4, DirectionalLightShadowSize, EShaderPrecisionModifier::Half)
-	DECLARE_UNIFORM_BUFFER_STRUCT_MEMBER_ARRAY(FMatrix, DirectionalLightScreenToShadow, [MAX_FORWARD_SHADOWCASCADES])
+	DECLARE_UNIFORM_BUFFER_STRUCT_MEMBER_ARRAY(FMatrix, DirectionalLightSvPositionToShadow, [MAX_FORWARD_SHADOWCASCADES])
 	DECLARE_UNIFORM_BUFFER_STRUCT_MEMBER_EX(FVector4, DirectionalLightShadowDistances, EShaderPrecisionModifier::Half)
 	DECLARE_UNIFORM_BUFFER_STRUCT_MEMBER_EX(FLinearColor,UpperSkyColor, EShaderPrecisionModifier::Half)
 	DECLARE_UNIFORM_BUFFER_STRUCT_MEMBER_EX(FLinearColor,LowerSkyColor, EShaderPrecisionModifier::Half)
diff --git a/Engine/Source/Runtime/Renderer/Private/DecalRenderingShared.cpp b/Engine/Source/Runtime/Renderer/Private/DecalRenderingShared.cpp
--- a/Engine/Source/Runtime/Renderer/Private/DecalRenderingShared.cpp
+++ b/Engine/Source/Runtime/Renderer/Private/DecalRenderingShared.cpp
@@ -170,7 +170,7 @@ public:
 		{
 			FVector2D InvViewSize = FVector2D(1.0f / View.ViewRect.Width(), 1.0f / View.ViewRect.Height());
 
-			// setup a matrix to transform float4(SVPosition.xyz,1) directly to Decal (quality, performance as we don't need to convert or use interpolator)
+			// setup a matrix to transform float4(SvPosition.xyz,1) directly to Decal (quality, performance as we don't need to convert or use interpolator)
 
 			//	new_xy = (xy - ViewRectMin.xy) * ViewSizeAndInvSize.zw * float2(2,-2) + float2(-1, 1);
 
diff --git a/Engine/Source/Runtime/Renderer/Private/LightFunctionRendering.cpp b/Engine/Source/Runtime/Renderer/Private/LightFunctionRendering.cpp
--- a/Engine/Source/Runtime/Renderer/Private/LightFunctionRendering.cpp
+++ b/Engine/Source/Runtime/Renderer/Private/LightFunctionRendering.cpp
@@ -110,7 +110,7 @@ public:
 
 			FVector2D InvViewSize = FVector2D(1.0f / View.ViewRect.Width(), 1.0f / View.ViewRect.Height());
 
-			// setup a matrix to transform float4(SVPosition.xyz,1) directly to Light (quality, performance as we don't need to convert or use interpolator)
+			// setup a matrix to transform float4(SvPosition.xyz,1) directly to Light (quality, performance as we don't need to convert or use interpolator)
 
 			//	new_xy = (xy - ViewRectMin.xy) * ViewSizeAndInvSize.zw * float2(2,-2) + float2(-1, 1);
 
diff --git a/Engine/Source/Runtime/Renderer/Private/SceneRendering.cpp b/Engine/Source/Runtime/Renderer/Private/SceneRendering.cpp
--- a/Engine/Source/Runtime/Renderer/Private/SceneRendering.cpp
+++ b/Engine/Source/Runtime/Renderer/Private/SceneRendering.cpp
@@ -526,7 +526,7 @@ TUniformBufferRef<FViewUniformShaderParameters> FViewInfo::CreateUniformBuffer(
 	ViewUniformShaderParameters.PrevInvViewProj = PrevViewProjMatrix.Inverse();
 
 	{
-		// setup a matrix to transform float4(SVPosition.xyz,1) directly to TranslatedWorld (quality, performance as we don't need to convert or use interpolator)
+		// setup a matrix to transform float4(SvPosition.xyz,1) directly to TranslatedWorld (quality, performance as we don't need to convert or use interpolator)
 
 		//	new_xy = (xy - ViewRectMin.xy) * ViewSizeAndInvSize.zw * float2(2,-2) + float2(-1, 1);
 
@@ -540,10 +540,10 @@ TUniformBufferRef<FViewUniformShaderParameters> FViewInfo::CreateUniformBuffer(
 		// http://stackoverflow.com/questions/9010546/java-transformation-matrix-operations
 
 		ViewUniformShaderParameters.SVPositionToTranslatedWorld = 
-			FMatrix(FPlane(Mx,  0,  0,   0),
-					FPlane(0,  My,  0,   0),
-					FPlane(0,   0,  1,   0),
-					FPlane(Ax, Ay,  0,   1)) * ViewMatrices.InvTranslatedViewProjectionMatrix;
+			FMatrix(FPlane(Mx,   0,  0,   0),
+					FPlane( 0,  My,  0,   0),
+					FPlane( 0,   0,  1,   0),
+					FPlane(Ax,  Ay,  0,   1)) * ViewMatrices.InvTranslatedViewProjectionMatrix;
 	}
 
 	ViewUniformShaderParameters.ScreenToWorld = FMatrix(
@@ -613,13 +613,13 @@ TUniformBufferRef<FViewUniformShaderParameters> FViewInfo::CreateUniformBuffer(
 				for (int32 i = 0; i < NumShadowsToCopy; ++i)
 				{
 					const FProjectedShadowInfo& ShadowInfo = *(*DirectionalLightShadowInfo)[i];
-					ViewUniformShaderParameters.DirectionalLightScreenToShadow[i] = ShadowInfo.GetScreenToShadowMatrix(*this);
+					ViewUniformShaderParameters.DirectionalLightSvPositionToShadow[i] = ShadowInfo.GetSvPositionToShadowMatrix(*this);
 					ViewUniformShaderParameters.DirectionalLightShadowDistances[i] = ShadowInfo.CascadeSettings.SplitFar;
 				}
 
 				for (int32 i = NumShadowsToCopy; i < MAX_FORWARD_SHADOWCASCADES; ++i)
 				{
-					ViewUniformShaderParameters.DirectionalLightScreenToShadow[i].SetIdentity();
+					ViewUniformShaderParameters.DirectionalLightSvPositionToShadow[i].SetIdentity();
 					ViewUniformShaderParameters.DirectionalLightShadowDistances[i] = 0.0f;
 				}
 			}
@@ -629,7 +629,7 @@ TUniformBufferRef<FViewUniformShaderParameters> FViewInfo::CreateUniformBuffer(
 				ViewUniformShaderParameters.DirectionalLightShadowSize = FVector::ZeroVector;
 				for (int32 i = 0; i < MAX_FORWARD_SHADOWCASCADES; ++i)
 				{
-					ViewUniformShaderParameters.DirectionalLightScreenToShadow[i].SetIdentity();
+					ViewUniformShaderParameters.DirectionalLightSvPositionToShadow[i].SetIdentity();
 					ViewUniformShaderParameters.DirectionalLightShadowDistances[i] = 0.0f;
 				}			
 			}			 
diff --git a/Engine/Source/Runtime/Renderer/Private/ShadowRendering.cpp b/Engine/Source/Runtime/Renderer/Private/ShadowRendering.cpp
--- a/Engine/Source/Runtime/Renderer/Private/ShadowRendering.cpp
+++ b/Engine/Source/Runtime/Renderer/Private/ShadowRendering.cpp
@@ -2548,23 +2548,41 @@ void FProjectedShadowInfo::RenderFrustumWireframe(FPrimitiveDrawInterface* PDI)
 		);
 }
 
-FMatrix FProjectedShadowInfo::GetScreenToShadowMatrix(const FSceneView& View) const
+FMatrix FProjectedShadowInfo::GetSvPositionToShadowMatrix(const FSceneView& View) const
 {
 	const FIntPoint ShadowBufferResolution = GetShadowBufferResolution();
 	const float InvBufferResolutionX = 1.0f / (float)ShadowBufferResolution.X;
 	const float ShadowResolutionFractionX = 0.5f * (float)ResolutionX * InvBufferResolutionX;
 	const float InvBufferResolutionY = 1.0f / (float)ShadowBufferResolution.Y;
 	const float ShadowResolutionFractionY = 0.5f * (float)ResolutionY * InvBufferResolutionY;
+
+	FMatrix SvPositionToClip;
+	{
+		FVector2D InvViewSize = FVector2D(1.0f / View.ViewRect.Width(), 1.0f / View.ViewRect.Height());
+
+		// setup a matrix to transform float4(SvPosition.xyz,1) directly to Shadow (quality, performance as we don't need to convert or use interpolator)
+
+		//	new_xy = (xy - ViewRectMin.xy) * ViewSizeAndInvSize.zw * float2(2,-2) + float2(-1, 1);
+
+		//  transformed into one MAD:  new_xy = xy * ViewSizeAndInvSize.zw * float2(2,-2)      +       (-ViewRectMin.xy) * ViewSizeAndInvSize.zw * float2(2,-2) + float2(-1, 1);
+
+		float Mx = 2.0f * InvViewSize.X;
+		float My = -2.0f * InvViewSize.Y;
+		float Ax = -1.0f - 2.0f * View.ViewRect.Min.X * InvViewSize.X;
+		float Ay = 1.0f + 2.0f * View.ViewRect.Min.Y * InvViewSize.Y;
+
+		SvPositionToClip = FMatrix(
+			FPlane(Mx,  0,  0,  0),
+			FPlane(0,  My,  0,  0),
+			FPlane(0,   0,  1,  0),
+			FPlane(Ax, Ay,  0,  1)
+			);
+	}
+
+
 	// Calculate the matrix to transform a screenspace position into shadow map space
-	FMatrix ScreenToShadow = 
-		// Z of the position being transformed is actually view space Z, 
-		// Transform it into post projection space by applying the projection matrix,
-		// Which is the required space before applying View.InvTranslatedViewProjectionMatrix
-		FMatrix(
-			FPlane(1,0,0,0),
-			FPlane(0,1,0,0),
-			FPlane(0,0,View.ViewMatrices.ProjMatrix.M[2][2],1),
-			FPlane(0,0,View.ViewMatrices.ProjMatrix.M[3][2],0)) * 
+	FMatrix SvPositionToShadowValue = 
+		SvPositionToClip * 
 		// Transform the post projection space position into translated world space
 		// Translated world space is normal world space translated to the view's origin, 
 		// Which prevents floating point imprecision far from the world origin.
@@ -2587,7 +2605,7 @@ FMatrix FProjectedShadowInfo::GetScreenToShadowMatrix(const FSceneView& View) co
 				1
 			)
 		);
-	return ScreenToShadow;
+	return SvPositionToShadowValue;
 }
 
 FMatrix FProjectedShadowInfo::GetWorldToShadowMatrix(FVector4& ShadowmapMinMax, const FIntPoint* ShadowBufferResolutionOverride, bool bHasShadowBorder ) const
diff --git a/Engine/Source/Runtime/Renderer/Private/ShadowRendering.h b/Engine/Source/Runtime/Renderer/Private/ShadowRendering.h
--- a/Engine/Source/Runtime/Renderer/Private/ShadowRendering.h
+++ b/Engine/Source/Runtime/Renderer/Private/ShadowRendering.h
@@ -710,7 +710,7 @@ public:
 	}
 
 	/** Returns a matrix that transforms a screen space position into shadow space. */
-	FMatrix GetScreenToShadowMatrix(const FSceneView& View) const;
+	FMatrix GetSvPositionToShadowMatrix(const FSceneView& View) const;
 
 	/** Returns a matrix that transforms a world space position into shadow space. */
 	FMatrix GetWorldToShadowMatrix(FVector4& ShadowmapMinMax, const FIntPoint* ShadowBufferResolutionOverride = nullptr, bool bHasShadowBorder = true ) const;
@@ -1066,7 +1066,7 @@ public:
 	void Bind(const FShaderParameterMap& ParameterMap)
 	{
 		DeferredParameters.Bind(ParameterMap);
-		ScreenToShadowMatrix.Bind(ParameterMap,TEXT("ScreenToShadowMatrix"));
+		SvPositionToShadow.Bind(ParameterMap,TEXT("SvPositionToShadow"));
 		SoftTransitionScale.Bind(ParameterMap,TEXT("SoftTransitionScale"));
 		ShadowBufferSize.Bind(ParameterMap,TEXT("ShadowBufferSize"));
 		ShadowDepthTexture.Bind(ParameterMap,TEXT("ShadowDepthTexture"));
@@ -1083,8 +1083,9 @@ public:
 		DeferredParameters.Set(RHICmdList, ShaderRHI, View);
 
 		// Set the transform from screen coordinates to shadow depth texture coordinates.
-		const FMatrix ScreenToShadow = ShadowInfo->GetScreenToShadowMatrix(View);
-		SetShaderValue(RHICmdList, ShaderRHI, ScreenToShadowMatrix, ScreenToShadow);
+		const FMatrix SvPositionToShadowValue = ShadowInfo->GetSvPositionToShadowMatrix(View);
+
+		SetShaderValue(RHICmdList, ShaderRHI, SvPositionToShadow, SvPositionToShadowValue);
 
 		const FIntPoint ShadowBufferResolution = ShadowInfo->GetShadowBufferResolution();
 
@@ -1131,7 +1132,7 @@ public:
 	friend FArchive& operator<<(FArchive& Ar,FShadowProjectionShaderParameters& P)
 	{
 		Ar << P.DeferredParameters;
-		Ar << P.ScreenToShadowMatrix;
+		Ar << P.SvPositionToShadow;
 		Ar << P.SoftTransitionScale;
 		Ar << P.ShadowBufferSize;
 		Ar << P.ShadowDepthTexture;
@@ -1145,7 +1146,7 @@ public:
 private:
 
 	FDeferredPixelShaderParameters DeferredParameters;
-	FShaderParameter ScreenToShadowMatrix;
+	FShaderParameter SvPositionToShadow;
 	FShaderParameter SoftTransitionScale;
 	FShaderParameter ShadowBufferSize;
 	FShaderResourceParameter ShadowDepthTexture;
