Merging

from Orion

2710525 2710527 2711233 2711790

minor tonemapping optimization,
added test when using the optional data on the shaders
fixed check() when using the Scalability command

diff --git a/Engine/Shaders/Common.usf b/Engine/Shaders/Common.usf
--- a/Engine/Shaders/Common.usf
+++ b/Engine/Shaders/Common.usf
@@ -1056,6 +1056,14 @@ float4 SvPositionToScreenPosition(float4 SvPosition)
 	return float4(NDCPos.xyz, 1) * SvPosition.w;
 }
 
+float2 SvPositionToViewportUV(float4 SvPosition)
+{
+	// can be optimized from 2SUB+2MUL to 2MAD
+	float2 PixelPos = SvPosition.xy - View.ViewRectMin.xy;	
+
+	return PixelPos.xy * View.ViewSizeAndInvSize.zw;
+}
+
 // ----------------------------
 
 /** 
diff --git a/Engine/Shaders/PostProcessTemporalAA.usf b/Engine/Shaders/PostProcessTemporalAA.usf
--- a/Engine/Shaders/PostProcessTemporalAA.usf
+++ b/Engine/Shaders/PostProcessTemporalAA.usf
@@ -297,7 +297,7 @@ void LightShaftTemporalAAPS( noperspective float4 UVAndScreenPos : TEXCOORD0, fl
 }
 
 void MainTemporalAAPS(
-	in noperspective float4 UVAndScreenPos : TEXCOORD0,
+	in noperspective float2 UV : TEXCOORD0,
 	in noperspective float3 InExposureScaleVignette : TEXCOORD1,
 	// Note: these are needed to match up with the vertex shader outputs which is FPostProcessTonemapVS, otherwise a d3d debug error is issued
 	in noperspective float4 Unused1 : TEXCOORD2,
@@ -305,6 +305,8 @@ void MainTemporalAAPS(
 	in float4 SvPosition : SV_Position,		// after all interpolators
 	out float4 OutColor : SV_Target0 )
 {
+	float4 UVAndScreenPos = float4(UV, SvPositionToScreenPosition(SvPosition).xy);
+
 #if ENABLE_TEMPORAL_AA
 	float InExposureScale = InExposureScaleVignette.x;
 	#define AA_FILTERED 1
diff --git a/Engine/Shaders/PostProcessTonemap.usf b/Engine/Shaders/PostProcessTonemap.usf
--- a/Engine/Shaders/PostProcessTonemap.usf
+++ b/Engine/Shaders/PostProcessTonemap.usf
@@ -94,15 +94,14 @@ float4 FringeUVParams;
 void MainVS(
 	in float4 InPosition : ATTRIBUTE0,
 	in float2 InTexCoord : ATTRIBUTE1,
-	out float4 OutTexCoord : TEXCOORD0,
+	out float2 OutTexCoord : TEXCOORD0,
 	out float3 OutExposureScaleVignette : TEXCOORD1,
 	out float4 OutGrainUV : TEXCOORD2,
 	out float4 OutFringe : TEXCOORD3,
 	out float4 OutPosition : SV_POSITION
 	)
 {
-	DrawRectangle(InPosition, InTexCoord, OutPosition, OutTexCoord.xy);
-	OutTexCoord = float4(OutTexCoord.xy, OutPosition.xy);
+	DrawRectangle(InPosition, InTexCoord, OutPosition, OutTexCoord);
 
 #if FEATURE_LEVEL >= FEATURE_LEVEL_SM5
 	// texture can be GWhiteTexture which is 1x1. It's important we don't read outside bounds.
@@ -113,14 +112,14 @@ void MainVS(
 #endif
 
 	// Scale vignette to always be a circle with consistent corner intensity.
-	OutExposureScaleVignette.yz = VignetteSpace(OutTexCoord.zw);
+	OutExposureScaleVignette.yz = VignetteSpace(OutPosition.xy);
 
 	// Grain
-	OutGrainUV.xy = OutTexCoord.xy + PostprocessInput0Size.zw * float2(-0.5,0.5);
-	OutGrainUV.zw = OutTexCoord.xy + GrainRandomFull.xy;
+	OutGrainUV.xy = OutTexCoord + PostprocessInput0Size.zw * float2(-0.5,0.5);
+	OutGrainUV.zw = OutTexCoord + GrainRandomFull.xy;
 
 	// Fringe
-	float2 ScreenPos = UVToScreenPos(OutTexCoord.xy);
+	float2 ScreenPos = UVToScreenPos(OutTexCoord);
 	OutFringe.xy = ScreenPosToUV(ScreenPos * FringeUVParams.r);
 	OutFringe.zw = ScreenPosToUV(ScreenPos * FringeUVParams.g);
 }
@@ -171,18 +170,16 @@ float3 Graph( float2 ScreenSpacePos )
 
 // pixel shader entry point
 void MainPS(
-	in noperspective float4 UVAndScreenPos : TEXCOORD0,
+	in noperspective float2 UV : TEXCOORD0,
 	in noperspective float3 InExposureScaleVignette : TEXCOORD1,
 	in noperspective float4 GrainUV : TEXCOORD2,
 	in noperspective float4 FringeUV : TEXCOORD3,
+	float4 SvPosition : SV_POSITION,		// after all interpolators
 	out float4 OutColor : SV_Target0
 	)
 {
 	OutColor = 0;
 
-	float2 UV = UVAndScreenPos.xy;
-	float2 ScreenSpacePos = UVAndScreenPos.zw;
-
 	#if USE_GRAIN_JITTER || USE_GRAIN_INTENSITY || USE_GRAIN_QUANTIZATION
 		half Grain = GrainFromUV(GrainUV.zw);
 	#endif
@@ -216,7 +213,7 @@ void MainPS(
 
 		#if USE_BLOOM
 			float4 CombinedBloom = Texture2DSample(PostprocessInput1, PostprocessInput1Sampler, UV);
-			float3 BloomDirtMaskColor = Texture2DSample(BloomDirtMask.Mask, BloomDirtMask.MaskSampler, ScreenSpacePos * float2(0.5, -0.5f) + 0.5f).rgb * BloomDirtMask.Tint.rgb; 
+			float3 BloomDirtMaskColor = Texture2DSample(BloomDirtMask.Mask, BloomDirtMask.MaskSampler, SvPositionToViewportUV(SvPosition)).rgb * BloomDirtMask.Tint.rgb; 
 			LinearColor += CombinedBloom.rgb * (ColorScale1.rgb + BloomDirtMaskColor); 
 		#endif
 
@@ -240,7 +237,7 @@ void MainPS(
 
 		half3 OutDeviceColor = ColorLookupTable( LinearColor );
 
-		//OutDeviceColor = Graph( ScreenSpacePos );
+		//OutDeviceColor = Graph(SvPositionToScreenPosition(SvPosition).xy);
 	
 		half LuminanceForPostProcessAA  = dot(OutDeviceColor, half3(0.299f, 0.587f, 0.114f));
 
diff --git a/Engine/Source/Runtime/Engine/Private/Scalability.cpp b/Engine/Source/Runtime/Engine/Private/Scalability.cpp
--- a/Engine/Source/Runtime/Engine/Private/Scalability.cpp
+++ b/Engine/Source/Runtime/Engine/Private/Scalability.cpp
@@ -425,6 +425,9 @@ void FQualityLevels::SetDefaults()
 
 void FQualityLevels::SetFromSingleQualityLevel(int32 Value)
 {
+	// clamp in the range "low" to "epic"
+	Value = FMath::Clamp(Value, 0, 3);
+
 	ResolutionQuality = GetRenderScaleLevelFromQualityLevel(Value);
 	ViewDistanceQuality = Value;
 	AntiAliasingQuality = Value;
diff --git a/Engine/Source/Runtime/Engine/Public/Scalability.h b/Engine/Source/Runtime/Engine/Public/Scalability.h
--- a/Engine/Source/Runtime/Engine/Public/Scalability.h
+++ b/Engine/Source/Runtime/Engine/Public/Scalability.h
@@ -38,7 +38,7 @@ namespace Scalability
 		}
 
 		// Sets all other settings based on an overall value
-		// @param Value 0:low, 1:medium, 2:high, 3:epic
+		// @param Value 0:low, 1:medium, 2:high, 3:epic (gets clamped if needed)
 		void SetFromSingleQualityLevel(int32 Value);
 
 		// Returns the overall value if all settings are set to the same thing
diff --git a/Engine/Source/Runtime/ShaderCore/Public/ShaderCore.h b/Engine/Source/Runtime/ShaderCore/Public/ShaderCore.h
--- a/Engine/Source/Runtime/ShaderCore/Public/ShaderCore.h
+++ b/Engine/Source/Runtime/ShaderCore/Public/ShaderCore.h
@@ -567,7 +567,7 @@ public:
 			uint8 Key = *Current++;
 			uint8 Size = *Current++;
 
-			if(Key == InKey)
+			if(Key == InKey && Size == ValueSize)
 			{
 				return Current;
 			}
