Integrate - Two new translucency lighting modes that fetch lighting from the volume textures per-vertex and interpolate them to the pixel shader.

diff --git a/Engine/Shaders/BasePassCommon.usf b/Engine/Shaders/BasePassCommon.usf
--- a/Engine/Shaders/BasePassCommon.usf
+++ b/Engine/Shaders/BasePassCommon.usf
@@ -43,6 +43,14 @@ struct FSharedBasePassInterpolants
 	#endif
 #endif
 
+#if TRANSLUCENCY_LIGHTING_VOLUMETRIC_PERVERTEX_NONDIRECTIONAL || TRANSLUCENCY_LIGHTING_VOLUMETRIC_PERVERTEX_DIRECTIONAL
+	float3 AmbientLightingVector : TEXCOORD12;
+#endif
+
+#if TRANSLUCENCY_LIGHTING_VOLUMETRIC_PERVERTEX_DIRECTIONAL
+	float3 DirectionalLightingVector : TEXCOORD13;
+#endif
+
 #if WRITES_VELOCITY_TO_GBUFFER
 	// .xy is clip position, pre divide by w; .w is clip W; .z is 0 or 1 to mask out the velocity output
 	float4 VelocityPrevScreenPosition : VELOCITY_PREV_POS;
@@ -78,3 +86,46 @@ struct FBasePassInterpolantsVSToDS : FSharedBasePassInterpolants
 #define FBasePassInterpolantsVSToPS FSharedBasePassInterpolants
 
 #endif
+
+Texture3D TranslucencyLightingVolumeAmbientInner;
+SamplerState TranslucencyLightingVolumeAmbientInnerSampler;
+Texture3D TranslucencyLightingVolumeAmbientOuter;
+SamplerState TranslucencyLightingVolumeAmbientOuterSampler;
+Texture3D TranslucencyLightingVolumeDirectionalInner;
+SamplerState TranslucencyLightingVolumeDirectionalInnerSampler;
+Texture3D TranslucencyLightingVolumeDirectionalOuter;
+SamplerState TranslucencyLightingVolumeDirectionalOuterSampler;
+
+void ComputeVolumeUVs(float3 WorldPosition, float3 LightingPositionOffset, out float3 InnerVolumeUVs, out float3 OuterVolumeUVs, out float FinalLerpFactor)
+{
+	// Apply a stable offset to the world position used for lighting, which breaks up artifacts from using a low res volume texture
+	InnerVolumeUVs = (WorldPosition + LightingPositionOffset - View.TranslucencyLightingVolumeMin[0].xyz) * View.TranslucencyLightingVolumeInvSize[0].xyz;
+	OuterVolumeUVs = (WorldPosition + LightingPositionOffset - View.TranslucencyLightingVolumeMin[1].xyz) * View.TranslucencyLightingVolumeInvSize[1].xyz;
+
+	// Controls how fast the lerp between the inner and outer cascades happens
+	// Larger values result in a shorter transition distance
+	float TransitionScale = 6;
+	// Setup a 3d lerp factor going to 0 at the edge of the inner volume
+	float3 LerpFactors = saturate((.5f - abs(InnerVolumeUVs - .5f)) * TransitionScale);
+	FinalLerpFactor = LerpFactors.x * LerpFactors.y * LerpFactors.z;
+}
+
+float4 GetAmbientLightingVectorFromTranslucentLightingVolume(float3 InnerVolumeUVs, float3 OuterVolumeUVs, float FinalLerpFactor)
+{
+	// Lookup the inner and outer cascade ambient lighting values
+	float4 InnerLighting = Texture3DSampleLevel(TranslucencyLightingVolumeAmbientInner, TranslucencyLightingVolumeAmbientInnerSampler, InnerVolumeUVs, 0);
+	float4 OuterLighting = Texture3DSampleLevel(TranslucencyLightingVolumeAmbientOuter, TranslucencyLightingVolumeAmbientOuterSampler, OuterVolumeUVs, 0);
+
+	// Lerp between cascades
+	return lerp(OuterLighting, InnerLighting, FinalLerpFactor);
+}
+
+float3 GetDirectionalLightingVectorFromTranslucentLightingVolume(float3 InnerVolumeUVs, float3 OuterVolumeUVs, float FinalLerpFactor)
+{
+	// Fetch both the ambient and directional values for both cascades
+	float3 InnerVector1 = Texture3DSampleLevel(TranslucencyLightingVolumeDirectionalInner, TranslucencyLightingVolumeDirectionalInnerSampler, InnerVolumeUVs, 0).rgb;
+	float3 OuterVector1 = Texture3DSampleLevel(TranslucencyLightingVolumeDirectionalOuter, TranslucencyLightingVolumeDirectionalOuterSampler, OuterVolumeUVs, 0).rgb;
+
+	// Lerp between cascades
+	return lerp(OuterVector1, InnerVector1, FinalLerpFactor);
+}
diff --git a/Engine/Shaders/BasePassPixelShader.usf b/Engine/Shaders/BasePassPixelShader.usf
--- a/Engine/Shaders/BasePassPixelShader.usf
+++ b/Engine/Shaders/BasePassPixelShader.usf
@@ -31,15 +31,6 @@
 	 
 #endif
 
-Texture3D TranslucencyLightingVolumeAmbientInner;
-SamplerState TranslucencyLightingVolumeAmbientInnerSampler;
-Texture3D TranslucencyLightingVolumeAmbientOuter;
-SamplerState TranslucencyLightingVolumeAmbientOuterSampler;
-Texture3D TranslucencyLightingVolumeDirectionalInner;
-SamplerState TranslucencyLightingVolumeDirectionalInnerSampler;
-Texture3D TranslucencyLightingVolumeDirectionalOuter;
-SamplerState TranslucencyLightingVolumeDirectionalOuterSampler;
-
 Texture2D		HZBTexture;
 SamplerState	HZBSampler;
 Texture2D		PrevSceneColor;
@@ -163,74 +154,80 @@ half3 GetImageBasedReflectionLighting(FMaterialPixelParameters MaterialParameter
 	#include "ForEachLight.usf"
 #endif
 
+void GetVolumeLightingNonDirectional(float4 AmbientLightingVector, float3 DiffuseColor, inout float3 InterpolatedLighting, out float4 VolumeLighting)
+{
+	// Normal is not taken into account with non directional lighting, and only the ambient term of the SH coefficients are needed
+	FOneBandSHVectorRGB TranslucentLighting;
+	TranslucentLighting.R.V.x = AmbientLightingVector.r;
+	TranslucentLighting.G.V.x = AmbientLightingVector.g;
+	TranslucentLighting.B.V.x = AmbientLightingVector.b;
+
+	FOneBandSHVector DiffuseTransferSH = CalcDiffuseTransferSH1(1);
+	VolumeLighting = float4(DotSH1(TranslucentLighting, DiffuseTransferSH), AmbientLightingVector.a);
+	InterpolatedLighting = DiffuseColor * VolumeLighting.rgb;
+}
+
+void GetVolumeLightingDirectional(float4 AmbientLightingVector, float3 DirectionalLightingVector, float3 WorldNormal, float3 DiffuseColor, inout float3 InterpolatedLighting, out float4 VolumeLighting)
+{
+	float DirectionalLightingIntensity = GetMaterialTranslucencyDirectionalLightingIntensity();
+
+	AmbientLightingVector.rgb /= DirectionalLightingIntensity;
+	DirectionalLightingVector.rgb *= DirectionalLightingIntensity;
+
+	// Reconstruct the SH coefficients based on what was encoded
+	FTwoBandSHVectorRGB TranslucentLighting;
+	TranslucentLighting.R.V.x = AmbientLightingVector.r;
+	TranslucentLighting.G.V.x = AmbientLightingVector.g;
+	TranslucentLighting.B.V.x = AmbientLightingVector.b;
+	float3 NormalizedAmbientColor = AmbientLightingVector.rgb / Luminance( AmbientLightingVector.rgb );
+
+	// Scale the monocrome directional coefficients with the normalzed ambient color as an approximation to the uncompressed values
+	TranslucentLighting.R.V.yzw = DirectionalLightingVector.rgb * NormalizedAmbientColor.r;
+	TranslucentLighting.G.V.yzw = DirectionalLightingVector.rgb * NormalizedAmbientColor.g;
+	TranslucentLighting.B.V.yzw = DirectionalLightingVector.rgb * NormalizedAmbientColor.b;
+			
+	// Compute diffuse lighting which takes the normal into account
+	FTwoBandSHVector DiffuseTransferSH = CalcDiffuseTransferSH(WorldNormal, 1);
+	VolumeLighting = float4(max(half3(0,0,0), DotSH(TranslucentLighting, DiffuseTransferSH)), AmbientLightingVector.a);
+	InterpolatedLighting += DiffuseColor * VolumeLighting.rgb;
+}
+
 /** Calculates lighting for translucency. */
-float3 GetTranslucencyLighting(FMaterialPixelParameters MaterialParameters, float3 DiffuseColor, half Roughness, half3 SpecularColor, float IndirectIrradiance, float MaterialAO)
+float3 GetTranslucencyLighting(
+	FMaterialPixelParameters MaterialParameters, 
+	FBasePassInterpolantsVSToPS BasePassInterpolants, 
+	float3 DiffuseColor, 
+	half Roughness, 
+	half3 SpecularColor, 
+	float IndirectIrradiance, 
+	float MaterialAO)
 {
+	float4 VolumeLighting;
 	float3 InterpolatedLighting = 0;
 
-	// Apply a stable offset to the world position used for lighting, which breaks up artifacts from using a low res volume texture
-	float3 InnerVolumeUVs = (MaterialParameters.AbsoluteWorldPosition + MaterialParameters.LightingPositionOffset - View.TranslucencyLightingVolumeMin[0].xyz) * View.TranslucencyLightingVolumeInvSize[0].xyz;
-	float3 OuterVolumeUVs = (MaterialParameters.AbsoluteWorldPosition + MaterialParameters.LightingPositionOffset - View.TranslucencyLightingVolumeMin[1].xyz) * View.TranslucencyLightingVolumeInvSize[1].xyz;
-
-	// Controls how fast the lerp between the inner and outer cascades happens
-	// Larger values result in a shorter transition distance
-	float TransitionScale = 6;
-	// Setup a 3d lerp factor going to 0 at the edge of the inner volume
-	float3 InnerLerpFactors = saturate((.5f - abs(InnerVolumeUVs - .5f)) * TransitionScale);
-	float InnerLerpFactor = InnerLerpFactors.x * InnerLerpFactors.y * InnerLerpFactors.z;
-
-	#if TRANSLUCENCY_LIGHTING_VOLUMETRIC_DIRECTIONAL || TRANSLUCENCY_LIGHTING_SURFACE
-
-		// Fetch both the ambient and directional values for both cascades
-		float4 InnerVector0 = Texture3DSampleLevel(TranslucencyLightingVolumeAmbientInner, TranslucencyLightingVolumeAmbientInnerSampler, InnerVolumeUVs, 0);
-		float3 InnerVector1 = Texture3DSampleLevel(TranslucencyLightingVolumeDirectionalInner, TranslucencyLightingVolumeDirectionalInnerSampler, InnerVolumeUVs, 0).rgb;
-		float4 OuterVector0 = Texture3DSampleLevel(TranslucencyLightingVolumeAmbientOuter, TranslucencyLightingVolumeAmbientOuterSampler, OuterVolumeUVs, 0);
-		float3 OuterVector1 = Texture3DSampleLevel(TranslucencyLightingVolumeDirectionalOuter, TranslucencyLightingVolumeDirectionalOuterSampler, OuterVolumeUVs, 0).rgb;
-
-		float DirectionalLightingIntensity = GetMaterialTranslucencyDirectionalLightingIntensity();
-
-		// Lerp between cascades
-		// Increase the directional coefficients and attenuate the ambient coefficient based on a tweaked value
-		float4 Vector0 = lerp(OuterVector0, InnerVector0, InnerLerpFactor) / DirectionalLightingIntensity;
-		float3 Vector1 = lerp(OuterVector1, InnerVector1, InnerLerpFactor) * DirectionalLightingIntensity;
-
-		// Reconstruct the SH coefficients based on what was encoded
-		FTwoBandSHVectorRGB TranslucentLighting;
-		TranslucentLighting.R.V.x = Vector0.r;
-		TranslucentLighting.G.V.x = Vector0.g;
-		TranslucentLighting.B.V.x = Vector0.b;
-		float3 NormalizedAmbientColor = Vector0.rgb / Luminance( Vector0.rgb );
-
-		// Scale the monocrome directional coefficients with the normalzed ambient color as an approximation to the uncompressed values
-		TranslucentLighting.R.V.yzw = Vector1.rgb * NormalizedAmbientColor.r;
-		TranslucentLighting.G.V.yzw = Vector1.rgb * NormalizedAmbientColor.g;
-		TranslucentLighting.B.V.yzw = Vector1.rgb * NormalizedAmbientColor.b;
-			
-		// Compute diffuse lighting which takes the normal into account
-		FTwoBandSHVector DiffuseTransferSH = CalcDiffuseTransferSH(MaterialParameters.WorldNormal, 1);
-		float4 VolumeLighting = float4(max(half3(0,0,0), DotSH(TranslucentLighting, DiffuseTransferSH)), Vector0.a);
-		InterpolatedLighting = DiffuseColor * VolumeLighting.rgb;
+	float3 InnerVolumeUVs;
+	float3 OuterVolumeUVs;
+	float FinalLerpFactor;
+	ComputeVolumeUVs(MaterialParameters.AbsoluteWorldPosition, MaterialParameters.LightingPositionOffset, InnerVolumeUVs, OuterVolumeUVs, FinalLerpFactor);
 
-	#elif TRANSLUCENCY_LIGHTING_VOLUMETRIC_NONDIRECTIONAL
+	#if TRANSLUCENCY_LIGHTING_VOLUMETRIC_PERVERTEX_DIRECTIONAL
 
-		// Lookup the inner and outer cascade ambient lighting values
-		float4 InnerLighting = Texture3DSampleLevel(TranslucencyLightingVolumeAmbientInner, TranslucencyLightingVolumeAmbientInnerSampler, InnerVolumeUVs, 0);
-		float4 OuterLighting = Texture3DSampleLevel(TranslucencyLightingVolumeAmbientOuter, TranslucencyLightingVolumeAmbientOuterSampler, OuterVolumeUVs, 0);
+		GetVolumeLightingDirectional(float4(BasePassInterpolants.AmbientLightingVector, 1), BasePassInterpolants.DirectionalLightingVector, MaterialParameters.WorldNormal, DiffuseColor, InterpolatedLighting, VolumeLighting);
 
-		// Lerp between cascades
-		float4 Vector0 = lerp(OuterLighting, InnerLighting, InnerLerpFactor);
+	#elif TRANSLUCENCY_LIGHTING_VOLUMETRIC_PERVERTEX_NONDIRECTIONAL
 
-		// Normal is not taken into account with non directional lighting, and only the ambient term of the SH coefficients are needed
-		FOneBandSHVectorRGB TranslucentLighting;
-		TranslucentLighting.R.V.x = Vector0.r;
-		TranslucentLighting.G.V.x = Vector0.g;
-		TranslucentLighting.B.V.x = Vector0.b;
+		GetVolumeLightingNonDirectional(float4(BasePassInterpolants.AmbientLightingVector, 1), DiffuseColor, InterpolatedLighting, VolumeLighting);
 
-		FOneBandSHVector DiffuseTransferSH = CalcDiffuseTransferSH1(MaterialParameters.WorldNormal, 1);
-		float4 VolumeLighting = float4( DotSH1(TranslucentLighting, DiffuseTransferSH), Vector0.a );
+	#elif TRANSLUCENCY_LIGHTING_VOLUMETRIC_DIRECTIONAL || TRANSLUCENCY_LIGHTING_SURFACE
+	
+		float4 AmbientLightingVector = GetAmbientLightingVectorFromTranslucentLightingVolume(InnerVolumeUVs, OuterVolumeUVs, FinalLerpFactor);
+		float3 DirectionalLightingVector = GetDirectionalLightingVectorFromTranslucentLightingVolume(InnerVolumeUVs, OuterVolumeUVs, FinalLerpFactor);
+		GetVolumeLightingDirectional(AmbientLightingVector, DirectionalLightingVector, MaterialParameters.WorldNormal, DiffuseColor, InterpolatedLighting, VolumeLighting);
+
+	#elif TRANSLUCENCY_LIGHTING_VOLUMETRIC_NONDIRECTIONAL
 
-		// Determine lighting if no self shadowing is applied
-		InterpolatedLighting = DiffuseColor * VolumeLighting.rgb;
+		float4 AmbientLightingVector = GetAmbientLightingVectorFromTranslucentLightingVolume(InnerVolumeUVs, OuterVolumeUVs, FinalLerpFactor);
+		GetVolumeLightingNonDirectional(AmbientLightingVector, DiffuseColor, InterpolatedLighting, VolumeLighting);
 
 	#elif TRANSLUCENCY_LIGHTING_SURFACE_PERPIXEL
 
@@ -365,7 +362,7 @@ void GetSkyLighting(float3 WorldNormal, float2 LightmapUV, out float3 OutDiffuse
 
 			#endif
 
-			#if (MATERIALBLENDING_TRANSLUCENT || MATERIALBLENDING_ADDITIVE) && TRANSLUCENCY_LIGHTING_VOLUMETRIC_NONDIRECTIONAL
+			#if (MATERIALBLENDING_TRANSLUCENT || MATERIALBLENDING_ADDITIVE) && (TRANSLUCENCY_LIGHTING_VOLUMETRIC_NONDIRECTIONAL || TRANSLUCENCY_LIGHTING_VOLUMETRIC_PERVERTEX_NONDIRECTIONAL)
 				// NonDirectional lighting can't depend on the normal
 				SkyLightingNormal = NormalizedBentNormal;
 			#else
@@ -461,11 +458,11 @@ void GetPrecomputedIndirectLightingAndSkyLight(
 			PointIndirectLighting.G.V = IndirectLightingSHCoefficients[1].x;
 			PointIndirectLighting.B.V = IndirectLightingSHCoefficients[2].x;
 
-			FOneBandSHVector DiffuseTransferSH = CalcDiffuseTransferSH1(MaterialParameters.WorldNormal, 1);
+			FOneBandSHVector DiffuseTransferSH = CalcDiffuseTransferSH1(1);
 			OutDiffuseLighting = DotSH1(PointIndirectLighting, DiffuseTransferSH);
 
 			#if MATERIAL_SHADINGMODEL_TWOSIDED_FOLIAGE
-				FOneBandSHVector SubsurfaceTransferSH = CalcDiffuseTransferSH1(-MaterialParameters.WorldNormal, 1);
+				FOneBandSHVector SubsurfaceTransferSH = CalcDiffuseTransferSH1(1);
 				OutSubsurfaceLighting += DotSH1(PointIndirectLighting, SubsurfaceTransferSH);
 			#endif
 
@@ -829,7 +826,7 @@ void Main(
 		
 	// Volume lighting for lit translucency
 	#if (MATERIAL_SHADINGMODEL_DEFAULT_LIT || MATERIAL_SHADINGMODEL_SUBSURFACE) && (MATERIALBLENDING_TRANSLUCENT || MATERIALBLENDING_ADDITIVE)
-		Color += GetTranslucencyLighting(MaterialParameters, DiffuseColor, Roughness, SpecularColor, IndirectIrradiance, MaterialAO);
+		Color += GetTranslucencyLighting(MaterialParameters, BasePassInterpolants, DiffuseColor, Roughness, SpecularColor, IndirectIrradiance, MaterialAO);
 	#endif
 
 	#if !MATERIAL_SHADINGMODEL_UNLIT
diff --git a/Engine/Shaders/BasePassTessellationShaders.usf b/Engine/Shaders/BasePassTessellationShaders.usf
--- a/Engine/Shaders/BasePassTessellationShaders.usf
+++ b/Engine/Shaders/BasePassTessellationShaders.usf
@@ -17,6 +17,14 @@ FBasePassVSToDS PassInterpolate(FBasePassVSToDS a, float aInterp, FBasePassVSToD
 		TESSELLATION_INTERPOLATE_MEMBER(BasePassInterpolants.VertexFog);
 	#endif
 
+	#if TRANSLUCENCY_LIGHTING_VOLUMETRIC_PERVERTEX_NONDIRECTIONAL || TRANSLUCENCY_LIGHTING_VOLUMETRIC_PERVERTEX_DIRECTIONAL
+		TESSELLATION_INTERPOLATE_MEMBER(AmbientLightingVector);
+	#endif
+
+	#if TRANSLUCENCY_LIGHTING_VOLUMETRIC_PERVERTEX_DIRECTIONAL
+		TESSELLATION_INTERPOLATE_MEMBER(DirectionalLightingVector);
+	#endif
+
 	#if USE_WORLD_POSITION_EXCLUDING_SHADER_OFFSETS
 		TESSELLATION_INTERPOLATE_MEMBER(BasePassInterpolants.WorldPositionExcludingWPO);
 	#endif
diff --git a/Engine/Shaders/BasePassVertexShader.usf b/Engine/Shaders/BasePassVertexShader.usf
--- a/Engine/Shaders/BasePassVertexShader.usf
+++ b/Engine/Shaders/BasePassVertexShader.usf
@@ -63,6 +63,31 @@ void Main(
 	#endif
 #endif
 
+#if TRANSLUCENCY_LIGHTING_VOLUMETRIC_PERVERTEX_DIRECTIONAL || TRANSLUCENCY_LIGHTING_VOLUMETRIC_PERVERTEX_NONDIRECTIONAL
+	float4 VolumeLighting;
+	float3 InterpolatedLighting = 0;
+
+	float3 InnerVolumeUVs;
+	float3 OuterVolumeUVs;
+	float FinalLerpFactor;
+
+	//@todo - get from VF
+	float3 LightingPositionOffset = 0;
+	ComputeVolumeUVs(WorldPosition.xyz - View.PreViewTranslation.xyz, LightingPositionOffset, InnerVolumeUVs, OuterVolumeUVs, FinalLerpFactor);
+
+	#if TRANSLUCENCY_LIGHTING_VOLUMETRIC_PERVERTEX_DIRECTIONAL
+	
+		Output.BasePassInterpolants.AmbientLightingVector = GetAmbientLightingVectorFromTranslucentLightingVolume(InnerVolumeUVs, OuterVolumeUVs, FinalLerpFactor).xyz;
+		Output.BasePassInterpolants.DirectionalLightingVector = GetDirectionalLightingVectorFromTranslucentLightingVolume(InnerVolumeUVs, OuterVolumeUVs, FinalLerpFactor);
+
+	#elif TRANSLUCENCY_LIGHTING_VOLUMETRIC_PERVERTEX_NONDIRECTIONAL
+
+		Output.BasePassInterpolants.AmbientLightingVector = GetAmbientLightingVectorFromTranslucentLightingVolume(InnerVolumeUVs, OuterVolumeUVs, FinalLerpFactor).xyz;
+
+	#endif
+#endif
+
+
 #if WRITES_VELOCITY_TO_GBUFFER
 	{
 		float4 PrevTranslatedWorldPosition = VertexFactoryGetPreviousWorldPosition( Input, VFIntermediates );	
diff --git a/Engine/Shaders/SHCommon.usf b/Engine/Shaders/SHCommon.usf
--- a/Engine/Shaders/SHCommon.usf
+++ b/Engine/Shaders/SHCommon.usf
@@ -128,7 +128,7 @@ half3 DotSH3(FThreeBandSHVectorRGB A,FThreeBandSHVector B)
 }
 
 /** Computes the value of the SH basis functions for a given direction. */
-FOneBandSHVector SHBasisFunction1(half3 InputVector)
+FOneBandSHVector SHBasisFunction1()
 {
 	FOneBandSHVector Result;
 	// These are derived from simplifying SHBasisFunction in C++
@@ -176,9 +176,9 @@ half SHAmbientFunction()
  * Computes the diffuse transfer function for a surface with the given normal and DiffusePower,
  * and projects it into the SH basis.
  */
-FOneBandSHVector CalcDiffuseTransferSH1(half3 Normal,half Exponent)
+FOneBandSHVector CalcDiffuseTransferSH1(half Exponent)
 {
-	FOneBandSHVector Result = SHBasisFunction1(Normal);
+	FOneBandSHVector Result = SHBasisFunction1();
 
 	// These formula are scaling factors for each SH band that convolve a SH with the circularly symmetric function
 	// max(0,cos(theta))^Exponent
diff --git a/Engine/Source/Runtime/Engine/Classes/Engine/EngineTypes.h b/Engine/Source/Runtime/Engine/Classes/Engine/EngineTypes.h
--- a/Engine/Source/Runtime/Engine/Classes/Engine/EngineTypes.h
+++ b/Engine/Source/Runtime/Engine/Classes/Engine/EngineTypes.h
@@ -132,7 +132,7 @@ enum ETranslucencyLightingMode
 	 * Lighting will be calculated for a volume, without directionality.  Use this on particle effects like smoke and dust.
 	 * This is the cheapest lighting method, however the material normal is not taken into account.
 	 */
-	TLM_VolumetricNonDirectional UMETA(DisplayName="Volumetric Non Directional"),
+	TLM_VolumetricNonDirectional UMETA(DisplayName="Volumetric NonDirectional"),
 
 	 /** 
 	 * Lighting will be calculated for a volume, with directionality so that the normal of the material is taken into account. 
@@ -140,6 +140,18 @@ enum ETranslucencyLightingMode
 	 */
 	TLM_VolumetricDirectional UMETA(DisplayName="Volumetric Directional"),
 
+	/** 
+	 * Same as Volumetric Non Directional, but lighting is only evaluated at vertices so the pixel shader cost is significantly less.
+	 * Note that lighting still comes from a volume texture, so it is limited in range.  Directional lights become unshadowed in the distance.
+	 */
+	TLM_VolumetricPerVertexNonDirectional UMETA(DisplayName="Volumetric PerVertex NonDirectional"),
+
+	 /** 
+	 * Same as Volumetric Directional, but lighting is only evaluated at vertices so the pixel shader cost is significantly less.
+	 * Note that lighting still comes from a volume texture, so it is limited in range.  Directional lights become unshadowed in the distance.
+	 */
+	TLM_VolumetricPerVertexDirectional UMETA(DisplayName="Volumetric PerVertex Directional"),
+
 	/** 
 	 * Lighting will be calculated for a surface. The light in accumulated in a volume so the result is blurry
 	 * (fixed resolution), limited distance but the per pixel cost is very low. Use this on translucent surfaces like glass and water.
diff --git a/Engine/Source/Runtime/Engine/Classes/Materials/Material.h b/Engine/Source/Runtime/Engine/Classes/Materials/Material.h
--- a/Engine/Source/Runtime/Engine/Classes/Materials/Material.h
+++ b/Engine/Source/Runtime/Engine/Classes/Materials/Material.h
@@ -431,14 +431,14 @@ public:
 	float TranslucencyDirectionalLightingIntensity;
 
 	/** Scale used to make translucent shadows more or less opaque than the material's actual opacity. */
-	UPROPERTY(EditAnywhere, Category=TranslucencySelfShadowing, meta=(DisplayName = "Density Scale"))
+	UPROPERTY(EditAnywhere, Category=TranslucencySelfShadowing, meta=(DisplayName = "Shadow Density Scale"))
 	float TranslucentShadowDensityScale;
 
 	/** 
 	 * Scale used to make translucent self-shadowing more or less opaque than the material's shadow on other objects. 
 	 * This is only used when the object is casting a volumetric translucent shadow.
 	 */
-	UPROPERTY(EditAnywhere, Category=TranslucencySelfShadowing, meta=(DisplayName = "Density Scale"))
+	UPROPERTY(EditAnywhere, Category=TranslucencySelfShadowing, meta=(DisplayName = "Self Shadow Density Scale"))
 	float TranslucentSelfShadowDensityScale;
 
 	/** Used to make a second self shadow gradient, to add interesting shading in the shadow of the first. */
diff --git a/Engine/Source/Runtime/Engine/Private/Materials/MaterialShared.cpp b/Engine/Source/Runtime/Engine/Private/Materials/MaterialShared.cpp
--- a/Engine/Source/Runtime/Engine/Private/Materials/MaterialShared.cpp
+++ b/Engine/Source/Runtime/Engine/Private/Materials/MaterialShared.cpp
@@ -1277,6 +1277,8 @@ void FMaterial::SetupMaterialEnvironment(
 		{
 		case TLM_VolumetricNonDirectional: OutEnvironment.SetDefine(TEXT("TRANSLUCENCY_LIGHTING_VOLUMETRIC_NONDIRECTIONAL"),TEXT("1")); break;
 		case TLM_VolumetricDirectional: OutEnvironment.SetDefine(TEXT("TRANSLUCENCY_LIGHTING_VOLUMETRIC_DIRECTIONAL"),TEXT("1")); break;
+		case TLM_VolumetricPerVertexNonDirectional: OutEnvironment.SetDefine(TEXT("TRANSLUCENCY_LIGHTING_VOLUMETRIC_PERVERTEX_NONDIRECTIONAL"),TEXT("1")); break;
+		case TLM_VolumetricPerVertexDirectional: OutEnvironment.SetDefine(TEXT("TRANSLUCENCY_LIGHTING_VOLUMETRIC_PERVERTEX_DIRECTIONAL"),TEXT("1")); break;
 		case TLM_Surface: OutEnvironment.SetDefine(TEXT("TRANSLUCENCY_LIGHTING_SURFACE"),TEXT("1")); break;
 		case TLM_SurfacePerPixelLighting: OutEnvironment.SetDefine(TEXT("TRANSLUCENCY_LIGHTING_SURFACE_PERPIXEL"),TEXT("1")); break;
 
diff --git a/Engine/Source/Runtime/Renderer/Private/BasePassRendering.cpp b/Engine/Source/Runtime/Renderer/Private/BasePassRendering.cpp
--- a/Engine/Source/Runtime/Renderer/Private/BasePassRendering.cpp
+++ b/Engine/Source/Runtime/Renderer/Private/BasePassRendering.cpp
@@ -118,37 +118,7 @@ void FTranslucentLightingParameters::Set(FRHICommandList& RHICmdList, FShader* S
 	FSceneRenderTargets& SceneContext = FSceneRenderTargets::Get(RHICmdList);
 	auto PixelShader = Shader->GetPixelShader();
 
-	SetTextureParameter(
-		RHICmdList, 
-		PixelShader, 
-		TranslucencyLightingVolumeAmbientInner, 
-		TranslucencyLightingVolumeAmbientInnerSampler, 
-		TStaticSamplerState<SF_Bilinear, AM_Clamp, AM_Clamp, AM_Clamp>::GetRHI(), 
-		SceneContext.GetTranslucencyVolumeAmbient(TVC_Inner)->GetRenderTargetItem().ShaderResourceTexture);
-
-	SetTextureParameter(
-		RHICmdList, 
-		PixelShader, 
-		TranslucencyLightingVolumeAmbientOuter, 
-		TranslucencyLightingVolumeAmbientOuterSampler, 
-		TStaticSamplerState<SF_Bilinear, AM_Clamp, AM_Clamp, AM_Clamp>::GetRHI(), 
-		SceneContext.GetTranslucencyVolumeAmbient(TVC_Outer)->GetRenderTargetItem().ShaderResourceTexture);
-
-	SetTextureParameter(
-		RHICmdList, 
-		PixelShader, 
-		TranslucencyLightingVolumeDirectionalInner, 
-		TranslucencyLightingVolumeDirectionalInnerSampler, 
-		TStaticSamplerState<SF_Bilinear, AM_Clamp, AM_Clamp, AM_Clamp>::GetRHI(), 
-		SceneContext.GetTranslucencyVolumeDirectional(TVC_Inner)->GetRenderTargetItem().ShaderResourceTexture);
-
-	SetTextureParameter(
-		RHICmdList, 
-		PixelShader, 
-		TranslucencyLightingVolumeDirectionalOuter, 
-		TranslucencyLightingVolumeDirectionalOuterSampler, 
-		TStaticSamplerState<SF_Bilinear, AM_Clamp, AM_Clamp, AM_Clamp>::GetRHI(), 
-		SceneContext.GetTranslucencyVolumeDirectional(TVC_Outer)->GetRenderTargetItem().ShaderResourceTexture);
+	TranslucentLightingVolumeParameters.Set(RHICmdList, PixelShader);
 
 	SkyLightReflectionParameters.SetParameters(RHICmdList, Shader->GetPixelShader(), (const FScene*)(View->Family->Scene), true);
 
diff --git a/Engine/Source/Runtime/Renderer/Private/BasePassRendering.h b/Engine/Source/Runtime/Renderer/Private/BasePassRendering.h
--- a/Engine/Source/Runtime/Renderer/Private/BasePassRendering.h
+++ b/Engine/Source/Runtime/Renderer/Private/BasePassRendering.h
@@ -16,7 +16,91 @@
 extern int32 GIndirectLightingCache;
 
 /** Whether some GBuffer targets are optionnal. */
-bool UseSelectiveBasePassOutputs();
+extern bool UseSelectiveBasePassOutputs();
 
+/** Parameters needed for looking up into translucency lighting volumes. */
+class FTranslucentLightingVolumeParameters
+{
+public:
+
+	void Bind(const FShaderParameterMap& ParameterMap)
+	{
+		TranslucencyLightingVolumeAmbientInner.Bind(ParameterMap, TEXT("TranslucencyLightingVolumeAmbientInner"));
+		TranslucencyLightingVolumeAmbientInnerSampler.Bind(ParameterMap, TEXT("TranslucencyLightingVolumeAmbientInnerSampler"));
+		TranslucencyLightingVolumeAmbientOuter.Bind(ParameterMap, TEXT("TranslucencyLightingVolumeAmbientOuter"));
+		TranslucencyLightingVolumeAmbientOuterSampler.Bind(ParameterMap, TEXT("TranslucencyLightingVolumeAmbientOuterSampler"));
+		TranslucencyLightingVolumeDirectionalInner.Bind(ParameterMap, TEXT("TranslucencyLightingVolumeDirectionalInner"));
+		TranslucencyLightingVolumeDirectionalInnerSampler.Bind(ParameterMap, TEXT("TranslucencyLightingVolumeDirectionalInnerSampler"));
+		TranslucencyLightingVolumeDirectionalOuter.Bind(ParameterMap, TEXT("TranslucencyLightingVolumeDirectionalOuter"));
+		TranslucencyLightingVolumeDirectionalOuterSampler.Bind(ParameterMap, TEXT("TranslucencyLightingVolumeDirectionalOuterSampler"));
+	}
+
+	template<typename ShaderRHIParamRef>
+	void Set(FRHICommandList& RHICmdList, const ShaderRHIParamRef& ShaderRHI)
+	{
+		if (TranslucencyLightingVolumeAmbientInner.IsBound())
+		{
+			FSceneRenderTargets& SceneContext = FSceneRenderTargets::Get(RHICmdList);
+
+			SetTextureParameter(
+				RHICmdList, 
+				ShaderRHI, 
+				TranslucencyLightingVolumeAmbientInner, 
+				TranslucencyLightingVolumeAmbientInnerSampler, 
+				TStaticSamplerState<SF_Bilinear, AM_Clamp, AM_Clamp, AM_Clamp>::GetRHI(), 
+				SceneContext.GetTranslucencyVolumeAmbient(TVC_Inner)->GetRenderTargetItem().ShaderResourceTexture);
+
+			SetTextureParameter(
+				RHICmdList, 
+				ShaderRHI, 
+				TranslucencyLightingVolumeAmbientOuter, 
+				TranslucencyLightingVolumeAmbientOuterSampler, 
+				TStaticSamplerState<SF_Bilinear, AM_Clamp, AM_Clamp, AM_Clamp>::GetRHI(), 
+				SceneContext.GetTranslucencyVolumeAmbient(TVC_Outer)->GetRenderTargetItem().ShaderResourceTexture);
+
+			SetTextureParameter(
+				RHICmdList, 
+				ShaderRHI, 
+				TranslucencyLightingVolumeDirectionalInner, 
+				TranslucencyLightingVolumeDirectionalInnerSampler, 
+				TStaticSamplerState<SF_Bilinear, AM_Clamp, AM_Clamp, AM_Clamp>::GetRHI(), 
+				SceneContext.GetTranslucencyVolumeDirectional(TVC_Inner)->GetRenderTargetItem().ShaderResourceTexture);
+
+			SetTextureParameter(
+				RHICmdList, 
+				ShaderRHI, 
+				TranslucencyLightingVolumeDirectionalOuter, 
+				TranslucencyLightingVolumeDirectionalOuterSampler, 
+				TStaticSamplerState<SF_Bilinear, AM_Clamp, AM_Clamp, AM_Clamp>::GetRHI(), 
+				SceneContext.GetTranslucencyVolumeDirectional(TVC_Outer)->GetRenderTargetItem().ShaderResourceTexture);
+		}
+	}
+
+	/** Serializer. */
+	friend FArchive& operator<<(FArchive& Ar,FTranslucentLightingVolumeParameters& P)
+	{
+		Ar << P.TranslucencyLightingVolumeAmbientInner;
+		Ar << P.TranslucencyLightingVolumeAmbientInnerSampler;
+		Ar << P.TranslucencyLightingVolumeAmbientOuter;
+		Ar << P.TranslucencyLightingVolumeAmbientOuterSampler;
+		Ar << P.TranslucencyLightingVolumeDirectionalInner;
+		Ar << P.TranslucencyLightingVolumeDirectionalInnerSampler;
+		Ar << P.TranslucencyLightingVolumeDirectionalOuter;
+		Ar << P.TranslucencyLightingVolumeDirectionalOuterSampler;
+		return Ar;
+	}
+
+private:
+
+	FShaderResourceParameter TranslucencyLightingVolumeAmbientInner;
+	FShaderResourceParameter TranslucencyLightingVolumeAmbientInnerSampler;
+	FShaderResourceParameter TranslucencyLightingVolumeAmbientOuter;
+	FShaderResourceParameter TranslucencyLightingVolumeAmbientOuterSampler;
+	FShaderResourceParameter TranslucencyLightingVolumeDirectionalInner;
+	FShaderResourceParameter TranslucencyLightingVolumeDirectionalInnerSampler;
+	FShaderResourceParameter TranslucencyLightingVolumeDirectionalOuter;
+	FShaderResourceParameter TranslucencyLightingVolumeDirectionalOuterSampler;
+};
+
 /**
  * The base shader type for vertex shaders that render the emissive color, and light-mapped/ambient lighting of a mesh.
@@ -34,6 +118,7 @@
 		LightMapPolicyType::VertexParametersType::Bind(Initializer.ParameterMap);
 		HeightFogParameters.Bind(Initializer.ParameterMap);
 		AtmosphericFogTextureParameters.Bind(Initializer.ParameterMap);
+		TranslucentLightingVolumeParameters.Bind(Initializer.ParameterMap);
 		const bool bOutputsVelocityToGBuffer = FVelocityRendering::OutputsToGBuffer();
 		if (bOutputsVelocityToGBuffer)
 		{
@@ -61,6 +146,7 @@
 		LightMapPolicyType::VertexParametersType::Serialize(Ar);
 		Ar << HeightFogParameters;
 		Ar << AtmosphericFogTextureParameters;
+		Ar << TranslucentLightingVolumeParameters;
 		Ar << PreviousLocalToWorldParameter;
 		Ar << SkipOutputVelocityParameter;
 		return bShaderHasOutdatedParameters;
@@ -83,6 +169,8 @@
 			HeightFogParameters.Set(RHICmdList, GetVertexShader(), &View);
 			AtmosphericFogTextureParameters.Set(RHICmdList, GetVertexShader(), View);
 		}
+
+		TranslucentLightingVolumeParameters.Set(RHICmdList, GetVertexShader());
 	}
 
 	void SetMesh(FRHICommandList& RHICmdList, const FVertexFactory* VertexFactory,const FSceneView& View,const FPrimitiveSceneProxy* Proxy, const FMeshBatch& Mesh, const FMeshBatchElement& BatchElement, float DitheredLODTransitionValue);
@@ -92,6 +180,7 @@
 	/** The parameters needed to calculate the fog contribution from height fog layers. */
 	FHeightFogShaderParameters HeightFogParameters;
 	FAtmosphereShaderTextureParameters AtmosphericFogTextureParameters;
+	FTranslucentLightingVolumeParameters TranslucentLightingVolumeParameters;
 	// When outputting from base pass, the previous transform
 	FShaderParameter PreviousLocalToWorldParameter;
 	FShaderParameter SkipOutputVelocityParameter;
@@ -272,14 +361,7 @@
 
 	void Bind(const FShaderParameterMap& ParameterMap)
 	{
-		TranslucencyLightingVolumeAmbientInner.Bind(ParameterMap, TEXT("TranslucencyLightingVolumeAmbientInner"));
-		TranslucencyLightingVolumeAmbientInnerSampler.Bind(ParameterMap, TEXT("TranslucencyLightingVolumeAmbientInnerSampler"));
-		TranslucencyLightingVolumeAmbientOuter.Bind(ParameterMap, TEXT("TranslucencyLightingVolumeAmbientOuter"));
-		TranslucencyLightingVolumeAmbientOuterSampler.Bind(ParameterMap, TEXT("TranslucencyLightingVolumeAmbientOuterSampler"));
-		TranslucencyLightingVolumeDirectionalInner.Bind(ParameterMap, TEXT("TranslucencyLightingVolumeDirectionalInner"));
-		TranslucencyLightingVolumeDirectionalInnerSampler.Bind(ParameterMap, TEXT("TranslucencyLightingVolumeDirectionalInnerSampler"));
-		TranslucencyLightingVolumeDirectionalOuter.Bind(ParameterMap, TEXT("TranslucencyLightingVolumeDirectionalOuter"));
-		TranslucencyLightingVolumeDirectionalOuterSampler.Bind(ParameterMap, TEXT("TranslucencyLightingVolumeDirectionalOuterSampler"));
+		TranslucentLightingVolumeParameters.Bind(ParameterMap);
 		ReflectionCubemap.Bind(ParameterMap, TEXT("ReflectionCubemap"));
 		ReflectionCubemapSampler.Bind(ParameterMap, TEXT("ReflectionCubemapSampler"));
 		CubemapArrayIndex.Bind(ParameterMap, TEXT("CubemapArrayIndex"));
@@ -298,14 +380,7 @@
 	/** Serializer. */
 	friend FArchive& operator<<(FArchive& Ar,FTranslucentLightingParameters& P)
 	{
-		Ar << P.TranslucencyLightingVolumeAmbientInner;
-		Ar << P.TranslucencyLightingVolumeAmbientInnerSampler;
-		Ar << P.TranslucencyLightingVolumeAmbientOuter;
-		Ar << P.TranslucencyLightingVolumeAmbientOuterSampler;
-		Ar << P.TranslucencyLightingVolumeDirectionalInner;
-		Ar << P.TranslucencyLightingVolumeDirectionalInnerSampler;
-		Ar << P.TranslucencyLightingVolumeDirectionalOuter;
-		Ar << P.TranslucencyLightingVolumeDirectionalOuterSampler;
+		Ar << P.TranslucentLightingVolumeParameters;
 		Ar << P.ReflectionCubemap;
 		Ar << P.ReflectionCubemapSampler;
 		Ar << P.CubemapArrayIndex;
@@ -320,14 +395,7 @@
 
 private:
 
-	FShaderResourceParameter TranslucencyLightingVolumeAmbientInner;
-	FShaderResourceParameter TranslucencyLightingVolumeAmbientInnerSampler;
-	FShaderResourceParameter TranslucencyLightingVolumeAmbientOuter;
-	FShaderResourceParameter TranslucencyLightingVolumeAmbientOuterSampler;
-	FShaderResourceParameter TranslucencyLightingVolumeDirectionalInner;
-	FShaderResourceParameter TranslucencyLightingVolumeDirectionalInnerSampler;
-	FShaderResourceParameter TranslucencyLightingVolumeDirectionalOuter;
-	FShaderResourceParameter TranslucencyLightingVolumeDirectionalOuterSampler;
+	FTranslucentLightingVolumeParameters TranslucentLightingVolumeParameters;
 	FShaderResourceParameter ReflectionCubemap;
 	FShaderResourceParameter ReflectionCubemapSampler;
 	FShaderParameter CubemapArrayIndex;
