Rolled back Change 2685306 'volume texture content no longer repeats outside the volume, it assumed unshadowed' as it was not working correctly (only applied to self shadowing translucency) and an inefficient implementation (should be done during volume texture injection)

diff --git a/Engine/Shaders/BasePassPixelShader.usf b/Engine/Shaders/BasePassPixelShader.usf
--- a/Engine/Shaders/BasePassPixelShader.usf
+++ b/Engine/Shaders/BasePassPixelShader.usf
@@ -240,19 +240,6 @@ float3 GetTranslucencyLighting(
 		// Only apply self shadowing if the shadow hasn't faded out completely
 		if (DirectionalLightColor.a > 0)
 		{
-			// fade volume contribution to not get tiling content outside
-			{
-				// Setup a 3d lerp factor going to 0 at the edge of the outer volume
-				// maybe we can replace this with a border color but some hardware doesn't support a float border color
-				float3 OuterLerpFactors = saturate((.5f - abs(OuterVolumeUVs - .5f)) * TransitionScale);
-				float OuterLerpFactor = OuterLerpFactors.x * OuterLerpFactors.y * OuterLerpFactors.z;
-
-				// 0..1, 0:fully shadowed, 1:fully unshadowed
-				const float VisibilityOutside = 1;
-
-				VolumeLighting = lerp(float4(DirectionalLightColor.rgb, VisibilityOutside), VolumeLighting, OuterLerpFactor);
-			}
-
 			// Determine the shadow space position
 			// Apply a stable offset to the world position used for shadowing, which blurs out high frequency details in the shadowmap with many layers
 			float4 HomogeneousShadowPosition = mul(float4(MaterialParameters.AbsoluteWorldPosition + MaterialParameters.LightingPositionOffset, 1), WorldToShadowMatrix);
