Merging CL #271952: fixed TemporalAA shader errors

diff --git a/Engine/Shaders/PostProcessTemporalAA.usf b/Engine/Shaders/PostProcessTemporalAA.usf
--- a/Engine/Shaders/PostProcessTemporalAA.usf
+++ b/Engine/Shaders/PostProcessTemporalAA.usf
@@ -244,7 +244,9 @@ float RandFast( uint2 PixelPos, float Magic = 3571.0 )
 	return Random.x;
 }
 
-void SSRTemporalAAPS( noperspective float4 UVAndScreenPos : TEXCOORD0, float3 InExposureScaleVignette : TEXCOORD1, out float4 OutColor : SV_Target0 )
+void SSRTemporalAAPS( noperspective float2 UV : TEXCOORD0, float3 InExposureScaleVignette : TEXCOORD1, 
+	float4 SvPosition : SV_POSITION,		// after all interpolators
+	out float4 OutColor : SV_Target0 )
 {
 #if ENABLE_TEMPORAL_AA
 	float InExposureScale = InExposureScaleVignette.x;
@@ -256,11 +258,13 @@ void SSRTemporalAAPS( noperspective float4 UVAndScreenPos : TEXCOORD0, float3 In
 	#include "PostProcessTemporalCommon.usf"
 #else 
 	// On broken platforms then at least draw something without AA.
-	OutColor = PostprocessInput0.SampleLevel(PostprocessInput0Sampler, UVAndScreenPos.xy, 0);
+	OutColor = PostprocessInput0.SampleLevel(PostprocessInput0Sampler, UV, 0);
 #endif
 }
 
-void DOFTemporalAAPS( noperspective float4 UVAndScreenPos : TEXCOORD0, float3 InExposureScaleVignette : TEXCOORD1, out float4 OutColor : SV_Target0 )
+void DOFTemporalAAPS( noperspective float2 UV : TEXCOORD0, float3 InExposureScaleVignette : TEXCOORD1, 
+	float4 SvPosition : SV_POSITION,		// after all interpolators
+	out float4 OutColor : SV_Target0 )
 {
 #if ENABLE_TEMPORAL_AA
 	float InExposureScale = InExposureScaleVignette.x;
@@ -277,11 +281,13 @@ void DOFTemporalAAPS( noperspective float4 UVAndScreenPos : TEXCOORD0, float3 In
 	#define AA_DOF 1
 	#include "PostProcessTemporalCommon.usf"
 #else 
-	OutColor = PostprocessInput0.SampleLevel(PostprocessInput0Sampler, UVAndScreenPos.xy, 0);
+	OutColor = PostprocessInput0.SampleLevel(PostprocessInput0Sampler, UV, 0);
 #endif
 }
 
-void LightShaftTemporalAAPS( noperspective float4 UVAndScreenPos : TEXCOORD0, float3 InExposureScaleVignette : TEXCOORD1, out float4 OutColor : SV_Target0 )
+void LightShaftTemporalAAPS( noperspective float2 UV : TEXCOORD0, float3 InExposureScaleVignette : TEXCOORD1, 
+	float4 SvPosition : SV_POSITION,		// after all interpolators
+	out float4 OutColor : SV_Target0 )
 {
 #if ENABLE_TEMPORAL_AA
 	float InExposureScale = InExposureScaleVignette.x;
@@ -292,7 +298,7 @@ void LightShaftTemporalAAPS( noperspective float4 UVAndScreenPos : TEXCOORD0, fl
 	#define AA_NAN 1
 	#include "PostProcessTemporalCommon.usf"
 #else 
-	OutColor = PostprocessInput0.SampleLevel(PostprocessInput0Sampler, UVAndScreenPos.xy, 0);
+	OutColor = PostprocessInput0.SampleLevel(PostprocessInput0Sampler, UV, 0);
 #endif
 }
 
@@ -305,8 +311,6 @@ void MainTemporalAAPS(
 	in float4 SvPosition : SV_Position,		// after all interpolators
 	out float4 OutColor : SV_Target0 )
 {
-	float4 UVAndScreenPos = float4(UV, SvPositionToScreenPosition(SvPosition).xy);
-
 #if ENABLE_TEMPORAL_AA
 	float InExposureScale = InExposureScaleVignette.x;
 	#define AA_FILTERED 1
@@ -326,12 +330,12 @@ void MainTemporalAAPS(
 	float Random = RandFast( PixelPos );
 	OutColor += (Random - 0.5) * (1.0 / 1024.0) / InExposureScale;
 #else 
-	OutColor = PostprocessInput0.SampleLevel(PostprocessInput0Sampler, UVAndScreenPos.xy, 0);
+	OutColor = PostprocessInput0.SampleLevel(PostprocessInput0Sampler, UV, 0);
 #endif
 }
 
 void MainFastTemporalAAPS(
-	in noperspective float4 UVAndScreenPos : TEXCOORD0,
+	in noperspective float2 UV: TEXCOORD0,
 	in noperspective float3 InExposureScaleVignette : TEXCOORD1,
 	in float4 SvPosition : SV_Position,		// after all interpolators
 	out float4 OutColor : SV_Target0 )
@@ -355,6 +359,6 @@ void MainFastTemporalAAPS(
 	float Random = RandFast( PixelPos );
 	OutColor += (Random - 0.5) * (1.0 / 1024.0) / InExposureScale;
 #else 
-	OutColor = PostprocessInput0.SampleLevel(PostprocessInput0Sampler, UVAndScreenPos.xy, 0);
+	OutColor = PostprocessInput0.SampleLevel(PostprocessInput0Sampler, UV, 0);
 #endif
 }
diff --git a/Engine/Shaders/PostProcessTemporalCommon.usf b/Engine/Shaders/PostProcessTemporalCommon.usf
--- a/Engine/Shaders/PostProcessTemporalCommon.usf
+++ b/Engine/Shaders/PostProcessTemporalCommon.usf
@@ -112,8 +112,8 @@ which can be configured via defines for various temporal AA passes.
 	// FIND MOTION OF PIXEL AND NEAREST IN NEIGHBORHOOD
     // ------------------------------------------------
 	float3 PosN; // Position of this pixel, possibly later nearest pixel in neighborhood.
-	PosN.xy = UVAndScreenPos.zw;
-	PosN.z = SceneDepthTexture.SampleLevel(SceneDepthTextureSampler, UVAndScreenPos.xy, 0).r;
+	PosN.xy = SvPositionToScreenPosition(SvPosition).xy;
+	PosN.z = SceneDepthTexture.SampleLevel(SceneDepthTextureSampler, UV, 0).r;
 	// Screen position of minimum depth.
 	float2 VelocityOffset = float2(0.0, 0.0);
 	#if AA_CROSS
@@ -121,10 +121,10 @@ which can be configured via defines for various temporal AA passes.
 		// This enables better quality outline on foreground against different motion background.
 		// Larger 2 pixel distance "x" works best (because AA dilates surface).
 		float4 Depths;
-		Depths.x = SceneDepthTexture.SampleLevel(SceneDepthTextureSampler, UVAndScreenPos.xy, 0, int2(-AA_CROSS, -AA_CROSS)).r;
-		Depths.y = SceneDepthTexture.SampleLevel(SceneDepthTextureSampler, UVAndScreenPos.xy, 0, int2( AA_CROSS, -AA_CROSS)).r;
-		Depths.z = SceneDepthTexture.SampleLevel(SceneDepthTextureSampler, UVAndScreenPos.xy, 0, int2(-AA_CROSS,  AA_CROSS)).r;
-		Depths.w = SceneDepthTexture.SampleLevel(SceneDepthTextureSampler, UVAndScreenPos.xy, 0, int2( AA_CROSS,  AA_CROSS)).r;
+		Depths.x = SceneDepthTexture.SampleLevel(SceneDepthTextureSampler, UV, 0, int2(-AA_CROSS, -AA_CROSS)).r;
+		Depths.y = SceneDepthTexture.SampleLevel(SceneDepthTextureSampler, UV, 0, int2( AA_CROSS, -AA_CROSS)).r;
+		Depths.z = SceneDepthTexture.SampleLevel(SceneDepthTextureSampler, UV, 0, int2(-AA_CROSS,  AA_CROSS)).r;
+		Depths.w = SceneDepthTexture.SampleLevel(SceneDepthTextureSampler, UV, 0, int2( AA_CROSS,  AA_CROSS)).r;
 
 		float2 DepthOffset = float2(AA_CROSS, AA_CROSS);
 		float DepthOffsetXx = float(AA_CROSS);
@@ -153,7 +153,7 @@ which can be configured via defines for various temporal AA passes.
 			// With the assumption that UV position scales between velocity and color.
 			VelocityOffset = DepthOffset * PostprocessInput0Size.zw;
 			// This is [0 to 1] flipped in Y.
-			//PosN.xy = UVAndScreenPos.zw + DepthOffset * ViewportSize.zw * 2.0;
+			//PosN.xy = SvPositionToScreenPosition(SvPosition).xy + DepthOffset * ViewportSize.zw * 2.0;
 			PosN.z = DepthsXYZW;
 		}
 #else
@@ -171,9 +171,9 @@ which can be configured via defines for various temporal AA passes.
 	#if AA_DYNAMIC
 		float2 VelocityN;
 		#if AA_CROSS
-			VelocityN = PostprocessInput3.SampleLevel(PostprocessInput3Sampler, UVAndScreenPos.xy + VelocityOffset, 0).xy;
+			VelocityN = PostprocessInput3.SampleLevel(PostprocessInput3Sampler, UV + VelocityOffset, 0).xy;
 		#else
-			VelocityN = PostprocessInput3.SampleLevel(PostprocessInput3Sampler, UVAndScreenPos.xy, 0).xy;
+			VelocityN = PostprocessInput3.SampleLevel(PostprocessInput3Sampler, UV, 0).xy;
 		#endif
 		bool DynamicN = VelocityN.x > 0.0;
 		if(DynamicN)
@@ -191,7 +191,7 @@ which can be configured via defines for various temporal AA passes.
 	// Easier to do off screen check before conversion.
 	// BackN is in units of 2pixels/viewportWidthInPixels
 	// This converts back projection vector to [-1 to 1] offset in viewport.
-	BackN = UVAndScreenPos.zw - BackN;
+	BackN = SvPositionToScreenPosition(SvPosition).xy - BackN;
 	bool OffScreen = max(abs(BackN.x), abs(BackN.y)) >= 1.0;
 	// Also clamp to be on screen (fixes problem with DOF).
 	// The .z and .w is the 1/width and 1/height.
@@ -212,11 +212,11 @@ which can be configured via defines for various temporal AA passes.
 	// 678
 	#if AA_YCOCG	
 		// Special case, only using 5 taps.
-		float4 Neighbor1 = PostprocessInput0.SampleLevel(PostprocessInput0Sampler, UVAndScreenPos.xy, 0, int2( 0, -1));
-		float4 Neighbor3 = PostprocessInput0.SampleLevel(PostprocessInput0Sampler, UVAndScreenPos.xy, 0, int2(-1,  0));
-		float4 Neighbor4 = PostprocessInput0.SampleLevel(PostprocessInput0Sampler, UVAndScreenPos.xy, 0);
-		float4 Neighbor5 = PostprocessInput0.SampleLevel(PostprocessInput0Sampler, UVAndScreenPos.xy, 0, int2( 1,  0));
-		float4 Neighbor7 = PostprocessInput0.SampleLevel(PostprocessInput0Sampler, UVAndScreenPos.xy, 0, int2( 0,  1));
+		float4 Neighbor1 = PostprocessInput0.SampleLevel(PostprocessInput0Sampler, UV, 0, int2( 0, -1));
+		float4 Neighbor3 = PostprocessInput0.SampleLevel(PostprocessInput0Sampler, UV, 0, int2(-1,  0));
+		float4 Neighbor4 = PostprocessInput0.SampleLevel(PostprocessInput0Sampler, UV, 0);
+		float4 Neighbor5 = PostprocessInput0.SampleLevel(PostprocessInput0Sampler, UV, 0, int2( 1,  0));
+		float4 Neighbor7 = PostprocessInput0.SampleLevel(PostprocessInput0Sampler, UV, 0, int2( 0,  1));
 		Neighbor1.rgb = RGBToYCoCg(Neighbor1.rgb);
 		Neighbor3.rgb = RGBToYCoCg(Neighbor3.rgb);
 		Neighbor4.rgb = RGBToYCoCg(Neighbor4.rgb);
@@ -238,7 +238,7 @@ which can be configured via defines for various temporal AA passes.
 				Neighbor7 * PlusWeights[4];
 			#if AA_BORDER
 				// Use unfiltered for 1 pixel border.
-				float2 TestPos = abs(UVAndScreenPos.zw);
+				float2 TestPos = abs(SvPositionToScreenPosition(SvPosition).xy);
 				// Add 1 pixel and check if off screen.
 				TestPos += ViewportSize.zw * 2.0;
 				bool FilteredOffScreen = max(TestPos.x, TestPos.y) >= 1.0;
@@ -256,15 +256,15 @@ which can be configured via defines for various temporal AA passes.
 		float4 NeighborMin = min(min(min(Neighbor1, Neighbor3), min(Neighbor4, Neighbor5)), Neighbor7);		
 		float4 NeighborMax = max(max(max(Neighbor1, Neighbor3), max(Neighbor4, Neighbor5)), Neighbor7);		
 	#else
-		float4 Neighbor0 = PostprocessInput0.SampleLevel(PostprocessInput0Sampler, UVAndScreenPos.xy, 0, int2(-1, -1));
-		float4 Neighbor1 = PostprocessInput0.SampleLevel(PostprocessInput0Sampler, UVAndScreenPos.xy, 0, int2( 0, -1));
-		float4 Neighbor2 = PostprocessInput0.SampleLevel(PostprocessInput0Sampler, UVAndScreenPos.xy, 0, int2( 1, -1));
-		float4 Neighbor3 = PostprocessInput0.SampleLevel(PostprocessInput0Sampler, UVAndScreenPos.xy, 0, int2(-1,  0));
-		float4 Neighbor4 = PostprocessInput0.SampleLevel(PostprocessInput0Sampler, UVAndScreenPos.xy, 0);
-		float4 Neighbor5 = PostprocessInput0.SampleLevel(PostprocessInput0Sampler, UVAndScreenPos.xy, 0, int2( 1,  0));
-		float4 Neighbor6 = PostprocessInput0.SampleLevel(PostprocessInput0Sampler, UVAndScreenPos.xy, 0, int2(-1,  1));
-		float4 Neighbor7 = PostprocessInput0.SampleLevel(PostprocessInput0Sampler, UVAndScreenPos.xy, 0, int2( 0,  1));
-		float4 Neighbor8 = PostprocessInput0.SampleLevel(PostprocessInput0Sampler, UVAndScreenPos.xy, 0, int2( 1,  1));
+		float4 Neighbor0 = PostprocessInput0.SampleLevel(PostprocessInput0Sampler, UV, 0, int2(-1, -1));
+		float4 Neighbor1 = PostprocessInput0.SampleLevel(PostprocessInput0Sampler, UV, 0, int2( 0, -1));
+		float4 Neighbor2 = PostprocessInput0.SampleLevel(PostprocessInput0Sampler, UV, 0, int2( 1, -1));
+		float4 Neighbor3 = PostprocessInput0.SampleLevel(PostprocessInput0Sampler, UV, 0, int2(-1,  0));
+		float4 Neighbor4 = PostprocessInput0.SampleLevel(PostprocessInput0Sampler, UV, 0);
+		float4 Neighbor5 = PostprocessInput0.SampleLevel(PostprocessInput0Sampler, UV, 0, int2( 1,  0));
+		float4 Neighbor6 = PostprocessInput0.SampleLevel(PostprocessInput0Sampler, UV, 0, int2(-1,  1));
+		float4 Neighbor7 = PostprocessInput0.SampleLevel(PostprocessInput0Sampler, UV, 0, int2( 0,  1));
+		float4 Neighbor8 = PostprocessInput0.SampleLevel(PostprocessInput0Sampler, UV, 0, int2( 1,  1));
 		#if AA_GREEN_AS_LUMA
 			Neighbor0.rgb *= HdrWeightG(Neighbor0.rgb, InExposureScale);
 			Neighbor1.rgb *= HdrWeightG(Neighbor1.rgb, InExposureScale);
@@ -313,7 +313,7 @@ which can be configured via defines for various temporal AA passes.
 			#endif
 			#if AA_BORDER
 				// Use unfiltered for 1 pixel border.
-				float2 TestPos = abs(UVAndScreenPos.zw);
+				float2 TestPos = abs(SvPositionToScreenPosition(SvPosition).xy);
 				// Add 1 pixel and check if off screen.
 				TestPos += ViewportSize.zw * 2.0;
 				bool FilteredOffScreen = max(TestPos.x, TestPos.y) >= 1.0;
