Manually done the fix for UE-22323 in Engine-Main
Merging and/or copying the fix over gave errors. I'll work out that workflow for streams later.

diff --git a/Engine/Shaders/Common.usf b/Engine/Shaders/Common.usf
--- a/Engine/Shaders/Common.usf
+++ b/Engine/Shaders/Common.usf
@@ -680,13 +680,21 @@ MaterialFloat4 EncodeLightAttenuation(MaterialFloat4 InColor)
 /** return the scene lighting texture */
 MaterialFloat3 CalcSceneColor(MaterialFloat2 ScreenUV)
 {
+#if SCENE_TEXTURES_DISABLED
+	return MaterialFloat3(0.0f,0.0f,0.0f);
+#else
 	return Texture2DSampleLevel(SceneColorTexture, SceneColorTextureSampler, ScreenUV, 0).rgb;
+#endif
 }
 
 /** return all channels of the scene lighting texture */
 MaterialFloat4 CalcFullSceneColor(MaterialFloat2 ScreenUV)
 {
-	  return Texture2DSample(SceneColorTexture, SceneColorTextureSampler,ScreenUV);
+#if SCENE_TEXTURES_DISABLED
+	return MaterialFloat4(0.0f, 0.0f, 0.0f, 0.0f);
+#else
+	return Texture2DSample(SceneColorTexture, SceneColorTextureSampler,ScreenUV);
+#endif
 }
 
 /** Encodes HDR linear scene color for storage in the 8 bit light attenuation texture. */
@@ -904,7 +912,7 @@ float2 CalcScreenUVFromOffsetFraction(float4 ScreenPosition, float2 OffsetFracti
 /** Applies an offset to the scene texture lookup and decodes the HDR linear space color. */
 float3 DecodeSceneColorForMaterialNode(float2 ScreenUV)
 {
-#if HIT_PROXY_SHADER
+#if HIT_PROXY_SHADER || SCENE_TEXTURES_DISABLED
 	// Hit proxies rendering pass doesn't have access to valide render buffers
 	return float3(0.0f, 0.0f, 0.0f);
 #else
@@ -930,6 +938,10 @@ float ConvertToDeviceZ(float SceneDepth)
 /** Returns clip space W, which is world space distance along the View Z axis. Note if you need DeviceZ LookupDeviceZ() is the faster option */
 float CalcSceneDepth( float2 ScreenUV )
 {
+#if SCENE_TEXTURES_DISABLED
+	return 0.0f;
+#else
+
 #if FEATURE_LEVEL > FEATURE_LEVEL_ES3_1
 	return ConvertFromDeviceZ(Texture2DSampleLevel(SceneDepthTexture, SceneDepthTextureSampler, ScreenUV, 0).r);
 #else
@@ -949,6 +961,7 @@ float CalcSceneDepth( float2 ScreenUV )
 		return ConvertFromDeviceZ(Texture2DSampleLevel(SceneDepthTexture, SceneDepthTextureSampler, ScreenUV, 0).r);
 	#endif
 #endif
+#endif
 }
 
 /** Returns DeviceZ which is the z value storesd in the depth buffer. Perfer this over CalcSceneDepth() on most hardware. */
@@ -979,10 +992,14 @@ float LookupDeviceZ( float2 ScreenUV )
 	/** Returns clip space W, which is world space distance along the View Z axis. */
 	float CalcSceneDepth(uint2 PixelPos)
 	{
+#if SCENE_TEXTURES_DISABLED
+		return 0.0f;
+#else
 		float DeviceZ = SceneDepthTextureNonMS.Load(int3(PixelPos, 0));
 
 		// Fetch the depth buffer Z / W value, solve for W
 		return ConvertFromDeviceZ(DeviceZ);
+#endif
 	}
 #endif
 
