Fix use of unavailable modulo on ES2

diff --git a/Engine/Shaders/MiniFontCommon.usf b/Engine/Shaders/MiniFontCommon.usf
--- a/Engine/Shaders/MiniFontCommon.usf
+++ b/Engine/Shaders/MiniFontCommon.usf
@@ -67,7 +67,13 @@ uint ExtractDigitFromFloat(float Number, float DigitValue)
 {
 	uint Temp = (uint)(Number / DigitValue);
 
+#if ES2_PROFILE
+	// Integer modulo/remainder is not an allowed operation on ES 2
+	// todo: move function to central spot, here we assume unsigned
+	return Temp - (Temp / 10) * 10;
+#else
 	return Temp % 10;
+#endif
 }
 
 
