Fix for metal issue with the DownsampleDepthPixelShader.usf shader

diff --git a/Engine/Shaders/DownsampleDepthPixelShader.usf b/Engine/Shaders/DownsampleDepthPixelShader.usf
--- a/Engine/Shaders/DownsampleDepthPixelShader.usf
+++ b/Engine/Shaders/DownsampleDepthPixelShader.usf
@@ -42,7 +42,7 @@ void CopyStencilToLightingChannelsPS(
 	out uint4 OutValue : SV_Target0
 	)
 {
-	uint2 IntUV = (uint2)(InUV * View.BufferSizeAndInvSize.xy);
+	uint2 IntUV = (uint2)(InUV * (float2)View.BufferSizeAndInvSize.xy);
 	uint Stencil = SceneStencilTexture.Load(uint3(IntUV, 0)).g;
 	uint ShiftedStencil = Stencil >> STENCIL_LIGHTING_CHANNELS_SHIFT;
 	// Flip the lowest channel bit, it was stored inverted so we can clear stencil to 0 as a default
