Fixed debug view mode having incorrect transform

diff --git a/Engine/Shaders/DebugViewModeCommon.usf b/Engine/Shaders/DebugViewModeCommon.usf
--- a/Engine/Shaders/DebugViewModeCommon.usf
+++ b/Engine/Shaders/DebugViewModeCommon.usf
@@ -14,12 +14,18 @@ struct FDebugVSToPS
 	// TexCoords must first in order to be compatible with viewmode requiring using only SV_Position
 	float4 Position : SV_POSITION;
 	float2 TexCoords : TEXCOORD0;
+#if INSTANCED_STEREO
+	nointerpolation uint EyeIndex : EYE_INDEX;
+#endif
 };
 
 struct FDebugPSIn
 {
 	float4 SvPosition : SV_POSITION;
 	float2 TexCoords : TEXCOORD0;
+#if INSTANCED_STEREO
+	nointerpolation uint EyeIndex : EYE_INDEX;
+#endif
 	uint SvPrimitiveID : SV_PrimitiveID;
 };
 
diff --git a/Engine/Shaders/DebugViewModeVertexShader.usf b/Engine/Shaders/DebugViewModeVertexShader.usf
--- a/Engine/Shaders/DebugViewModeVertexShader.usf
+++ b/Engine/Shaders/DebugViewModeVertexShader.usf
@@ -45,7 +45,7 @@
 		ISOLATE
 		{
 
-			O.Position = mul(WorldPosition, View.TranslatedWorldToClip);
+			O.Position = mul(WorldPosition, ResolvedView.TranslatedWorldToClip);
 		}
 
 #if NUM_MATERIAL_TEXCOORDS > 0
@@ -54,6 +54,9 @@
 		O.TexCoords = 0;
 #endif
 
+#if INSTANCED_STEREO
+		O.EyeIndex = Interpolants.FactoryInterpolants.EyeIndex;
+#endif
 		
 		return O;
 	}
@@ -67,8 +70,21 @@ void Main(
 	FVertexFactoryInput Input,
 	OPTIONAL_VertexID
 	out FDebugVSOutput Output
+#if INSTANCED_STEREO
+	, uint InstanceId : SV_InstanceID
+	, out float4 OutClipDistance : SV_ClipDistance
+	, out float4 OutCullDistance : SV_CullDistance
+#endif
 	)
 {
+#if INSTANCED_STEREO
+	OutCullDistance = OutClipDistance = 0.0;
+	const uint EyeIndex = VertexFactoryGetEyeIndex(InstanceId);
+	ResolvedView = ResolveView(EyeIndex);
+#else
+	ResolvedView = ResolveView();
+#endif
+
 	FVertexFactoryIntermediates VFIntermediates = GetVertexFactoryIntermediates(Input);
 	float4 WorldPos = VertexFactoryGetWorldPosition(Input, VFIntermediates);
 
@@ -87,19 +103,41 @@ void Main(
 	// Transformation is done in Domain shader when tessellating
 	Output.Position = WorldPos;
 	Output.FactoryInterpolants = VertexFactoryGetInterpolantsVSToDS(Input, VFIntermediates, VertexParameters);
+#if INSTANCED_STEREO
+	Output.FactoryInterpolants.EyeIndex = EyeIndex;
+#endif
+
 #else
 	ISOLATE
 	{
 		float4 RasterizedWorldPosition = VertexFactoryGetRasterizedWorldPosition(Input, VFIntermediates, WorldPos);
-		Output.Position = mul(RasterizedWorldPosition, View.TranslatedWorldToClip);
-	}
+		Output.Position = mul(RasterizedWorldPosition, ResolvedView.TranslatedWorldToClip);
 
+		#if INSTANCED_STEREO
+		BRANCH 
+		if (bIsInstancedStereo)  
+		{
+			// Clip at the center of the screen
+			OutCullDistance.x = OutClipDistance.x = dot(Output.Position, EyeClipEdge[EyeIndex]);
+
+			// Scale to the width of a single eye viewport
+			Output.Position.x *= 0.5 * Frame.HMDEyePaddingOffset;
+
+			// Shift to the eye viewport
+			Output.Position.x += (EyeOffsetScale[EyeIndex] * Output.Position.w) / Frame.HMDEyePaddingOffset;
+		}
+		#endif
+	}
 #if NUM_MATERIAL_TEXCOORDS > 0
 	Output.TexCoords = VertexParameters.TexCoords[0];
 #else
 	Output.TexCoords = 0;
 #endif
 
+#if INSTANCED_STEREO
+	Output.EyeIndex = EyeIndex;
+#endif
+
 #endif
 
 	OutputVertexID( Output );
diff --git a/Engine/Shaders/TexelFactorAccuracyPixelShader.usf b/Engine/Shaders/TexelFactorAccuracyPixelShader.usf
--- a/Engine/Shaders/TexelFactorAccuracyPixelShader.usf
+++ b/Engine/Shaders/TexelFactorAccuracyPixelShader.usf
@@ -18,6 +18,12 @@ void Main(
 	out float4 OutColor : SV_Target0
 	)
 {
+#if INSTANCED_STEREO
+	ResolvedView = ResolveView(DebugInputs.EyeIndex);
+#else
+	ResolvedView = ResolveView();
+#endif
+
 	float3 Result = float3(.1, .1, .1);
 
 	if (CPUTexelFactor >= 0)
diff --git a/Engine/Shaders/WantedMipsAccuracyPixelShader.usf b/Engine/Shaders/WantedMipsAccuracyPixelShader.usf
--- a/Engine/Shaders/WantedMipsAccuracyPixelShader.usf
+++ b/Engine/Shaders/WantedMipsAccuracyPixelShader.usf
@@ -18,6 +18,12 @@ void Main(
 	out float4 OutColor : SV_Target0
 	)
 {
+#if INSTANCED_STEREO
+	ResolvedView = ResolveView(DebugInputs.EyeIndex);
+#else
+	ResolvedView = ResolveView();
+#endif
+
 	float3 Result = float3(.1, .1, .1);
 
 	if (CPUWantedMips >= 0)
