Copying //UE4/Dev-Rendering to Dev-Main (//UE4/Dev-Main)

diff --git a/Engine/Config/BaseScalability.ini b/Engine/Config/BaseScalability.ini
--- a/Engine/Config/BaseScalability.ini
+++ b/Engine/Config/BaseScalability.ini
@@ -215,8 +215,9 @@ r.SceneColorFormat=3
 r.DetailMode=0
 r.TranslucencyVolumeBlur=0
 r.MaterialQualityLevel=0
-r.SSS.Scale = 0
+r.SSS.Scale=0
 r.SSS.SampleSet=0
+r.SSS.Quality=0
 r.EmitterSpawnRateScale=0.125
 
 [EffectsQuality@1]
@@ -227,8 +228,9 @@ r.SceneColorFormat=3
 r.DetailMode=1
 r.TranslucencyVolumeBlur=0
 r.MaterialQualityLevel=1
-r.SSS.Scale = 0.75
+r.SSS.Scale=0.75
 r.SSS.SampleSet=0
+r.SSS.Quality=0
 r.EmitterSpawnRateScale=0.25
 
 [EffectsQuality@2]
@@ -239,8 +241,9 @@ r.SceneColorFormat=3
 r.DetailMode=1
 r.TranslucencyVolumeBlur=1
 r.MaterialQualityLevel=1
-r.SSS.Scale = 1
+r.SSS.Scale=1
 r.SSS.SampleSet=1
+r.SSS.Quality=-1
 r.EmitterSpawnRateScale=0.5
 
 [EffectsQuality@3]
@@ -251,6 +254,7 @@ r.SceneColorFormat=4
 r.DetailMode=2
 r.TranslucencyVolumeBlur=1
 r.MaterialQualityLevel=1
-r.SSS.Scale = 1
+r.SSS.Scale=1
 r.SSS.SampleSet=2
+r.SSS.Quality=1
 r.EmitterSpawnRateScale=1.0
\ No newline at end of file
diff --git a/Engine/Shaders/AtmosphereCommon.usf b/Engine/Shaders/AtmosphereCommon.usf
--- a/Engine/Shaders/AtmosphereCommon.usf
+++ b/Engine/Shaders/AtmosphereCommon.usf
@@ -72,12 +72,24 @@ const static int InscatterNuNum = 8;
 #endif
 
 /** Textures for rendering */
-Texture2D AtmosphereTransmittanceTexture;
-SamplerState AtmosphereTransmittanceTextureSampler;
-Texture2D AtmosphereIrradianceTexture;
-SamplerState AtmosphereIrradianceTextureSampler;
-Texture3D AtmosphereInscatterTexture;
-SamplerState AtmosphereInscatterTextureSampler;
+#if IS_MATERIAL_SHADER
+    // for materials, these are in the frame UB
+	#define AtmosphereTransmittanceTexture			Frame.AtmosphereTransmittanceTexture_UB
+	#define AtmosphereTransmittanceTextureSampler	Frame.AtmosphereTransmittanceTextureSampler_UB
+	#define AtmosphereIrradianceTexture				Frame.AtmosphereIrradianceTexture_UB
+	#define AtmosphereIrradianceTextureSampler		Frame.AtmosphereIrradianceTextureSampler_UB
+	#define AtmosphereInscatterTexture				Frame.AtmosphereInscatterTexture_UB
+	#define AtmosphereInscatterTextureSampler		Frame.AtmosphereInscatterTextureSampler_UB
+#else
+	// these are only used for the precomputation shader; which doesn't have a frame UB
+	Texture2D AtmosphereTransmittanceTexture;
+	SamplerState AtmosphereTransmittanceTextureSampler;
+	Texture2D AtmosphereIrradianceTexture;
+	SamplerState AtmosphereIrradianceTextureSampler;
+	Texture3D AtmosphereInscatterTexture;
+	SamplerState AtmosphereInscatterTextureSampler;
+	#define ATMOSPHERIC_TEXTURE_PREFIX
+#endif
 
 struct FAtmosphereVSOutput
 {
diff --git a/Engine/Shaders/Common.usf b/Engine/Shaders/Common.usf
--- a/Engine/Shaders/Common.usf
+++ b/Engine/Shaders/Common.usf
@@ -197,6 +197,19 @@
 #define STENCIL_COMPONENT_SWIZZLE .g
 #endif
 
+#if !TEX_COORD_SCALE_ANALYSIS
+// Texture coordinate scale analysis is used by the texture streamer to compute material scale factors (relatively to UV0) when sampling textures.
+// The pratical implementation is in TexCoordScaleAnalysisPixelShader even though the code is inserted in all shaders, using the stubs above.
+MaterialFloat StoreTexCoordScale(float2 UV, uint TextureReferenceIndex)
+{
+	return 1.f;
+}
+MaterialFloat StoreTexSample(float4 C, uint TextureReferenceIndex)
+{
+	return 1.f;
+}
+#endif
+
 // Whether the platform supports independent texture and samplers
 // When enabled, different texture lookups can share samplers to allow more artist samplers in the base pass
 // Ideally this would just be enabled for all SM4 and above feature level platforms
@@ -1101,6 +1114,11 @@ float2 SvPositionToViewportUV(float4 SvPosition)
 	return PixelPos.xy * Frame.ViewSizeAndInvSize.zw;
 }
 
+float2 BufferUVToViewportUV(float2 BufferUV)
+{
+	float2 PixelPos = BufferUV.xy * Frame.BufferSizeAndInvSize.xy - Frame.ViewRectMin.xy;
+	return PixelPos.xy * Frame.ViewSizeAndInvSize.zw;
+}
 // ----------------------------
 
 /** 
diff --git a/Engine/Shaders/DebugViewModeCommon.usf b/Engine/Shaders/DebugViewModeCommon.usf
--- a/Engine/Shaders/DebugViewModeCommon.usf
+++ b/Engine/Shaders/DebugViewModeCommon.usf
@@ -1,4 +1,4 @@
-// Copyright 1998-2015 Epic Games, Inc. All Rights Reserved.
+// Copyright 1998-2016 Epic Games, Inc. All Rights Reserved.
 
 /*=============================================================================
 	DebugViewModeCommon.hlsl: Debug pixel shader interpolants.
@@ -13,7 +13,11 @@ struct FDebugVSToPS
 {
 	// TexCoords must first in order to be compatible with viewmode requiring using only SV_Position
 	float4 Position : SV_POSITION;
-	float2 TexCoords : TEXCOORD0;
+	float4 VertexColor : TEXCOORD0;
+	float2 TexCoord0 : TEXCOORD1;
+	float3 TangentToWorld0 : TEXCOORD2;
+	float3 TangentToWorld1 : TEXCOORD3;
+	float3 TangentToWorld2 : TEXCOORD4;
 #if INSTANCED_STEREO
 	nointerpolation uint EyeIndex : EYE_INDEX;
 #endif
@@ -22,7 +26,11 @@ struct FDebugVSToPS
 struct FDebugPSIn
 {
 	float4 SvPosition : SV_POSITION;
-	float2 TexCoords : TEXCOORD0;
+	float4 VertexColor : TEXCOORD0;
+	float2 TexCoord0 : TEXCOORD1;
+	float3 TangentToWorld0 : TEXCOORD2;
+	float3 TangentToWorld1 : TEXCOORD3;
+	float3 TangentToWorld2 : TEXCOORD4;
 #if INSTANCED_STEREO
 	nointerpolation uint EyeIndex : EYE_INDEX;
 #endif
@@ -35,3 +43,23 @@ struct FDebugPSInLean
 	float4 SvPosition : SV_POSITION;
 	uint PrimitiveID : SV_PrimitiveID;
 };
+
+#if DEBUG_MATERIAL_PARAMETERS
+/** Converts from vertex factory specific interpolants to a FMaterialPixelParameters, which is used by material inputs. */
+FMaterialPixelParameters GetMaterialPixelParameters(FDebugPSIn DebugInputs, float4 SvPosition)
+{
+	// GetMaterialPixelParameters is responsible for fully initializing the result
+	FMaterialPixelParameters Result = MakeInitializedMaterialPixelParameters();
+
+	Result.VertexColor = DebugInputs.VertexColor;
+
+	Result.TangentToWorld = half3x3(half3(DebugInputs.TangentToWorld0), half3(DebugInputs.TangentToWorld1), half3(DebugInputs.TangentToWorld2));
+
+#if NUM_MATERIAL_TEXCOORDS > 0
+	Result.TexCoords[0] = DebugInputs.TexCoord0;
+#endif
+
+	Result.TwoSidedSign = 1;
+	return Result;
+}
+#endif
\ No newline at end of file
diff --git a/Engine/Shaders/DebugViewModeVertexShader.usf b/Engine/Shaders/DebugViewModeVertexShader.usf
--- a/Engine/Shaders/DebugViewModeVertexShader.usf
+++ b/Engine/Shaders/DebugViewModeVertexShader.usf
@@ -30,35 +30,42 @@
 	
 	FDebugVSToDS PassInterpolate(FDebugVSToDS a, float aInterp, FDebugVSToDS b, float bInterp)
 	{
-		FDebugVSToDS O;
+		FDebugVSToDS Output;
 		
-		O.FactoryInterpolants = VertexFactoryInterpolate(a.FactoryInterpolants, aInterp, b.FactoryInterpolants, bInterp);
+		Output.FactoryInterpolants = VertexFactoryInterpolate(a.FactoryInterpolants, aInterp, b.FactoryInterpolants, bInterp);
 
-		return O;
+		return Output;
 	}
 
 	FDebugVSToPS PassFinalizeTessellationOutput(FDebugVSToDS Interpolants, float4 WorldPosition, FMaterialTessellationParameters MaterialParameters)
 	{
-		FDebugVSToPS O;
+		FDebugVSToPS Output;
 		
 		// Finally, transform position to clip-space
 		ISOLATE
 		{
 
-			O.Position = mul(WorldPosition, ResolvedView.TranslatedWorldToClip);
+			Output.Position = mul(WorldPosition, ResolvedView.TranslatedWorldToClip);
 		}
 
+		Output.VertexColor = MaterialParameters.VertexColor;
+
+		float3x3 TangentToWorld = MaterialParameters.TangentToWorld;
+		Output.TangentToWorld0 = TangentToWorld[0];
+		Output.TangentToWorld1 = TangentToWorld[1];
+		Output.TangentToWorld2 = TangentToWorld[2];
+
 #if NUM_MATERIAL_TEXCOORDS > 0
-		O.TexCoords = MaterialParameters.TexCoords[0];
+		Output.TexCoord0 = MaterialParameters.TexCoords[0];
 #else
-		O.TexCoords = 0;
+		Output.TexCoord0 = 0;
 #endif
 
 #if INSTANCED_STEREO
-		O.EyeIndex = Interpolants.FactoryInterpolants.EyeIndex;
+		Output.EyeIndex = Interpolants.FactoryInterpolants.EyeIndex;
 #endif
-		
-		return O;
+
+		return Output;
 	}
 	
 	#include "Tessellation.usf"
@@ -130,10 +137,18 @@ void Main(
 		}
 		#endif
 	}
+
+	Output.VertexColor = VertexParameters.VertexColor;
+
+	float3x3 TangentToWorld = VertexParameters.TangentToWorld;
+	Output.TangentToWorld0 = float3(TangentToWorld[0]);
+	Output.TangentToWorld1 = float3(TangentToWorld[1]);
+	Output.TangentToWorld2 = float3(TangentToWorld[2]);
+
 #if NUM_MATERIAL_TEXCOORDS > 0
-	Output.TexCoords = VertexParameters.TexCoords[0];
+	Output.TexCoord0 = VertexParameters.TexCoords[0];
 #else
-	Output.TexCoords = 0;
+	Output.TexCoord0 = 0;
 #endif
 
 #if INSTANCED_STEREO
diff --git a/Engine/Shaders/Definitions.usf b/Engine/Shaders/Definitions.usf
--- a/Engine/Shaders/Definitions.usf
+++ b/Engine/Shaders/Definitions.usf
@@ -184,6 +184,10 @@
 #define INSTANCED_STEREO								0
 #endif	
 
+#ifndef TEX_COORD_SCALE_ANALYSIS
+#define TEX_COORD_SCALE_ANALYSIS						0		
+#endif
+
 #if SM5_PROFILE || COMPILER_SUPPORTS_ATTRIBUTES
 /** Avoids flow control constructs. */
 #define UNROLL											[unroll]
diff --git a/Engine/Shaders/GlobalDistanceFieldShared.usf b/Engine/Shaders/GlobalDistanceFieldShared.usf
--- a/Engine/Shaders/GlobalDistanceFieldShared.usf
+++ b/Engine/Shaders/GlobalDistanceFieldShared.usf
@@ -14,18 +14,35 @@ float MaxGlobalDistance;
 // Must match C++
 #define MAX_GLOBAL_DF_CLIPMAPS 4
 
-Texture3D GlobalDistanceFieldTexture0;
-Texture3D GlobalDistanceFieldTexture1;
-Texture3D GlobalDistanceFieldTexture2;
-Texture3D GlobalDistanceFieldTexture3;
-SamplerState GlobalDistanceFieldSampler0;
-SamplerState GlobalDistanceFieldSampler1;
-SamplerState GlobalDistanceFieldSampler2;
-SamplerState GlobalDistanceFieldSampler3;
-float4 GlobalVolumeCenterAndExtent[MAX_GLOBAL_DF_CLIPMAPS];
-float4 GlobalVolumeWorldToUVAddAndMul[MAX_GLOBAL_DF_CLIPMAPS];
-float GlobalVolumeDimension;
-float GlobalVolumeTexelSize;
+#if IS_MATERIAL_SHADER
+    // for materials, these are in the frame UB
+	#define GlobalDistanceFieldTexture0		Frame.GlobalDistanceFieldTexture0_UB
+	#define GlobalDistanceFieldTexture1		Frame.GlobalDistanceFieldTexture1_UB
+	#define GlobalDistanceFieldTexture2		Frame.GlobalDistanceFieldTexture2_UB
+	#define GlobalDistanceFieldTexture3		Frame.GlobalDistanceFieldTexture3_UB
+	#define GlobalDistanceFieldSampler0		Frame.GlobalDistanceFieldSampler0_UB
+	#define GlobalDistanceFieldSampler1		Frame.GlobalDistanceFieldSampler1_UB
+	#define GlobalDistanceFieldSampler2		Frame.GlobalDistanceFieldSampler2_UB
+	#define GlobalDistanceFieldSampler3		Frame.GlobalDistanceFieldSampler3_UB
+	#define GlobalVolumeCenterAndExtent		Frame.GlobalVolumeCenterAndExtent_UB
+	#define GlobalVolumeWorldToUVAddAndMul	Frame.GlobalVolumeWorldToUVAddAndMul_UB
+	#define GlobalVolumeDimension			Frame.GlobalVolumeDimension_UB
+	#define GlobalVolumeTexelSize			Frame.GlobalVolumeTexelSize_UB
+#else
+	// these are only used for the precomputation shaders; which don't have a frame UB
+	Texture3D GlobalDistanceFieldTexture0;
+	Texture3D GlobalDistanceFieldTexture1;
+	Texture3D GlobalDistanceFieldTexture2;
+	Texture3D GlobalDistanceFieldTexture3;
+	SamplerState GlobalDistanceFieldSampler0;
+	SamplerState GlobalDistanceFieldSampler1;
+	SamplerState GlobalDistanceFieldSampler2;
+	SamplerState GlobalDistanceFieldSampler3;
+	float4 GlobalVolumeCenterAndExtent[MAX_GLOBAL_DF_CLIPMAPS];
+	float4 GlobalVolumeWorldToUVAddAndMul[MAX_GLOBAL_DF_CLIPMAPS];
+	float GlobalVolumeDimension;
+	float GlobalVolumeTexelSize;
+#endif
 
 float4 SampleGlobalDistanceField(int ClipmapIndex, float3 UV)
 {
diff --git a/Engine/Shaders/GpuSkinCacheComputeShader.usf b/Engine/Shaders/GpuSkinCacheComputeShader.usf
--- a/Engine/Shaders/GpuSkinCacheComputeShader.usf
+++ b/Engine/Shaders/GpuSkinCacheComputeShader.usf
@@ -9,6 +9,9 @@
 #include "VertexFactoryCommon.usf"
 #include "GpuSkinCommon.usf"
 
+// 0/1 setup by CPU
+// #define GPUSKIN_MORPH_BLEND
+
 // Set to 1 to enable passing through values from vertex buffer (ie do not skin)
 #define GPUSKIN_DEBUG_PASSTHROUGH	0
 
@@ -16,20 +19,30 @@
 
 Buffer<float4> BoneMatrices;
 
-struct FVertexFactoryIntermediates
-{
-	// Blend Matrix (used for position/tangents)
-	FBoneMatrix BlendMatrix;
-
-	// Unpacked position
-	float3 UnpackedPosition;
+#if GPUSKIN_MORPH_BLEND
+// float3 DeltaPosition, PackedNormal, see FMorphGPUSkinVertex
+Buffer<float> MorphInputStreamBuffer;
+// data offset to start indexing the data in MorphInputStreamBuffer, in float units
+uint MorphInputStreamFloatOffset;
 	
-	// Tangent Basis
-	float3x3 TangentToLocal; 
-
-	// Vertex Color
-	float4 Color;
-};
+void GetMorphInputStreamBuffer(uint VertexIndex, out float3 DeltaPosition, out half3 DeltaTangentZ)
+{
+	// see CPU code: "check(MorphStride == sizeof(float) * 6);"
+	// in floats
+	// 3 + 3 floats because of the format defined in InitMorphVertexFactoryComponents / FMorphGPUSkinVertex
+	uint Offset = MorphInputStreamFloatOffset + VertexIndex * (3 + 3);
+
+	DeltaPosition = float3(
+		MorphInputStreamBuffer[Offset + 0],
+		MorphInputStreamBuffer[Offset + 1],
+		MorphInputStreamBuffer[Offset + 2]);
+
+	DeltaTangentZ = float3(
+		MorphInputStreamBuffer[Offset + 3],
+		MorphInputStreamBuffer[Offset + 4],
+		MorphInputStreamBuffer[Offset + 5]);
+}
+#endif
 
 struct FVertexUnpacked
 {
@@ -44,6 +57,28 @@ struct FVertexUnpacked
 #if GPUSKIN_USE_EXTRA_INFLUENCES
 	float4	BlendWeights2;
 #endif
+
+#if GPUSKIN_MORPH_BLEND
+	// morph target, added to the Position
+	float3	DeltaPosition;
+	// morph target, added to the TangentZ and then used to derive new TangentX and TangentY, -2..2
+	half3	DeltaTangentZ;
+#endif
+};
+
+struct FVertexFactoryIntermediates
+{
+	// Blend Matrix (used for position/tangents)
+	FBoneMatrix BlendMatrix;
+
+	// Unpacked position (includes DeltaPosition if GPUSKIN_MORPH_BLEND)
+	float3 UnpackedPosition;
+	
+	// Tangent Basis (includes DeltaTangentZ if GPUSKIN_MORPH_BLEND)
+	float3x3 TangentToLocal; 
+
+	// Vertex Color
+	float4 Color;
 };
 
 FBoneMatrix CalcBoneMatrix( FVertexUnpacked VertexUnpacked )
@@ -52,12 +87,12 @@ FBoneMatrix CalcBoneMatrix( FVertexUnpacked VertexUnpacked )
 	BoneMatrix += VertexUnpacked.BlendWeights.y * FBoneMatrix(BoneMatrices[VertexUnpacked.BlendIndices.y*3], BoneMatrices[VertexUnpacked.BlendIndices.y*3+1], BoneMatrices[VertexUnpacked.BlendIndices.y*3+2]);
 	BoneMatrix += VertexUnpacked.BlendWeights.z * FBoneMatrix(BoneMatrices[VertexUnpacked.BlendIndices.z*3], BoneMatrices[VertexUnpacked.BlendIndices.z*3+1], BoneMatrices[VertexUnpacked.BlendIndices.z*3+2]);
 	BoneMatrix += VertexUnpacked.BlendWeights.w * FBoneMatrix(BoneMatrices[VertexUnpacked.BlendIndices.w*3], BoneMatrices[VertexUnpacked.BlendIndices.w*3+1], BoneMatrices[VertexUnpacked.BlendIndices.w*3+2]);
-	#if GPUSKIN_USE_EXTRA_INFLUENCES
+#if GPUSKIN_USE_EXTRA_INFLUENCES
 	BoneMatrix += VertexUnpacked.BlendWeights2.x * FBoneMatrix(BoneMatrices[VertexUnpacked.BlendIndices2.x*3], BoneMatrices[VertexUnpacked.BlendIndices2.x*3+1], BoneMatrices[VertexUnpacked.BlendIndices2.x*3+2]);
 	BoneMatrix += VertexUnpacked.BlendWeights2.y * FBoneMatrix(BoneMatrices[VertexUnpacked.BlendIndices2.y*3], BoneMatrices[VertexUnpacked.BlendIndices2.y*3+1], BoneMatrices[VertexUnpacked.BlendIndices2.y*3+2]);
 	BoneMatrix += VertexUnpacked.BlendWeights2.z * FBoneMatrix(BoneMatrices[VertexUnpacked.BlendIndices2.z*3], BoneMatrices[VertexUnpacked.BlendIndices2.z*3+1], BoneMatrices[VertexUnpacked.BlendIndices2.z*3+2]);
 	BoneMatrix += VertexUnpacked.BlendWeights2.w * FBoneMatrix(BoneMatrices[VertexUnpacked.BlendIndices2.w*3], BoneMatrices[VertexUnpacked.BlendIndices2.w*3+1], BoneMatrices[VertexUnpacked.BlendIndices2.w*3+2]);
-	#endif
+#endif
 	return BoneMatrix;
 }
 
@@ -76,7 +111,7 @@ float3 MeshExtension;
 
 //uint4 DebugParameter;
 
-FVertexUnpacked UnpackedVertex( uint Offset )
+FVertexUnpacked UnpackedVertex(uint VertexIndex, uint Offset )
 {
 	FVertexUnpacked Unpacked;
 
@@ -113,6 +148,9 @@ FVertexUnpacked UnpackedVertex( uint Offset )
 	Unpacked.BlendWeights2.y = float(BlendWeights >> 8 & 0xff) / 255.0f;
 	Unpacked.BlendWeights2.z = float(BlendWeights >> 16 & 0xff) / 255.0f;
 	Unpacked.BlendWeights2.w = float(BlendWeights >> 24 & 0xff) / 255.0f;
+#endif
+#if GPUSKIN_MORPH_BLEND
+	GetMorphInputStreamBuffer(VertexIndex, Unpacked.DeltaPosition, Unpacked.DeltaTangentZ);
 #endif
 	return Unpacked;
 }
@@ -132,22 +170,23 @@ void SkinCacheUpdateBatchCS(uint3 GroupID : SV_GroupID,
  uint3 DispatchThreadID : SV_DispatchThreadID,
  uint3 GroupThreadID : SV_GroupThreadID)
 {
-	uint ThreadId = DispatchThreadID.x;
-	if (ThreadId >= SkinCacheVertexCount)
+	uint VertexIndex = DispatchThreadID.x;
+	if (VertexIndex >= SkinCacheVertexCount)
 	{
 		return;
 	}
 
 	uint FloatsPerVertex = SkinCacheVertexStride / 4;
-	uint VBOffset = ThreadId * FloatsPerVertex;
-	uint RWOffset = ThreadId * GPUSKIN_RWBUFFER_NUM_FLOATS;
+	uint VBOffset = VertexIndex * FloatsPerVertex;
+	uint RWOffset = VertexIndex * GPUSKIN_RWBUFFER_NUM_FLOATS;
 	uint InputOffset = SkinCacheInputStreamFloatOffset + VBOffset;
 	uint OutputOffset = SkinCacheOutputBufferFloatOffset + RWOffset;
 
-	FVertexUnpacked	Unpacked = UnpackedVertex(InputOffset);
+	FVertexUnpacked	Unpacked = UnpackedVertex(VertexIndex, InputOffset);
 
 	// Perform the skinning
-	FVertexFactoryIntermediates	Intermediates;
+	FVertexFactoryIntermediates	Intermediates = (FVertexFactoryIntermediates)0;
+	
 #if 1
 	Intermediates.BlendMatrix = CalcBoneMatrix(Unpacked);
 #else
@@ -161,10 +200,35 @@ void SkinCacheUpdateBatchCS(uint3 GroupID : SV_GroupID,
 	Intermediates.UnpackedPosition.y = SkinStreamInputBuffer[InputOffset + GPUSKIN_VB_OFFSET_POSITION + 1] * MeshExtension.y + MeshOrigin.y;
 	Intermediates.UnpackedPosition.z = SkinStreamInputBuffer[InputOffset + GPUSKIN_VB_OFFSET_POSITION + 2] * MeshExtension.z + MeshOrigin.z;
 
+	float3 LocalTangentX = Unpacked.TangentX;
+	float3 LocalTangentZ = Unpacked.TangentZ;
+
+#if GPUSKIN_MORPH_BLEND
+	{
+		Intermediates.UnpackedPosition += Unpacked.DeltaPosition;
+
+		// calc new normal by offseting it with the delta
+		LocalTangentZ = normalize( LocalTangentZ + Unpacked.DeltaTangentZ);
+		// derive the new tangent by orthonormalizing the new normal against
+		// the base tangent vector (assuming these are normalized)
+		LocalTangentX = normalize( LocalTangentX - (dot(LocalTangentX, LocalTangentZ) * LocalTangentZ) );
+	}
+#endif
+
 	float3 SPos = SkinPosition(Intermediates);
+		
+	float3 TangentX = mul((float3x3)Intermediates.BlendMatrix, LocalTangentX);
+	float3 TangentZ = mul((float3x3)Intermediates.BlendMatrix, LocalTangentZ);
 
-	float3 TangentX = mul((float3x3)Intermediates.BlendMatrix, Unpacked.TangentX);
-	float3 TangentZ = mul((float3x3)Intermediates.BlendMatrix, Unpacked.TangentZ);
+	SkinCacheBuffer[OutputOffset + GPUSKIN_RWBUFFER_OFFSET_POSITION + 0] = SPos.x;
+	SkinCacheBuffer[OutputOffset + GPUSKIN_RWBUFFER_OFFSET_POSITION + 1] = SPos.y;
+	SkinCacheBuffer[OutputOffset + GPUSKIN_RWBUFFER_OFFSET_POSITION + 2] = SPos.z;
+	// we keep the w component (what is in it?)
+	uint TangentXCompressed = PackS8x3N(TangentX);
+	SkinCacheBuffer[OutputOffset + GPUSKIN_RWBUFFER_OFFSET_TANGENT_X] = asfloat(TangentXCompressed | (asuint(SkinStreamInputBuffer[InputOffset + GPUSKIN_VB_OFFSET_TANGENT_X]) & 0xff000000));
+	// we keep the w component which contains the sign of tangent basis determinant (left or right handed)
+	uint TangentZCompressed = PackS8x3N(TangentZ.xyz);
+	SkinCacheBuffer[OutputOffset + GPUSKIN_RWBUFFER_OFFSET_TANGENT_Z] = asfloat(TangentZCompressed | (asuint(SkinStreamInputBuffer[ InputOffset + GPUSKIN_VB_OFFSET_TANGENT_Z]) & 0xff000000));
 
 #if GPUSKIN_DEBUG_PASSTHROUGH
 	// Passthrough debug code
@@ -173,11 +237,5 @@ void SkinCacheUpdateBatchCS(uint3 GroupID : SV_GroupID,
 	SkinCacheBuffer[OutputOffset + GPUSKIN_RWBUFFER_OFFSET_POSITION + 2] = SkinStreamInputBuffer[InputOffset + GPUSKIN_VB_OFFSET_POSITION + 2];
 	SkinCacheBuffer[OutputOffset + GPUSKIN_RWBUFFER_OFFSET_TANGENT_X] = SkinStreamInputBuffer[InputOffset + GPUSKIN_VB_OFFSET_TANGENT_X];
 	SkinCacheBuffer[OutputOffset + GPUSKIN_RWBUFFER_OFFSET_TANGENT_Z] = SkinStreamInputBuffer[InputOffset + GPUSKIN_VB_OFFSET_TANGENT_Z];
-#else
-	SkinCacheBuffer[OutputOffset + GPUSKIN_RWBUFFER_OFFSET_POSITION + 0] = SPos.x;
-	SkinCacheBuffer[OutputOffset + GPUSKIN_RWBUFFER_OFFSET_POSITION + 1] = SPos.y;
-	SkinCacheBuffer[OutputOffset + GPUSKIN_RWBUFFER_OFFSET_POSITION + 2] = SPos.z;
-	SkinCacheBuffer[OutputOffset + GPUSKIN_RWBUFFER_OFFSET_TANGENT_X] = asfloat(asuint((uint((TangentX.x + 1.0f) * 128.0f) << 0) | (uint((TangentX.y + 1.0f) * 128.0f) << 8) | (uint((TangentX.z + 1.0f) * 128.0f) << 16) | (asuint(SkinStreamInputBuffer[InputOffset + GPUSKIN_VB_OFFSET_TANGENT_X]) & 0xff000000)));
-	SkinCacheBuffer[OutputOffset + GPUSKIN_RWBUFFER_OFFSET_TANGENT_Z] = asfloat(asuint((uint((TangentZ.x + 1.0f) * 128.0f) << 0) | (uint((TangentZ.y + 1.0f) * 128.0f) << 8) | (uint((TangentZ.z + 1.0f) * 128.0f) << 16) | (asuint(SkinStreamInputBuffer[ InputOffset + GPUSKIN_VB_OFFSET_TANGENT_Z]) & 0xff000000)));
 #endif
 }
diff --git a/Engine/Shaders/GpuSkinCommon.usf b/Engine/Shaders/GpuSkinCommon.usf
--- a/Engine/Shaders/GpuSkinCommon.usf
+++ b/Engine/Shaders/GpuSkinCommon.usf
@@ -48,16 +48,7 @@
 #define GPUSKIN_RWBUFFER_NUM_FLOATS			5
 #endif
 
-float4 UnpackS8x4N(uint Packed)
-{
-	float4 Unpacked;
-	Unpacked.x = (float((Packed >> 0) & 0xff) / 127.5f) - 1.0f;
-	Unpacked.y = (float((Packed >> 8) & 0xff) / 127.5f) - 1.0f;
-	Unpacked.z = (float((Packed >> 16) & 0xff) / 127.5f) - 1.0f;
-	Unpacked.w = (float((Packed >> 24) & 0xff) / 127.5f) - 1.0f;
-	return Unpacked;
-}
-
+// @return 0..255
 uint4 UnpackU8x4(uint Packed)
 {
 	uint4 Unpacked;
@@ -68,34 +59,32 @@ uint4 UnpackU8x4(uint Packed)
 	return Unpacked;
 }
 
+// @return 0..1
 float4 UnpackU8x4N(uint Packed)
 {
-	float4 Unpacked;
-	Unpacked.x = (float((Packed >> 0) & 0xff) / 255.0f);
-	Unpacked.y = (float((Packed >> 8) & 0xff) / 255.0f);
-	Unpacked.z = (float((Packed >> 16) & 0xff) / 255.0f);
-	Unpacked.w = (float((Packed >> 24) & 0xff) / 255.0f);
-	return Unpacked;
+	return UnpackU8x4(Packed) / 255.0f;
 }
 
-// Leaves the high 8 bits unset
-uint PackS8x3N(float3 Unpacked)
+// @return -1..1
+float4 UnpackS8x4N(uint Packed)
 {
-	uint Packed = asuint(
-		(uint((Unpacked.x + 1.0f) * 127.5f) << 0) |
-		(uint((Unpacked.y + 1.0f) * 127.5f) << 8) |
-		(uint((Unpacked.z + 1.0f) * 127.5f) << 16))
-		;
-	return Packed;
+	return UnpackU8x4(Packed) / 127.5f - 1;
 }
 
+// @param Unpacked -1..1
 uint PackS8x4N(float4 Unpacked)
 {
 	uint Packed = asuint(
 		(uint((Unpacked.x + 1.0f) * 127.5f) << 0) |
 		(uint((Unpacked.y + 1.0f) * 127.5f) << 8) |
 		(uint((Unpacked.z + 1.0f) * 127.5f) << 16) |
-		(uint((Unpacked.w + 1.0f) * 127.5f) << 24))
-		;
+		(uint((Unpacked.w + 1.0f) * 127.5f) << 24));
 	return Packed;
 }
+
+// @param Unpacked -1..1
+// Leaves the high 8 bits unset
+uint PackS8x3N(float3 Unpacked)
+{
+	return PackS8x4N(float4(Unpacked, 0));
+}
\ No newline at end of file
diff --git a/Engine/Shaders/GpuSkinVertexFactory.usf b/Engine/Shaders/GpuSkinVertexFactory.usf
--- a/Engine/Shaders/GpuSkinVertexFactory.usf
+++ b/Engine/Shaders/GpuSkinVertexFactory.usf
@@ -58,7 +58,9 @@ float ClothBlendWeight;
 struct FVertexFactoryInput
 {
 	float4	Position		: ATTRIBUTE0;
+	// 0..1
 	half3	TangentX		: ATTRIBUTE1;
+	// 0..1
 	// TangentZ.w contains sign of tangent basis determinant
 	half4	TangentZ		: ATTRIBUTE2;	
 
@@ -88,7 +90,7 @@ struct FVertexFactoryInput
 	// NOTE: TEXCOORD6,TEXCOORD7 used instead of POSITION1,NORMAL1 since those semantics are not supported by Cg 
 	/** added to the Position */
 	float3	DeltaPosition	: ATTRIBUTE9;	//POSITION1;
-	/** added to the TangentZ and then used to derive new TangentX,TangentY */
+	/** added to the TangentZ and then used to derive new TangentX,TangentY, .w contains the weight of the tangent blend */
 	float3	DeltaTangentZ	: ATTRIBUTE10;	//NORMAL1;
 #elif GPUSKIN_APEX_CLOTH // exclusive with GPUSKIN_MORPH_BLEND
 	// APEX cloth mesh-mesh mapping data
@@ -264,7 +266,7 @@ float3 ClothingPosition( FVertexFactoryInput Input, bool bPrevious)
 float3 UnpackedPosition( FVertexFactoryInput Input )
 {
 		return float3(Input.Position.xyz*MeshExtension + MeshOrigin);
-	}
+}
 
 #if GPUSKIN_MORPH_BLEND
 /** 
@@ -371,22 +373,21 @@ float3 SkinPreviousPosition( FVertexFactoryInput Input, FVertexFactoryIntermedia
 /** transform the tangent basis vectors */
 float3x3 SkinTangents( FVertexFactoryInput Input, FVertexFactoryIntermediates Intermediates )
 {
-	float2x3 Tangents;
 	float3x3 TangentToLocal;
-#if GPUSKIN_MORPH_BLEND
+
 	// tangent
-	float3 TangentX = TangentBias(Input.TangentX);
-	// normal
-	float4 TangentZ = TangentBias(Input.TangentZ);
+	// -1..1
+	half3 LocalTangentX = TangentBias(Input.TangentX);
+	// -1..1 .xyz:normal, .w:contains sign of tangent basis determinant (left or right handed)
+	half4 LocalTangentZ = TangentBias(Input.TangentZ);
 
+#if GPUSKIN_MORPH_BLEND
 	// calc new normal by offseting it with the delta
-	Tangents[1] = normalize( TangentZ.xyz + Input.DeltaTangentZ.xyz );
+	LocalTangentZ.xyz = normalize(LocalTangentZ.xyz + Input.DeltaTangentZ);
 	// derive the new tangent by orthonormalizing the new normal against
 	// the base tangent vector (assuming these are normalized)
-	Tangents[0] = normalize( TangentX - (dot(TangentX,Tangents[1]) * Tangents[1]) );
+	LocalTangentX = normalize(LocalTangentX - dot(LocalTangentX, LocalTangentZ.xyz) * LocalTangentZ.xyz);
 #else
-	half4 TangentZ = TangentBias(Input.TangentZ);
-	
 	#if GPUSKIN_APEX_CLOTH
 
 	float3x3 ClothTangent = float3x3(1,0,0, 0,1,0, 0,0,1);
@@ -422,30 +423,25 @@ float3x3 SkinTangents( FVertexFactoryInput Input, FVertexFactoryIntermediates In
 
 		// derive the new binormal by getting the cross product of the normal and tangent
 		// and flip vector based on sign of tangent basis determinant
-		TangentToLocal[1] = cross(TangentToLocal[2], TangentToLocal[0]) * TangentZ.w;
+		TangentToLocal[1] = cross(TangentToLocal[2], TangentToLocal[0]) * LocalTangentZ.w;
 	
 		ClothTangent = TangentToLocal;
 	}
 	#endif // GPUSKIN_APEX_CLOTH
 
-	// pass-thru the tangent
-	Tangents[0] = TangentBias(Input.TangentX);
-	// pass-thru the normal
-	Tangents[1] = float3(TangentZ.x,TangentZ.y,TangentZ.z);
-
 #endif // GPUSKIN_MORPH_BLEND
 
 	// Note the use of mul(Matrix,Vector), bone matrices are stored transposed
 	// for tighter packing.
-	TangentToLocal[0] = mul( Intermediates.BlendMatrix, float4(Tangents[0],0) );
-	TangentToLocal[2] = mul( Intermediates.BlendMatrix, float4(Tangents[1],0) );
+	TangentToLocal[0] = mul( Intermediates.BlendMatrix, float4(LocalTangentX, 0) );
+	TangentToLocal[2] = mul( Intermediates.BlendMatrix, float4(LocalTangentZ.xyz, 0) );
 
 	// derive the new binormal by getting the cross product of the normal and tangent
 	// and flip vector based on sign of tangent basis determinant
-	TangentToLocal[1] = cross(TangentToLocal[2], TangentToLocal[0]) * TangentZ.w;
+	TangentToLocal[1] = cross(TangentToLocal[2], TangentToLocal[0]) * LocalTangentZ.w;
 
 #if GPUSKIN_APEX_CLOTH
-	return ClothTangent*TempClothBlendWeight + TangentToLocal*(1.0f - TempClothBlendWeight);
+	return ClothTangent * TempClothBlendWeight + TangentToLocal * (1.0f - TempClothBlendWeight);
 #else
 	return TangentToLocal;
 #endif	// GPUSKIN_APEX_CLOTH
diff --git a/Engine/Shaders/MeshParticleVertexFactory.usf b/Engine/Shaders/MeshParticleVertexFactory.usf
--- a/Engine/Shaders/MeshParticleVertexFactory.usf
+++ b/Engine/Shaders/MeshParticleVertexFactory.usf
@@ -11,6 +11,14 @@
 
 float4x4 PreviousLocalToWorld;
 
+#if PARTICLE_MESH_INSTANCED
+	Buffer<float4> PrevTransformBuffer;
+#else
+	float4 PrevTransform0;
+	float4 PrevTransform1;
+	float4 PrevTransform2;
+#endif
+	
 #if !PARTICLE_MESH_INSTANCED
 	float4 Transform1;
 	float4 Transform2;
@@ -58,7 +66,7 @@ struct FVertexFactoryInput
 #endif
 
 /** Optional instance ID for vertex layered rendering */
-#if FEATURE_LEVEL >= FEATURE_LEVEL_SM4 && ONEPASS_POINTLIGHT_SHADOW && USING_VERTEX_SHADER_LAYER
+#if FEATURE_LEVEL > FEATURE_LEVEL_ES2 
 	uint InstanceId	: SV_InstanceID;
 #endif
 };
@@ -150,6 +158,29 @@ float4x4 GetParticleTransform(FVertexFactoryInput Input)
 #endif
 }
 
+float4x4 GetParticlePrevTransform(FVertexFactoryInput Input)
+{
+#if PARTICLE_MESH_INSTANCED
+	uint TransformBaseIdx = Input.InstanceId * 3;
+
+	float4x4 PrevTransform4x4 = float4x4(
+		PrevTransformBuffer[TransformBaseIdx + 0], 
+		PrevTransformBuffer[TransformBaseIdx + 1], 
+		PrevTransformBuffer[TransformBaseIdx + 2], 
+		float4(0, 0, 0, 1)
+		);
+#else
+	float4x4 PrevTransform4x4 = float4x4(
+		PrevTransform0, 
+		PrevTransform1, 
+		PrevTransform2, 
+		float4(0, 0, 0, 1)
+		);
+#endif
+	//Transform is transposed on CPU and position is packed into W. Transpose again to get the correct transform back.
+	return transpose(PrevTransform4x4);
+}
+
 /** Converts from vertex factory specific interpolants to a FMaterialPixelParameters, which is used by material inputs. */
 FMaterialPixelParameters GetMaterialPixelParameters(FVertexFactoryInterpolantsVSToPS Interpolants, float4 SvPosition)
 {
@@ -273,6 +304,20 @@ float4 CalcWorldPosition(FVertexFactoryInput Input)
 	return float4(WorldPosition, Input.Position.w);
 }
 
+float4 CalcPrevWorldPosition(FVertexFactoryInput Input)
+{
+	if (MeshParticleVF.PrevTransformAvailable > 0)
+	{
+		float4x4 PrevTransform = GetParticlePrevTransform(Input);
+		float3 PrevWorldPosition = mul(Input.Position, PrevTransform).xyz + ResolvedView.PrevPreViewTranslation;
+		return float4(PrevWorldPosition, Input.Position.w);
+	}
+	else
+	{
+		return CalcWorldPosition(Input);
+	}
+}
+
 FVertexFactoryIntermediates GetVertexFactoryIntermediates(FVertexFactoryInput Input)
 {
 	FVertexFactoryIntermediates Intermediates = (FVertexFactoryIntermediates)0;
@@ -459,9 +504,8 @@ FVertexFactoryInterpolantsVSToPS VertexFactoryGetInterpolantsVSToPS(FVertexFacto
 
 // @return previous translated world position
 float4 VertexFactoryGetPreviousWorldPosition(FVertexFactoryInput Input, FVertexFactoryIntermediates Intermediates)
-{
-	// The previous local to world instance transform isn't available.
-	return CalcWorldPosition(Input);
+{	
+	return CalcPrevWorldPosition(Input);
 }
 
 #if USING_TESSELLATION
diff --git a/Engine/Shaders/MissingShaderPixelShader.usf b/Engine/Shaders/MissingShaderPixelShader.usf
new file mode 100644
--- /dev/null
+++ b/Engine/Shaders/MissingShaderPixelShader.usf
@@ -0,0 +1,32 @@
+// Copyright 1998-2016 Epic Games, Inc. All Rights Reserved.
+
+/*=============================================================================
+	MissingShaderPixelShader.usf: Outputs the wanted mips.
+=============================================================================*/
+
+#include "Common.usf"
+
+#include "DebugViewModeCommon.usf"
+
+// This returns a check board shader. Usefull to figure out that it has maximum variance.
+
+
+EARLYDEPTHSTENCIL
+void Main(
+	in FDebugPSIn DebugInputs,
+	out float4 OutColor : SV_Target0
+	)
+{
+#if INSTANCED_STEREO
+	ResolvedView = ResolveView(DebugInputs.EyeIndex);
+#else
+	ResolvedView = ResolveView();
+#endif
+
+	int2 CheckBoardCoord = DebugInputs.TexCoord0.xy * 64;
+	float CheckBoard = (CheckBoardCoord.x + CheckBoardCoord.y) % 2;
+
+	float3 Result = float3(CheckBoard, CheckBoard, CheckBoard);
+
+	OutColor = RETURN_COLOR(float4(Result, .25f));
+}
diff --git a/Engine/Shaders/PostProcessSubsurface.usf b/Engine/Shaders/PostProcessSubsurface.usf
--- a/Engine/Shaders/PostProcessSubsurface.usf
+++ b/Engine/Shaders/PostProcessSubsurface.usf
@@ -19,6 +19,10 @@ float4 SSSParams;
 // 0: (testing, could be faster), 1: higher quality
 #define FULLRES_WHERE_POSSIBLE 1
 
+#ifndef RECOMBINE_QUALITY
+	#define RECOMBINE_QUALITY 0
+#endif
+
 // ------------------------------------------
 
 // setup for "SeparableSSS.usf"
@@ -226,7 +230,8 @@ struct SDiffuseAndSpecular
 };
 
 // @param UVSceneColor for the full res rendertarget (BufferSize) e.g. SceneColor or GBuffers
-SDiffuseAndSpecular ReconstructDiffuse(float2 UVSceneColor)
+// @param ReconstructDiffuseMethod 0/1/2/3 (should be a literal constant to allow compiler optimizations)
+SDiffuseAndSpecular ReconstructDiffuse(float2 UVSceneColor, uint ReconstructDiffuseMethod)
 {
 	SDiffuseAndSpecular Ret;
 
@@ -238,40 +243,49 @@ SDiffuseAndSpecular ReconstructDiffuse(float2 UVSceneColor)
 
 	// todo: expose as scalability setting (can be evaluate best without TemporalAA)
 	// 0:fast but pattern can appear, 1:better, 2: even better, 3: best but expensive
+	float3 Quant1;
 
-	// 0/1/2/3 With TemporalAA they all look the same but 0 has an offset which we can fixed by alternating the 2 (L R) or 4 directions (L R U D)
-	#define QUALITY_LEVEL 1
-
-#if QUALITY_LEVEL == 0
-	float3 Quant1 = Texture2DSample(PostprocessInput0, PostprocessInput0Sampler, UVSceneColor + float2(1, 0) * Frame.BufferSizeAndInvSize.zw).rgb;
-#elif QUALITY_LEVEL == 1
-	float3 Quant1 = 0.5f * (
-		Texture2DSample(PostprocessInput0, PostprocessInput0Sampler, UVSceneColor + float2( 1,  0) * Frame.BufferSizeAndInvSize.zw).rgb +
-		Texture2DSample(PostprocessInput0, PostprocessInput0Sampler, UVSceneColor + float2(-1,  0) * Frame.BufferSizeAndInvSize.zw).rgb);
-#elif QUALITY_LEVEL == 2
-	float3 Quant1 = 0.25f * (
-		Texture2DSample(PostprocessInput0, PostprocessInput0Sampler, UVSceneColor + float2( 1,  0) * Frame.BufferSizeAndInvSize.zw).rgb +
-		Texture2DSample(PostprocessInput0, PostprocessInput0Sampler, UVSceneColor + float2( 0,  1) * Frame.BufferSizeAndInvSize.zw).rgb +
-		Texture2DSample(PostprocessInput0, PostprocessInput0Sampler, UVSceneColor + float2(-1,  0) * Frame.BufferSizeAndInvSize.zw).rgb +
-		Texture2DSample(PostprocessInput0, PostprocessInput0Sampler, UVSceneColor + float2( 0, -1) * Frame.BufferSizeAndInvSize.zw).rgb);
-#else // QUALITY_LEVEL == 3
-	float3 A = Texture2DSample(PostprocessInput0, PostprocessInput0Sampler, UVSceneColor + float2( 1,  0) * Frame.BufferSizeAndInvSize.zw).rgb;
-	float3 B = Texture2DSample(PostprocessInput0, PostprocessInput0Sampler, UVSceneColor + float2(-1,  0) * Frame.BufferSizeAndInvSize.zw).rgb;
-	float3 C = Texture2DSample(PostprocessInput0, PostprocessInput0Sampler, UVSceneColor + float2( 0,  1) * Frame.BufferSizeAndInvSize.zw).rgb;
-	float3 D = Texture2DSample(PostprocessInput0, PostprocessInput0Sampler, UVSceneColor + float2( 0, -1) * Frame.BufferSizeAndInvSize.zw).rgb;
-
-	// Luminance could be green channel only
-	float a = Luminance(A);
-	float b = Luminance(B);
-	float c = Luminance(C);
-	float d = Luminance(D);
-
-	float ab = abs(a - b);
-	float cd = abs(c - d);
-
-	// take the average in the direction that avoids dither pattern
-	float3 Quant1 = 0.5f * lerp(A + B, C + D, ab > cd);
-#endif
+	if(ReconstructDiffuseMethod == 0)
+	{
+		// cheap, crappy
+		Quant1 = Texture2DSample(PostprocessInput0, PostprocessInput0Sampler, UVSceneColor + float2(1, 0) * Frame.BufferSizeAndInvSize.zw).rgb;
+	}
+	else if(ReconstructDiffuseMethod == 1)
+	{
+		// acceptable but not perfect
+		Quant1 = 0.5f * (
+			Texture2DSample(PostprocessInput0, PostprocessInput0Sampler, UVSceneColor + float2( 1,  0) * Frame.BufferSizeAndInvSize.zw).rgb +
+			Texture2DSample(PostprocessInput0, PostprocessInput0Sampler, UVSceneColor + float2(-1,  0) * Frame.BufferSizeAndInvSize.zw).rgb);
+	}
+	else if(ReconstructDiffuseMethod == 2)
+	{
+		// almost same as 1?
+		Quant1 = 0.25f * (
+			Texture2DSample(PostprocessInput0, PostprocessInput0Sampler, UVSceneColor + float2( 1,  0) * Frame.BufferSizeAndInvSize.zw).rgb +
+			Texture2DSample(PostprocessInput0, PostprocessInput0Sampler, UVSceneColor + float2( 0,  1) * Frame.BufferSizeAndInvSize.zw).rgb +
+			Texture2DSample(PostprocessInput0, PostprocessInput0Sampler, UVSceneColor + float2(-1,  0) * Frame.BufferSizeAndInvSize.zw).rgb +
+			Texture2DSample(PostprocessInput0, PostprocessInput0Sampler, UVSceneColor + float2( 0, -1) * Frame.BufferSizeAndInvSize.zw).rgb);
+	}
+	else if(ReconstructDiffuseMethod == 3)
+	{
+		// very good 
+		float3 A = Texture2DSample(PostprocessInput0, PostprocessInput0Sampler, UVSceneColor + float2( 1,  0) * Frame.BufferSizeAndInvSize.zw).rgb;
+		float3 B = Texture2DSample(PostprocessInput0, PostprocessInput0Sampler, UVSceneColor + float2(-1,  0) * Frame.BufferSizeAndInvSize.zw).rgb;
+		float3 C = Texture2DSample(PostprocessInput0, PostprocessInput0Sampler, UVSceneColor + float2( 0,  1) * Frame.BufferSizeAndInvSize.zw).rgb;
+		float3 D = Texture2DSample(PostprocessInput0, PostprocessInput0Sampler, UVSceneColor + float2( 0, -1) * Frame.BufferSizeAndInvSize.zw).rgb;
+
+		// Luminance could be green channel only
+		float a = Luminance(A);
+		float b = Luminance(B);
+		float c = Luminance(C);
+		float d = Luminance(D);
+
+		float ab = abs(a - b);
+		float cd = abs(c - d);
+
+		// take the average in the direction that avoids dither pattern
+		Quant1 = 0.5f * lerp(A + B, C + D, ab > cd);
+	}
 
 	Ret.Diffuse = lerp(Quant1, Quant0, bChecker);
 	Ret.Specular = lerp(Quant0, Quant1, bChecker);
@@ -289,7 +303,8 @@ float4 SetupSubsurfaceForOnePixel(float2 UVSceneColor)
 
 	FLATTEN if (ScreenSpaceData.GBuffer.ShadingModelID == SHADINGMODELID_SUBSURFACE_PROFILE)
 	{
-		SDiffuseAndSpecular DiffuseAndSpecular = ReconstructDiffuse(UVSceneColor);
+		// '1' is lower quality but that is acceptable here
+		SDiffuseAndSpecular DiffuseAndSpecular = ReconstructDiffuse(UVSceneColor, 1);
 
 		Ret.rgb = DiffuseAndSpecular.Diffuse;
 
@@ -437,7 +452,9 @@ void SubsurfaceRecombinePS(noperspective float4 UVAndScreenPos : TEXCOORD0, floa
 	float3 StoredBaseColor = ScreenSpaceData.GBuffer.StoredBaseColor;
 	float3 StoredSpecular = ScreenSpaceData.GBuffer.StoredSpecular;
 
-	SDiffuseAndSpecular DiffuseAndSpecular = ReconstructDiffuse(UVSceneColor);
+	uint ReconstructDiffuseMethod = RECOMBINE_QUALITY ? 3 : 1;
+
+	SDiffuseAndSpecular DiffuseAndSpecular = ReconstructDiffuse(UVSceneColor, ReconstructDiffuseMethod);
 	
 	float3 ExtractedNonSubsurface = DiffuseAndSpecular.Specular * StoredSpecular;
 
diff --git a/Engine/Shaders/Random.usf b/Engine/Shaders/Random.usf
--- a/Engine/Shaders/Random.usf
+++ b/Engine/Shaders/Random.usf
@@ -202,14 +202,6 @@ float PerlinNoise3D_TEX(float3 fv)
 }
 */
 
-// 128x128 texture with random directions (12 edge corners)
-Texture2D PerlinNoiseGradientTexture;
-SamplerState PerlinNoiseGradientTextureSampler;
-
-// 16x16x16 texture with random directions (12 edge corners) and precomputed gradient Alpha
-Texture3D PerlinNoise3DTexture;
-SamplerState PerlinNoise3DTextureSampler;
-
 // filtered 3D noise, can be optimized
 // @param v >0
 // @return random number in the range -1 .. 1
@@ -227,14 +219,14 @@ float GradientNoise3D_TEX(float3 fv)
 	float2 TexB = (iv.xy + OffsetB + 0.5f) / 128.0f;
 
 	// can be optimized to 1 or 2 texture lookups (4 or 8 channel encoded in 8, 16 or 32 bit)
-	float3 A = Texture2DSampleLevel(PerlinNoiseGradientTexture, PerlinNoiseGradientTextureSampler, TexA + float2(0, 0) / 128.0f, 0).xyz * 2 - 1;
-	float3 B = Texture2DSampleLevel(PerlinNoiseGradientTexture, PerlinNoiseGradientTextureSampler, TexA + float2(1, 0) / 128.0f, 0).xyz * 2 - 1;
-	float3 C = Texture2DSampleLevel(PerlinNoiseGradientTexture, PerlinNoiseGradientTextureSampler, TexA + float2(0, 1) / 128.0f, 0).xyz * 2 - 1;
-	float3 D = Texture2DSampleLevel(PerlinNoiseGradientTexture, PerlinNoiseGradientTextureSampler, TexA + float2(1, 1) / 128.0f, 0).xyz * 2 - 1;
-	float3 E = Texture2DSampleLevel(PerlinNoiseGradientTexture, PerlinNoiseGradientTextureSampler, TexB + float2(0, 0) / 128.0f, 0).xyz * 2 - 1;
-	float3 F = Texture2DSampleLevel(PerlinNoiseGradientTexture, PerlinNoiseGradientTextureSampler, TexB + float2(1, 0) / 128.0f, 0).xyz * 2 - 1;
-	float3 G = Texture2DSampleLevel(PerlinNoiseGradientTexture, PerlinNoiseGradientTextureSampler, TexB + float2(0, 1) / 128.0f, 0).xyz * 2 - 1;
-	float3 H = Texture2DSampleLevel(PerlinNoiseGradientTexture, PerlinNoiseGradientTextureSampler, TexB + float2(1, 1) / 128.0f, 0).xyz * 2 - 1;
+	float3 A = Texture2DSampleLevel(Frame.PerlinNoiseGradientTexture, Frame.PerlinNoiseGradientTextureSampler, TexA + float2(0, 0) / 128.0f, 0).xyz * 2 - 1;
+	float3 B = Texture2DSampleLevel(Frame.PerlinNoiseGradientTexture, Frame.PerlinNoiseGradientTextureSampler, TexA + float2(1, 0) / 128.0f, 0).xyz * 2 - 1;
+	float3 C = Texture2DSampleLevel(Frame.PerlinNoiseGradientTexture, Frame.PerlinNoiseGradientTextureSampler, TexA + float2(0, 1) / 128.0f, 0).xyz * 2 - 1;
+	float3 D = Texture2DSampleLevel(Frame.PerlinNoiseGradientTexture, Frame.PerlinNoiseGradientTextureSampler, TexA + float2(1, 1) / 128.0f, 0).xyz * 2 - 1;
+	float3 E = Texture2DSampleLevel(Frame.PerlinNoiseGradientTexture, Frame.PerlinNoiseGradientTextureSampler, TexB + float2(0, 0) / 128.0f, 0).xyz * 2 - 1;
+	float3 F = Texture2DSampleLevel(Frame.PerlinNoiseGradientTexture, Frame.PerlinNoiseGradientTextureSampler, TexB + float2(1, 0) / 128.0f, 0).xyz * 2 - 1;
+	float3 G = Texture2DSampleLevel(Frame.PerlinNoiseGradientTexture, Frame.PerlinNoiseGradientTextureSampler, TexB + float2(0, 1) / 128.0f, 0).xyz * 2 - 1;
+	float3 H = Texture2DSampleLevel(Frame.PerlinNoiseGradientTexture, Frame.PerlinNoiseGradientTextureSampler, TexB + float2(1, 1) / 128.0f, 0).xyz * 2 - 1;
 
 	float a = dot(A, Frac - float3(0, 0, 0));
 	float b = dot(B, Frac - float3(1, 0, 0));
@@ -275,7 +267,7 @@ float FastGradientPerlinNoise3D_TEX(float3 xyz)
 	// scalar: 3 add
 	float3 p = p0 + f;
 	// scalar: 3 mad
-	float4 NoiseSample = Texture3DSampleLevel(PerlinNoise3DTexture, PerlinNoise3DTextureSampler, p / Extent + 0.5f / Extent, 0);		// +0.5f to get rid of bilinear offset
+	float4 NoiseSample = Texture3DSampleLevel(Frame.PerlinNoise3DTexture, Frame.PerlinNoise3DTextureSampler, p / Extent + 0.5f / Extent, 0);		// +0.5f to get rid of bilinear offset
 
 	// reconstruct from 8bit (using mad with 2 constants and dot4 was same instruction count)
 	// scalar: 4 mad, 3 mul, 3 add 
@@ -350,7 +342,7 @@ float2 GetPerlinNoiseGradientTextureAt(float2 v)
 	float2 TexA = (v.xy + 0.5f) / 128.0f;
 
 	// todo: storing random 2d unit vectors would be better
-	float3 p = Texture2DSampleLevel(PerlinNoiseGradientTexture, PerlinNoiseGradientTextureSampler, TexA, 0).xyz * 2 - 1;
+	float3 p = Texture2DSampleLevel(Frame.PerlinNoiseGradientTexture, Frame.PerlinNoiseGradientTextureSampler, TexA, 0).xyz * 2 - 1;
 	return normalize(p.xy + p.z * 0.33f);
 }
 
@@ -361,7 +353,7 @@ float3 GetPerlinNoiseGradientTextureAt(float3 v)
 	float2 OffsetA = v.z * ZShear;
 	float2 TexA = (v.xy + OffsetA + 0.5f) / 128.0f;
 
-	return Texture2DSampleLevel(PerlinNoiseGradientTexture, PerlinNoiseGradientTextureSampler, TexA , 0).xyz * 2 - 1;
+	return Texture2DSampleLevel(Frame.PerlinNoiseGradientTexture, Frame.PerlinNoiseGradientTextureSampler, TexA , 0).xyz * 2 - 1;
 }
 
 float2 SkewSimplex(float2 In)
diff --git a/Engine/Shaders/ShaderVersion.usf b/Engine/Shaders/ShaderVersion.usf
--- a/Engine/Shaders/ShaderVersion.usf
+++ b/Engine/Shaders/ShaderVersion.usf
@@ -2,4 +2,4 @@
 // Each time the console command is executed it generates a new GUID. As this file is included
 // in common.usf (which should be included in any shader) it allows to invalidate the shader DDC.
 // 
-// GUID = AC39AE50474FA530A4ADB4B5629A0714
+// GUID = 8743EFC793274990926114600E8C3827  
diff --git a/Engine/Shaders/TexCoordScaleAnalysisPixelShader.usf b/Engine/Shaders/TexCoordScaleAnalysisPixelShader.usf
new file mode 100644
--- /dev/null
+++ b/Engine/Shaders/TexCoordScaleAnalysisPixelShader.usf
@@ -0,0 +1,130 @@
+// Copyright 1998-2016 Epic Games, Inc. All Rights Reserved.
+
+/*=============================================================================
+	TextureStreamingAnalysis.usf: Pixel shader to analyse coordinate scale per texture
+=============================================================================*/
+
+#define TEX_COORD_SCALE_ANALYSIS 1
+#define DEBUG_MATERIAL_PARAMETERS 1
+#define INITIAL_GPU_SCALE 1024
+
+#include "Common.usf"
+#include "SHCommon.usf"
+
+// Scale is 16 bits interger, 16 bits fractionnal.
+const static uint FRACTIONNAL_SCALE = 0xFFFF;
+
+uint FloatToDiscreteScale(float Value)
+{
+	return (uint)(Value * (float)FRACTIONNAL_SCALE); 
+}
+
+RWTexture2D<uint> RWCoordScales : register(u7);
+
+ // If -1, output everything to the RW texture.
+float4 AccuracyColors[5];
+uint TextureAnalysisIndex;
+float CPUScale;
+
+// Global var used to reduce the number of param passed to the shaders.
+uint2 PixelPosition;
+uint CoordScaleHeight;
+uint CoordScaleWidth;
+float OneOverDDX0;
+float OneOverDDY0;
+float GPUScale;
+float TexSample;
+  
+MaterialFloat StoreTexCoordScale(float2 UV, uint TextureReferenceIndex)
+{
+ 	// Take the minimum scale since high scale requires less resolution.
+	float Scale = min(length(ddx(UV)) * OneOverDDX0, length(ddy(UV)) * OneOverDDY0);
+
+	if (TextureAnalysisIndex == TextureReferenceIndex)
+	{
+		GPUScale = min(Scale, GPUScale);
+	} 
+
+	[branch]
+	if (TextureAnalysisIndex == -1) 
+	{
+		InterlockedMin(RWCoordScales[PixelPosition + uint2(TextureReferenceIndex * CoordScaleWidth, 0)], FloatToDiscreteScale(Scale));
+	}
+
+	return 1.f;
+}
+
+MaterialFloat StoreTexSample(float4 C, uint TextureReferenceIndex)
+{
+	if (TextureAnalysisIndex == TextureReferenceIndex)
+	{
+		// Alpha?
+		TexSample = lerp(.25f, 1.f, Luminance(C.rgb));
+	}
+
+	return 1.f;
+}
+
+
+#include "Material.usf"
+#include "DebugViewModeCommon.usf"
+
+void Main(
+	in FDebugPSIn DebugInputs
+	OPTIONAL_IsFrontFace,
+	out float4 OutColor : SV_Target0
+	)
+{
+#if INSTANCED_STEREO
+	ResolvedView = ResolveView(DebugInputs.EyeIndex);
+#else
+	ResolvedView = ResolveView();
+#endif
+
+
+	float3 Result = float3(.1, .1, .1);
+
+	// Set global var used for outputting data
+	PixelPosition = DebugInputs.SvPosition;
+	RWCoordScales.GetDimensions(CoordScaleWidth, CoordScaleHeight);
+	OneOverDDX0 = 1 / length(ddx(DebugInputs.TexCoord0));
+	OneOverDDY0 = 1 / length(ddy(DebugInputs.TexCoord0));
+	GPUScale = INITIAL_GPU_SCALE;
+	TexSample = 1.f;
+
+	FMaterialPixelParameters MaterialParameters = GetMaterialPixelParameters(DebugInputs, DebugInputs.SvPosition);
+	FPixelMaterialInputs PixelMaterialInputs;
+	CalcMaterialParameters(MaterialParameters, PixelMaterialInputs, DebugInputs.SvPosition, bIsFrontFace);
+
+	// Sample material properties. The results are not used, but the calls to StoreTexCoordScale are still be made.
+	half3 BaseColor = GetMaterialBaseColor(PixelMaterialInputs);
+	half  Metallic = GetMaterialMetallic(PixelMaterialInputs);
+	half  Specular = GetMaterialSpecular(PixelMaterialInputs);
+	float Roughness = GetMaterialRoughness(PixelMaterialInputs);
+	half3 Emissive = GetMaterialEmissive(PixelMaterialInputs);
+	half Opacity = GetMaterialOpacity(PixelMaterialInputs);
+	half Mask = GetMaterialMask(PixelMaterialInputs);
+	half4 SSData = GetMaterialSubsurfaceData(MaterialParameters);
+	float Custom0 = GetMaterialCustomData0(MaterialParameters);
+	float Custom1 = GetMaterialCustomData1(MaterialParameters);
+	float MaterialAO = GetMaterialAmbientOcclusion(PixelMaterialInputs);
+	float PixelDepthOffset = GetMaterialPixelDepthOffset(PixelMaterialInputs);
+
+	// Output accuracy
+	if (GPUScale != INITIAL_GPU_SCALE)
+	{
+		float Accuracy = clamp(log2(GPUScale) - log2(CPUScale), -1.99, 1.99);
+		int ColorIndex = floor(Accuracy) + 2;
+
+		Result = TexSample * lerp(AccuracyColors[ColorIndex].rgb, AccuracyColors[ColorIndex + 1].rgb, frac(Accuracy));
+	}
+
+	OutColor = RETURN_COLOR(float4(Result, .25f));
+
+}
+
+
+
+
+
+
diff --git a/Engine/Shaders/TexelFactorAccuracyPixelShader.usf b/Engine/Shaders/TexelFactorAccuracyPixelShader.usf
--- a/Engine/Shaders/TexelFactorAccuracyPixelShader.usf
+++ b/Engine/Shaders/TexelFactorAccuracyPixelShader.usf
@@ -1,4 +1,4 @@
-// Copyright 1998-2015 Epic Games, Inc. All Rights Reserved.
+// Copyright 1998-2016 Epic Games, Inc. All Rights Reserved.
 
 /*=============================================================================
 	TexelFactorAccuracyPixelShader.usf: Outputs the accucary of the texel factor.
@@ -8,7 +8,7 @@
 
 #include "DebugViewModeCommon.usf"
 
-float4 StreamingAccuracyColors[5];
+float4 AccuracyColors[5];
 float CPUTexelFactor; // In world space units.
 
 // Needed because UAV accesses disable early depth test otherwises.
@@ -30,8 +30,8 @@ void Main(
 	{
 		float3 WorldPosition = SvPositionToResolvedTranslatedWorld(DebugInputs.SvPosition);
 
-		float2 CoordDDX = ddx(DebugInputs.TexCoords.xy);
-		float2 CoordDDY = ddy(DebugInputs.TexCoords.xy);
+		float2 CoordDDX = ddx(DebugInputs.TexCoord0.xy);
+		float2 CoordDDY = ddy(DebugInputs.TexCoord0.xy);
 		float3 WorldPosDDX = ddx(WorldPosition);
 		float3 WorldPosDDY = ddy(WorldPosition);
 
@@ -40,7 +40,7 @@ void Main(
 		float Accuracy = clamp(log2(CPUTexelFactor) - log2(GPUTexelFactor), -1.99, 1.99);
 		int ColorIndex = floor(Accuracy) + 2;
 
-		Result = lerp(StreamingAccuracyColors[ColorIndex], StreamingAccuracyColors[ColorIndex + 1], frac(Accuracy));
+		Result = lerp(AccuracyColors[ColorIndex], AccuracyColors[ColorIndex + 1], frac(Accuracy));
 	}
 
 	OutColor = RETURN_COLOR(float4(Result, .25f));
diff --git a/Engine/Shaders/WantedMipsAccuracyPixelShader.usf b/Engine/Shaders/WantedMipsAccuracyPixelShader.usf
--- a/Engine/Shaders/WantedMipsAccuracyPixelShader.usf
+++ b/Engine/Shaders/WantedMipsAccuracyPixelShader.usf
@@ -1,4 +1,4 @@
-// Copyright 1998-2015 Epic Games, Inc. All Rights Reserved.
+// Copyright 1998-2016 Epic Games, Inc. All Rights Reserved.
 
 /*=============================================================================
 	WantedMipsAccuracyPixelShader.usf: Outputs the wanted mips.
@@ -8,7 +8,7 @@
 
 #include "DebugViewModeCommon.usf"
 
-float4 StreamingAccuracyColors[5];
+float4 AccuracyColors[5];
 float CPUWantedMips;
 
 // Needed because UAV accesses disable early depth test otherwises.
@@ -28,8 +28,8 @@ void Main(
 
 	if (CPUWantedMips >= 0)
 	{
-		float2 CoordDDX = ddx(DebugInputs.TexCoords.xy);
-		float2 CoordDDY = ddy(DebugInputs.TexCoords.xy);
+		float2 CoordDDX = ddx(DebugInputs.TexCoord0.xy);
+		float2 CoordDDY = ddy(DebugInputs.TexCoord0.xy);
 
 		float TexelRatio = max(1 / length(CoordDDX), 1 / length(CoordDDY));
 
@@ -39,7 +39,7 @@ void Main(
 		float Accuracy = clamp(CPUWantedMips - GPUWantedMips, -1.99, 1.99);
 		int ColorIndex = floor(Accuracy) + 2;
 
-		Result = lerp(StreamingAccuracyColors[ColorIndex].rgb, StreamingAccuracyColors[ColorIndex + 1].rgb, frac(Accuracy));
+		Result = lerp(AccuracyColors[ColorIndex].rgb, AccuracyColors[ColorIndex + 1].rgb, frac(Accuracy));
 	}
 
 	OutColor = RETURN_COLOR(float4(Result, .25f));
diff --git a/Engine/Source/Developer/SimplygonMeshReduction/Private/SimplygonMeshReduction.cpp b/Engine/Source/Developer/SimplygonMeshReduction/Private/SimplygonMeshReduction.cpp
--- a/Engine/Source/Developer/SimplygonMeshReduction/Private/SimplygonMeshReduction.cpp
+++ b/Engine/Source/Developer/SimplygonMeshReduction/Private/SimplygonMeshReduction.cpp
@@ -1897,7 +1897,6 @@ private:
 
 		// Set texture coordinate count on the new model.
 		NewModel->NumTexCoords = MeshData.TexCoordCount;
-		NewModel->Size = 0;
 	}
 
 	/**
diff --git a/Engine/Source/Developer/Windows/ShaderFormatD3D/Private/D3D11ShaderCompiler.cpp b/Engine/Source/Developer/Windows/ShaderFormatD3D/Private/D3D11ShaderCompiler.cpp
--- a/Engine/Source/Developer/Windows/ShaderFormatD3D/Private/D3D11ShaderCompiler.cpp
+++ b/Engine/Source/Developer/Windows/ShaderFormatD3D/Private/D3D11ShaderCompiler.cpp
@@ -753,15 +753,13 @@ static bool CompileAndProcessD3DShader(FString& PreprocessedShaderSource, const
 				FMemoryWriter UniformBufferNameWriter(UniformBufferNameBytes);
 				UniformBufferNameWriter << UniformBufferNames;
 
-				// At runtime textures are just SRVs, so combine them for the purposes of building token streams.
-				GenericSRT.ShaderResourceViewMap.Append(GenericSRT.TextureMap);
-
 				// Copy over the bits indicating which resource tables are active.
 				SRT.ResourceTableBits = GenericSRT.ResourceTableBits;
 
 				SRT.ResourceTableLayoutHashes = GenericSRT.ResourceTableLayoutHashes;
 
 				// Now build our token streams.
+				BuildResourceTableTokenStream(GenericSRT.TextureMap, GenericSRT.MaxBoundResourceTable, SRT.TextureMap);
 				BuildResourceTableTokenStream(GenericSRT.ShaderResourceViewMap, GenericSRT.MaxBoundResourceTable, SRT.ShaderResourceViewMap);
 				BuildResourceTableTokenStream(GenericSRT.SamplerMap, GenericSRT.MaxBoundResourceTable, SRT.SamplerMap);
 				BuildResourceTableTokenStream(GenericSRT.UnorderedAccessViewMap, GenericSRT.MaxBoundResourceTable, SRT.UnorderedAccessViewMap);
diff --git a/Engine/Source/Editor/EditorStyle/Private/SlateEditorStyle.cpp b/Engine/Source/Editor/EditorStyle/Private/SlateEditorStyle.cpp
--- a/Engine/Source/Editor/EditorStyle/Private/SlateEditorStyle.cpp
+++ b/Engine/Source/Editor/EditorStyle/Private/SlateEditorStyle.cpp
@@ -4319,9 +4319,11 @@ void FSlateEditorStyle::FStyle::SetupLevelEditorStyle()
 		Set( "EditorViewport.LightingOnlyMode", new IMAGE_BRUSH( "Icons/icon_ViewMode_LightingOnly_16px", Icon16x16 ) );
 		Set( "EditorViewport.LightComplexityMode", new IMAGE_BRUSH( "Icons/icon_ViewMode_LightComplexity_16px", Icon16x16 ) );
 		Set( "EditorViewport.ShaderComplexityMode", new IMAGE_BRUSH( "Icons/icon_ViewMode_Shadercomplexity_16px", Icon16x16 ) );
-		Set( "EditorViewport.QuadComplexityMode", new IMAGE_BRUSH( "Icons/icon_ViewMode_Quadcomplexity_16px", Icon16x16 ) );
-		Set( "EditorViewport.WantedMipsAccuracyMode", new IMAGE_BRUSH( "Icons/icon_ViewMode_WantedMipsAccuracy_16px", Icon16x16 ) );
-		Set( "EditorViewport.TexelFactorAccuracyMode", new IMAGE_BRUSH( "Icons/icon_ViewMode_TexelFactorAccuracy_16px", Icon16x16 ) );
+		Set( "EditorViewport.QuadOverdrawMode", new IMAGE_BRUSH( "Icons/icon_ViewMode_QuadOverdraw_16px", Icon16x16 ) );
+		Set( "EditorViewport.ShaderComplexityWithQuadOverdrawMode", new IMAGE_BRUSH( "Icons/icon_ViewMode_Shadercomplexity_16px", Icon16x16 ) );
+		Set( "EditorViewport.WantedMipsAccuracyMode", new IMAGE_BRUSH( "Icons/icon_ViewMode_TextureStreamingAccuracy_16px", Icon16x16 ) );
+		Set( "EditorViewport.TexelFactorAccuracyMode", new IMAGE_BRUSH( "Icons/icon_ViewMode_TextureStreamingAccuracy_16px", Icon16x16 ) );
+		Set( "EditorViewport.TexCoordScaleAccuracyMode", new IMAGE_BRUSH( "Icons/icon_ViewMode_TextureStreamingAccuracy_16px", Icon16x16 ) );
 		Set( "EditorViewport.StationaryLightOverlapMode", new IMAGE_BRUSH( "Icons/icon_ViewMode_StationaryLightOverlap_16px", Icon16x16 ) );
 		Set( "EditorViewport.LightmapDensityMode", new IMAGE_BRUSH( "Icons/icon_ViewMode_LightmapDensity_16px", Icon16x16 ) );
 
diff --git a/Engine/Source/Editor/LevelEditor/Private/LevelEditor.cpp b/Engine/Source/Editor/LevelEditor/Private/LevelEditor.cpp
--- a/Engine/Source/Editor/LevelEditor/Private/LevelEditor.cpp
+++ b/Engine/Source/Editor/LevelEditor/Private/LevelEditor.cpp
@@ -1384,6 +1384,9 @@ void FLevelEditorModule::BindGlobalLevelEditorCommands()
 	ActionList.MapAction(Commands.BuildLODsOnly,
 		FExecuteAction::CreateStatic(&FLevelEditorActionCallbacks::BuildLODsOnly_Execute));
 	
+	ActionList.MapAction(Commands.BuildTextureStreamingOnly,
+		FExecuteAction::CreateStatic(&FLevelEditorActionCallbacks::BuildTextureStreamingOnly_Execute));
+
 	ActionList.MapAction( 
 		Commands.LightingQuality_Production, 
 		FExecuteAction::CreateStatic( &FLevelEditorActionCallbacks::SetLightingQuality, (ELightingBuildQuality)Quality_Production ),
diff --git a/Engine/Source/Editor/LevelEditor/Private/LevelEditorActions.cpp b/Engine/Source/Editor/LevelEditor/Private/LevelEditorActions.cpp
--- a/Engine/Source/Editor/LevelEditor/Private/LevelEditorActions.cpp
+++ b/Engine/Source/Editor/LevelEditor/Private/LevelEditorActions.cpp
@@ -654,6 +654,14 @@ void FLevelEditorActionCallbacks::BuildLODsOnly_Execute()
 	FEditorBuildUtils::EditorBuild(GetWorld(), FBuildOptions::BuildHierarchicalLOD);
 }
 
+ENGINE_API void BuildTextureStreamingData(UWorld*);
+
+void FLevelEditorActionCallbacks::BuildTextureStreamingOnly_Execute()
+{
+	// Build TexCoord Scale Shaders
+	BuildTextureStreamingData(GetWorld());
+}
+
 bool FLevelEditorActionCallbacks::IsLightingQualityChecked( ELightingBuildQuality TestQuality )
 {
 	int32 CurrentQualityLevel;
@@ -2870,6 +2878,7 @@ void FLevelEditorCommands::RegisterCommands()
 	UI_COMMAND( BuildGeometryOnly_OnlyCurrentLevel, "Build Geometry (Current Level)", "Builds geometry, only for the current level", EUserInterfaceActionType::Button, FInputChord() );
 	UI_COMMAND( BuildPathsOnly, "Build Paths", "Only builds paths (all levels.)", EUserInterfaceActionType::Button, FInputChord() );
 	UI_COMMAND( BuildLODsOnly, "Build LODs", "Only builds LODs (all levels.)", EUserInterfaceActionType::Button, FInputChord() );
+	UI_COMMAND( BuildTextureStreamingOnly, "Build Texture Streaming", "Build texture streaming data", EUserInterfaceActionType::Button, FInputChord() );
 	
 	UI_COMMAND( LightingQuality_Production, "Production", "Sets precomputed lighting quality to highest possible quality (slowest computation time.)", EUserInterfaceActionType::RadioButton, FInputChord() );
 	UI_COMMAND( LightingQuality_High, "High", "Sets precomputed lighting quality to high quality", EUserInterfaceActionType::RadioButton, FInputChord() );
diff --git a/Engine/Source/Editor/LevelEditor/Private/LevelEditorToolBar.cpp b/Engine/Source/Editor/LevelEditor/Private/LevelEditorToolBar.cpp
--- a/Engine/Source/Editor/LevelEditor/Private/LevelEditorToolBar.cpp
+++ b/Engine/Source/Editor/LevelEditor/Private/LevelEditorToolBar.cpp
@@ -1569,6 +1569,13 @@ TSharedRef< SWidget > FLevelEditorToolBar::GenerateBuildMenuContent( TSharedRef<
 	}
 	MenuBuilder.EndSection();
 
+	MenuBuilder.BeginSection("LevelEditorTextureStreaming", LOCTEXT("TextureStreamingHeading", "Texture Streaming"));
+	{
+		MenuBuilder.AddMenuEntry(FLevelEditorCommands::Get().BuildTextureStreamingOnly);
+	}
+	MenuBuilder.EndSection();
+
+
 	MenuBuilder.BeginSection("LevelEditorAutomation", LOCTEXT( "AutomationHeading", "Automation" ) );
 	{
 		MenuBuilder.AddMenuEntry( 
diff --git a/Engine/Source/Editor/LevelEditor/Public/LevelEditorActions.h b/Engine/Source/Editor/LevelEditor/Public/LevelEditorActions.h
--- a/Engine/Source/Editor/LevelEditor/Public/LevelEditorActions.h
+++ b/Engine/Source/Editor/LevelEditor/Public/LevelEditorActions.h
@@ -74,6 +74,7 @@ public:
 	TSharedPtr< FUICommandInfo > BuildGeometryOnly_OnlyCurrentLevel;
 	TSharedPtr< FUICommandInfo > BuildPathsOnly;
 	TSharedPtr< FUICommandInfo > BuildLODsOnly;
+	TSharedPtr< FUICommandInfo > BuildTextureStreamingOnly;
 	TSharedPtr< FUICommandInfo > LightingQuality_Production;
 	TSharedPtr< FUICommandInfo > LightingQuality_High;
 	TSharedPtr< FUICommandInfo > LightingQuality_Medium;
@@ -717,6 +718,7 @@ public:
 	static void BuildGeometryOnly_OnlyCurrentLevel_Execute();
 	static void BuildPathsOnly_Execute();
 	static void BuildLODsOnly_Execute();
+	static void BuildTextureStreamingOnly_Execute();
 	static void SetLightingQuality( ELightingBuildQuality NewQuality );
 	static bool IsLightingQualityChecked( ELightingBuildQuality TestQuality );
 	static float GetLightingDensityIdeal();
diff --git a/Engine/Source/Editor/UnrealEd/Private/EditorViewportCommands.cpp b/Engine/Source/Editor/UnrealEd/Private/EditorViewportCommands.cpp
--- a/Engine/Source/Editor/UnrealEd/Private/EditorViewportCommands.cpp
+++ b/Engine/Source/Editor/UnrealEd/Private/EditorViewportCommands.cpp
@@ -23,9 +23,11 @@ void FEditorViewportCommands::RegisterCommands()
 	UI_COMMAND( LightingOnlyMode, "Lighting Only View Mode", "Renders the scene with lights only, no textures", EUserInterfaceActionType::RadioButton, FInputChord( EModifierKey::Alt, EKeys::Six ) );
 	UI_COMMAND( LightComplexityMode, "Light Complexity View Mode", "Renders the scene with light complexity visualization", EUserInterfaceActionType::RadioButton, FInputChord( EModifierKey::Alt, EKeys::Seven ) );
 	UI_COMMAND( ShaderComplexityMode, "Shader Complexity View Mode", "Renders the scene with shader complexity visualization", EUserInterfaceActionType::RadioButton, FInputChord( EModifierKey::Alt, EKeys::Eight ) );
-	UI_COMMAND( QuadComplexityMode, "Quad Complexity View Mode", "Renders the scene with quad complexity visualization", EUserInterfaceActionType::RadioButton, FInputChord() );
-	UI_COMMAND( WantedMipsAccuracyMode, "Wanted Mips Accuracy View Mode", "Visualize how accurate is the wanted mips computed by the texture streamer", EUserInterfaceActionType::RadioButton, FInputChord() );
-	UI_COMMAND( TexelFactorAccuracyMode, "Texel Factor Accuracy View Mode", "Visualize how accurate is the texel factor is computed on each mesh", EUserInterfaceActionType::RadioButton, FInputChord() );
+	UI_COMMAND( QuadOverdrawMode, "Quad Complexity View Mode", "Renders the scene with quad complexity visualization", EUserInterfaceActionType::RadioButton, FInputChord() );
+	UI_COMMAND( ShaderComplexityWithQuadOverdrawMode, "Shader Complexity & Quads visualization", "Renders the scene with shader complexity and quad overdraw visualization", EUserInterfaceActionType::RadioButton, FInputChord() );
+	UI_COMMAND( WantedMipsAccuracyMode, "Wanted Mips Accuracy View Mode", "Visualize the accuracy of the wanted mips used by the texture streamer", EUserInterfaceActionType::RadioButton, FInputChord() );
+	UI_COMMAND( TexelFactorAccuracyMode, "Texel Factor Accuracy View Mode", "Visualize the accuracy of the primitive texel factors used by the texture streamer", EUserInterfaceActionType::RadioButton, FInputChord() );
+	UI_COMMAND( TexCoordScaleAccuracyMode, "Texture Coordinate Scale Accuracy View Mode", "Visualize the accuracy of the material texture coordinate scale factors used by the texture streamer", EUserInterfaceActionType::RadioButton, FInputChord() );
 	UI_COMMAND( StationaryLightOverlapMode, "Stationary Light Overlap View Mode", "Visualizes overlap of stationary lights", EUserInterfaceActionType::RadioButton, FInputChord() );
 	UI_COMMAND( LightmapDensityMode, "Lightmap Density View Mode", "Renders the scene with lightmap density visualization", EUserInterfaceActionType::RadioButton, FInputChord( EModifierKey::Alt, EKeys::Zero ) );
 
diff --git a/Engine/Source/Editor/UnrealEd/Private/SEditorViewport.cpp b/Engine/Source/Editor/UnrealEd/Private/SEditorViewport.cpp
--- a/Engine/Source/Editor/UnrealEd/Private/SEditorViewport.cpp
+++ b/Engine/Source/Editor/UnrealEd/Private/SEditorViewport.cpp
@@ -349,9 +349,11 @@ void SEditorViewport::BindCommands()
 	MAP_VIEWMODE_ACTION( Commands.LightingOnlyMode, VMI_LightingOnly );
 	MAP_VIEWMODE_ACTION( Commands.LightComplexityMode, VMI_LightComplexity );
 	MAP_VIEWMODE_ACTION( Commands.ShaderComplexityMode, VMI_ShaderComplexity );
-	MAP_VIEWMODE_ACTION( Commands.QuadComplexityMode, VMI_QuadComplexity );
+	MAP_VIEWMODE_ACTION( Commands.QuadOverdrawMode, VMI_QuadOverdraw);
+	MAP_VIEWMODE_ACTION( Commands.ShaderComplexityWithQuadOverdrawMode, VMI_ShaderComplexityWithQuadOverdraw );
 	MAP_VIEWMODE_ACTION( Commands.WantedMipsAccuracyMode, VMI_WantedMipsAccuracy );
 	MAP_VIEWMODE_ACTION( Commands.TexelFactorAccuracyMode, VMI_TexelFactorAccuracy );
+	MAP_VIEWMODE_ACTION( Commands.TexCoordScaleAccuracyMode, VMI_TexCoordScaleAccuracy );
 	MAP_VIEWMODE_ACTION( Commands.StationaryLightOverlapMode, VMI_StationaryLightOverlap );
 	MAP_VIEWMODE_ACTION( Commands.LightmapDensityMode, VMI_LightmapDensity );
 	MAP_VIEWMODE_ACTION( Commands.ReflectionOverrideMode, VMI_ReflectionOverride );
diff --git a/Engine/Source/Editor/UnrealEd/Private/SEditorViewportViewMenu.cpp b/Engine/Source/Editor/UnrealEd/Private/SEditorViewportViewMenu.cpp
--- a/Engine/Source/Editor/UnrealEd/Private/SEditorViewportViewMenu.cpp
+++ b/Engine/Source/Editor/UnrealEd/Private/SEditorViewportViewMenu.cpp
@@ -68,10 +68,14 @@ FText SEditorViewportViewMenu::GetViewMenuLabel() const
 				Label = LOCTEXT("ViewMenuTitle_ShaderComplexity", "Shader Complexity");
 				break;
 
-			case VMI_QuadComplexity:
+			case VMI_QuadOverdraw:
 				Label = LOCTEXT("ViewMenuTitle_QuadOverdraw", "Quad Overdraw");
 				break;
 
+			case VMI_ShaderComplexityWithQuadOverdraw:
+				Label = LOCTEXT("ViewMenuTitle_ShaderComplexityWithQuadOverdraw", "Shader Complexity & Quads");
+				break;
+
 			case VMI_WantedMipsAccuracy:
 				Label = LOCTEXT("ViewMenuTitle_WantedMipAccuracy", "Wanted Mips Accuracy");
 				break;
@@ -80,6 +84,10 @@ FText SEditorViewportViewMenu::GetViewMenuLabel() const
 				Label = LOCTEXT("ViewMenuTitle_TexelFactorAccuracy", "Texel Factor Accuracy");
 				break;
 
+			case VMI_TexCoordScaleAccuracy:
+				Label = LOCTEXT("ViewMenuTitle_TexCoordScaleAccuracy", "Texture Coordinate Scale Accuracy");
+				break;
+
 			case VMI_StationaryLightOverlap:
 				Label = LOCTEXT("ViewMenuTitle_StationaryLightOverlap", "Stationary Light Overlap");
 				break;
@@ -132,7 +140,8 @@ const FSlateBrush* SEditorViewportViewMenu::GetViewMenuLabelIcon() const
 		static FName LightingOnlyIcon("EditorViewport.LightingOnlyMode");
 		static FName LightComplexityIcon("EditorViewport.LightComplexityMode");
 		static FName ShaderComplexityIcon("EditorViewport.ShaderComplexityMode");
-		static FName QuadComplexityIcon("EditorViewport.QuadComplexityMode");
+		static FName QuadOverdrawIcon("EditorViewport.QuadOverdrawMode");
+		static FName ShaderComplexityWithQuadOverdrawIcon("EditorViewport.ShaderCOmplexityWithQuadOverdrawMode");
 		static FName WantedMipsAccuracyIcon("EditorViewport.WantedMipsAccuracyMode");
 		static FName TexelFactorAccuracyIcon("EditorViewport.TexelFactorAccuracyMode");
 		static FName LightOverlapIcon("EditorViewport.StationaryLightOverlapMode");
@@ -177,8 +186,12 @@ const FSlateBrush* SEditorViewportViewMenu::GetViewMenuLabelIcon() const
 				Icon = ShaderComplexityIcon;
 				break;
 
-			case VMI_QuadComplexity:
-				Icon = QuadComplexityIcon;
+			case VMI_QuadOverdraw:
+				Icon = QuadOverdrawIcon;
+				break;
+
+			case VMI_ShaderComplexityWithQuadOverdraw:
+				Icon = ShaderComplexityWithQuadOverdrawIcon;
 				break;
 
 			case VMI_WantedMipsAccuracy:
@@ -189,6 +202,10 @@ const FSlateBrush* SEditorViewportViewMenu::GetViewMenuLabelIcon() const
 				Icon = TexelFactorAccuracyIcon;
 				break;
 
+			case VMI_TexCoordScaleAccuracy:
+				Icon = TexelFactorAccuracyIcon;
+				break;
+
 			case VMI_StationaryLightOverlap:
 				Icon = LightOverlapIcon;
 				break;
@@ -265,10 +282,15 @@ TSharedRef<SWidget> SEditorViewportViewMenu::GenerateViewMenuContent() const
 						Menu.AddMenuEntry( BaseViewportCommands.LightmapDensityMode, NAME_None, LOCTEXT("LightmapDensityViewModeDisplayName", "Lightmap Density") );
 						Menu.AddMenuEntry( BaseViewportCommands.StationaryLightOverlapMode, NAME_None, LOCTEXT("StationaryLightOverlapViewModeDisplayName", "Stationary Light Overlap") );
 						Menu.AddMenuEntry( BaseViewportCommands.ShaderComplexityMode, NAME_None, LOCTEXT("ShaderComplexityViewModeDisplayName", "Shader Complexity") );
-						Menu.AddMenuEntry( BaseViewportCommands.QuadComplexityMode, NAME_None, LOCTEXT("QuadComplexityViewModeDisplayName", "Quad Overdraw") );
+						Menu.AddMenuEntry( BaseViewportCommands.ShaderComplexityWithQuadOverdrawMode, NAME_None, LOCTEXT("ShaderComplexityWithQuadOverdrawViewModeDisplayName", "Shader Complexity & Quads") );
+						Menu.AddMenuEntry( BaseViewportCommands.QuadOverdrawMode, NAME_None, LOCTEXT("QuadOverdrawViewModeDisplayName", "Quad Overdraw") );
 						Menu.AddMenuEntry( BaseViewportCommands.LODColorationMode, NAME_None, LOCTEXT("LODColorationViewModeDisplayName", "LOD Coloration") );
+
+						Menu.BeginSection("TextureStreaming", LOCTEXT("TextureStreamingHeader", "Texture Streaming") );
 						Menu.AddMenuEntry( BaseViewportCommands.WantedMipsAccuracyMode, NAME_None, LOCTEXT("WantedMipsAccuracyDisplayName", "Wanted Mips Accuracy") );
 						Menu.AddMenuEntry( BaseViewportCommands.TexelFactorAccuracyMode, NAME_None, LOCTEXT("TexelFactorAccuracyDisplayName", "Texel Factor Accuracy") );
+						Menu.AddMenuEntry( BaseViewportCommands.TexCoordScaleAccuracyMode, NAME_None, LOCTEXT("TexCoordScaleAccuracyDisplayName", "TexCoord Scale Accuracy") );
+						Menu.EndSection();
 					}
 				};
 
diff --git a/Engine/Source/Editor/UnrealEd/Public/EditorViewportCommands.h b/Engine/Source/Editor/UnrealEd/Public/EditorViewportCommands.h
--- a/Engine/Source/Editor/UnrealEd/Public/EditorViewportCommands.h
+++ b/Engine/Source/Editor/UnrealEd/Public/EditorViewportCommands.h
@@ -65,8 +65,11 @@ public:
 	/** Changes the viewport to shader complexity mode */
 	TSharedPtr< FUICommandInfo > ShaderComplexityMode;
 
-	/** Changes the viewport to quad complexity mode */
-	TSharedPtr< FUICommandInfo > QuadComplexityMode;
+	/** Changes the viewport to quad overdraw mode */
+	TSharedPtr< FUICommandInfo > QuadOverdrawMode;
+
+	/** Changes the viewport to shader complexity with quad overdraw mode */
+	TSharedPtr< FUICommandInfo > ShaderComplexityWithQuadOverdrawMode;
 
 	/** Changes the viewport to wanted mips accuracy mode */
 	TSharedPtr< FUICommandInfo > WantedMipsAccuracyMode;
@@ -74,6 +77,9 @@ public:
 	/** Changes the viewport to texel factor accuracy mode */
 	TSharedPtr< FUICommandInfo > TexelFactorAccuracyMode;
 
+	/** Changes the viewport to texcoord scale accuracy mode */
+	TSharedPtr< FUICommandInfo > TexCoordScaleAccuracyMode;
+
 	/** Changes the viewport to stationary light overlap mode */
 	TSharedPtr< FUICommandInfo > StationaryLightOverlapMode;
 
diff --git a/Engine/Source/Runtime/Apple/MetalRHI/Private/MetalUAV.cpp b/Engine/Source/Runtime/Apple/MetalRHI/Private/MetalUAV.cpp
--- a/Engine/Source/Runtime/Apple/MetalRHI/Private/MetalUAV.cpp
+++ b/Engine/Source/Runtime/Apple/MetalRHI/Private/MetalUAV.cpp
@@ -188,6 +188,11 @@ FShaderResourceViewRHIRef FMetalDynamicRHI::RHICreateShaderResourceView_RenderTh
 	return GDynamicRHI->RHICreateShaderResourceView(VertexBuffer, Stride, Format);
 }
 
+FShaderResourceViewRHIRef FMetalDynamicRHI::RHICreateShaderResourceView_RenderThread(class FRHICommandListImmediate& RHICmdList, FIndexBufferRHIParamRef Buffer)
+{
+	return GDynamicRHI->RHICreateShaderResourceView(Buffer);
+}
+
 FShaderResourceViewRHIRef FMetalDynamicRHI::RHICreateShaderResourceView_RenderThread(class FRHICommandListImmediate& RHICmdList, FStructuredBufferRHIParamRef StructuredBuffer)
 {
 	UE_LOG(LogRHI, Fatal,TEXT("Metal RHI doesn't support RHICreateShaderResourceView with FStructuredBufferRHIParamRef yet!"));
@@ -218,6 +223,13 @@ FShaderResourceViewRHIRef FMetalDynamicRHI::RHICreateShaderResourceView(FVertexB
 	return SRV;
 }
 
+FShaderResourceViewRHIRef FMetalDynamicRHI::RHICreateShaderResourceView(FIndexBufferRHIParamRef BufferRHI)
+{
+	UE_LOG(LogRHI, Fatal, TEXT("Metal RHI doesn't support RHICreateShaderResourceView with FIndexBufferRHIParamRef yet!"));
+	
+	return FShaderResourceViewRHIRef();
+}
+
 FShaderResourceViewRHIRef FMetalDynamicRHI::RHICreateShaderResourceView(FTexture2DRHIParamRef Texture2DRHI, uint8 MipLevel)
 {
     FMetalShaderResourceView* SRV = new FMetalShaderResourceView;
diff --git a/Engine/Source/Runtime/Apple/MetalRHI/Public/MetalRHI.h b/Engine/Source/Runtime/Apple/MetalRHI/Public/MetalRHI.h
--- a/Engine/Source/Runtime/Apple/MetalRHI/Public/MetalRHI.h
+++ b/Engine/Source/Runtime/Apple/MetalRHI/Public/MetalRHI.h
@@ -65,6 +65,7 @@ public:
 	virtual FUnorderedAccessViewRHIRef RHICreateUnorderedAccessView(FVertexBufferRHIParamRef VertexBuffer, uint8 Format) final override;
 	virtual FShaderResourceViewRHIRef RHICreateShaderResourceView(FStructuredBufferRHIParamRef StructuredBuffer) final override;
 	virtual FShaderResourceViewRHIRef RHICreateShaderResourceView(FVertexBufferRHIParamRef VertexBuffer, uint32 Stride, uint8 Format) final override;
+	virtual FShaderResourceViewRHIRef RHICreateShaderResourceView(FIndexBufferRHIParamRef Buffer) final override;
 	virtual uint64 RHICalcTexture2DPlatformSize(uint32 SizeX, uint32 SizeY, uint8 Format, uint32 NumMips, uint32 NumSamples, uint32 Flags, uint32& OutAlign) final override;
 	virtual uint64 RHICalcTexture3DPlatformSize(uint32 SizeX, uint32 SizeY, uint32 SizeZ, uint8 Format, uint32 NumMips, uint32 Flags, uint32& OutAlign) final override;
 	virtual uint64 RHICalcTextureCubePlatformSize(uint32 Size, uint8 Format, uint32 NumMips, uint32 Flags, uint32& OutAlign) final override;
@@ -154,6 +155,7 @@ public:
 	virtual FShaderResourceViewRHIRef RHICreateShaderResourceView_RenderThread(class FRHICommandListImmediate& RHICmdList, FTexture2DArrayRHIParamRef Texture2DArrayRHI, uint8 MipLevel) final override;
 	virtual FShaderResourceViewRHIRef RHICreateShaderResourceView_RenderThread(class FRHICommandListImmediate& RHICmdList, FTextureCubeRHIParamRef TextureCubeRHI, uint8 MipLevel) final override;
 	virtual FShaderResourceViewRHIRef RHICreateShaderResourceView_RenderThread(class FRHICommandListImmediate& RHICmdList, FVertexBufferRHIParamRef VertexBuffer, uint32 Stride, uint8 Format) final override;
+	virtual FShaderResourceViewRHIRef RHICreateShaderResourceView_RenderThread(class FRHICommandListImmediate& RHICmdList, FIndexBufferRHIParamRef Buffer) final override;
 	virtual FShaderResourceViewRHIRef RHICreateShaderResourceView_RenderThread(class FRHICommandListImmediate& RHICmdList, FStructuredBufferRHIParamRef StructuredBuffer) final override;
 
 	// FRHICommandContext API accessible only on the immediate device context
diff --git a/Engine/Source/Runtime/EmptyRHI/Private/EmptyUAV.cpp b/Engine/Source/Runtime/EmptyRHI/Private/EmptyUAV.cpp
--- a/Engine/Source/Runtime/EmptyRHI/Private/EmptyUAV.cpp
+++ b/Engine/Source/Runtime/EmptyRHI/Private/EmptyUAV.cpp
@@ -61,6 +61,12 @@ FShaderResourceViewRHIRef FEmptyDynamicRHI::RHICreateShaderResourceView(FVertexB
 	return SRV;
 }
 
+FShaderResourceViewRHIRef FEmptyDynamicRHI::RHICreateShaderResourceView(FIndexBufferRHIParamRef BufferRHI)
+{
+	// there should be no need to create an object
+	return FShaderResourceViewRHIRef();
+}
+
 FShaderResourceViewRHIRef FEmptyDynamicRHI::RHICreateShaderResourceView(FTexture2DRHIParamRef Texture2DRHI, uint8 MipLevel)
 {
 	FEmptyShaderResourceView* SRV = new FEmptyShaderResourceView;
diff --git a/Engine/Source/Runtime/EmptyRHI/Public/EmptyRHI.h b/Engine/Source/Runtime/EmptyRHI/Public/EmptyRHI.h
--- a/Engine/Source/Runtime/EmptyRHI/Public/EmptyRHI.h
+++ b/Engine/Source/Runtime/EmptyRHI/Public/EmptyRHI.h
@@ -64,6 +64,7 @@ public:
 	virtual FUnorderedAccessViewRHIRef RHICreateUnorderedAccessView(FVertexBufferRHIParamRef VertexBuffer, uint8 Format) final override;
 	virtual FShaderResourceViewRHIRef RHICreateShaderResourceView(FStructuredBufferRHIParamRef StructuredBuffer) final override;
 	virtual FShaderResourceViewRHIRef RHICreateShaderResourceView(FVertexBufferRHIParamRef VertexBuffer, uint32 Stride, uint8 Format) final override;
+	virtual FShaderResourceViewRHIRef RHICreateShaderResourceView(FIndexBufferRHIParamRef Buffer) final override;
 	virtual uint64 RHICalcTexture2DPlatformSize(uint32 SizeX, uint32 SizeY, uint8 Format, uint32 NumMips, uint32 NumSamples, uint32 Flags, uint32& OutAlign) final override;
 	virtual uint64 RHICalcTexture3DPlatformSize(uint32 SizeX, uint32 SizeY, uint32 SizeZ, uint8 Format, uint32 NumMips, uint32 Flags, uint32& OutAlign) final override;
 	virtual uint64 RHICalcTextureCubePlatformSize(uint32 Size, uint8 Format, uint32 NumMips, uint32 Flags, uint32& OutAlign) final override;
diff --git a/Engine/Source/Runtime/Engine/Classes/Components/PrimitiveComponent.h b/Engine/Source/Runtime/Engine/Classes/Components/PrimitiveComponent.h
--- a/Engine/Source/Runtime/Engine/Classes/Components/PrimitiveComponent.h
+++ b/Engine/Source/Runtime/Engine/Classes/Components/PrimitiveComponent.h
@@ -1107,7 +1107,7 @@ public:
 	/**  @return True if a primitive's parameters as well as its position is static during gameplay, and can thus use static lighting. */
 	bool HasStaticLighting() const;
 
-	virtual bool HasValidSettingsForStaticLighting() const 
+	virtual bool HasValidSettingsForStaticLighting(bool bOverlookInvalidComponents) const 
 	{
 		return HasStaticLighting();
 	}
diff --git a/Engine/Source/Runtime/Engine/Classes/Components/StaticMeshComponent.h b/Engine/Source/Runtime/Engine/Classes/Components/StaticMeshComponent.h
--- a/Engine/Source/Runtime/Engine/Classes/Components/StaticMeshComponent.h
+++ b/Engine/Source/Runtime/Engine/Classes/Components/StaticMeshComponent.h
@@ -332,7 +332,7 @@ public:
 	virtual bool GetLightMapResolution( int32& Width, int32& Height ) const override;
 	virtual int32 GetStaticLightMapResolution() const override;
 	/** Returns true if the component is static AND has the right static mesh setup to support lightmaps. */
-	virtual bool HasValidSettingsForStaticLighting() const override;
+	virtual bool HasValidSettingsForStaticLighting(bool bOverlookInvalidComponents) const override;
 
 	virtual void GetLightAndShadowMapMemoryUsage( int32& LightMapMemoryUsage, int32& ShadowMapMemoryUsage ) const override;
 	virtual void GetUsedMaterials(TArray<UMaterialInterface*>& OutMaterials) const override;
diff --git a/Engine/Source/Runtime/Engine/Classes/Engine/Engine.h b/Engine/Source/Runtime/Engine/Classes/Engine/Engine.h
--- a/Engine/Source/Runtime/Engine/Classes/Engine/Engine.h
+++ b/Engine/Source/Runtime/Engine/Classes/Engine/Engine.h
@@ -963,7 +963,7 @@ public:
 	UPROPERTY(globalconfig)
 	TArray<FLinearColor> HLODColorationColors;
 
-	/** The colors used for TexelFactorAccuracy and WantedMipsAccuracy. */
+	/** The colors used for TexelFactorAccuracy, WantedMipsAccuracy and TexCoordScaleAccuracy. */
 	UPROPERTY(globalconfig)
 	TArray<FLinearColor> StreamingAccuracyColors;
 
diff --git a/Engine/Source/Runtime/Engine/Classes/Engine/EngineBaseTypes.h b/Engine/Source/Runtime/Engine/Classes/Engine/EngineBaseTypes.h
--- a/Engine/Source/Runtime/Engine/Classes/Engine/EngineBaseTypes.h
+++ b/Engine/Source/Runtime/Engine/Classes/Engine/EngineBaseTypes.h
@@ -872,17 +872,21 @@ enum EViewModeIndex
 	VMI_CollisionVisibility = 16, 
 	VMI_VertexDensities = 17,
 	/** Colored according to the current LOD index. */
-	VMI_LODColoration = 18,	
+	VMI_LODColoration = 18,
 	/** Colored according to the quad coverage. */
-	VMI_QuadComplexity = 19,
+	VMI_QuadOverdraw = 19,
 	/** Colored according to the accuracy of the texture streamer wanted mips computation. */
 	VMI_WantedMipsAccuracy = 20,
 	/** Colored according to the texel factor accuracy. */
 	VMI_TexelFactorAccuracy = 21,
+	/** Colored according to shader complexity, including quad overdraw. */
+	VMI_ShaderComplexityWithQuadOverdraw = 22,
 	/** Colored according to the current HLOD index. */
-	VMI_HLODColoration = 22,
+	VMI_HLODColoration = 23,
 	/** Group item for LOD and HLOD coloration*/
-	VMI_GroupLODColoration = 23,
+	VMI_GroupLODColoration = 24,
+	/** Colored according to how accuracte the texcoord scales computed by the texture streamer are. */
+	VMI_TexCoordScaleAccuracy = 25,
 
 	VMI_Max UMETA(Hidden),
 
diff --git a/Engine/Source/Runtime/Engine/Classes/Materials/MaterialExpressionScreenPosition.h b/Engine/Source/Runtime/Engine/Classes/Materials/MaterialExpressionScreenPosition.h
--- a/Engine/Source/Runtime/Engine/Classes/Materials/MaterialExpressionScreenPosition.h
+++ b/Engine/Source/Runtime/Engine/Classes/Materials/MaterialExpressionScreenPosition.h
@@ -5,11 +5,26 @@
 #include "Materials/MaterialExpression.h"
 #include "MaterialExpressionScreenPosition.generated.h"
 
+UENUM()
+enum EMaterialExpressionScreenPositionMapping
+{
+	/** A UV in the 0..1 range for use with the ScreeTnexture material expression.*/
+	MESP_SceneTextureUV UMETA(DisplayName = "SceneTextureUV"),
+	/** A UV in the 0..1 range that maps to the local viewport */
+	MESP_ViewportUV     UMETA(DisplayName = "ViewportUV"),
+
+	MESP_Max,
+};
+
 UCLASS(collapsecategories, hidecategories=Object)
 class UMaterialExpressionScreenPosition : public UMaterialExpression
 {
 	GENERATED_UCLASS_BODY()
 
+	/** View input property to be accessed */
+	UPROPERTY(EditAnywhere, Category = UMaterialExpressionScreenPosition, meta = (DisplayName = "Mapping"))
+	TEnumAsByte<EMaterialExpressionScreenPositionMapping> Mapping;
+
 	//~ Begin UMaterialExpression Interface
 	virtual int32 Compile(class FMaterialCompiler* Compiler, int32 OutputIndex, int32 MultiplexIndex) override;
 	virtual void GetCaption(TArray<FString>& OutCaptions) const override;
diff --git a/Engine/Source/Runtime/Engine/Classes/Particles/Location/ParticleModuleLocation.h b/Engine/Source/Runtime/Engine/Classes/Particles/Location/ParticleModuleLocation.h
--- a/Engine/Source/Runtime/Engine/Classes/Particles/Location/ParticleModuleLocation.h
+++ b/Engine/Source/Runtime/Engine/Classes/Particles/Location/ParticleModuleLocation.h
@@ -5,8 +5,8 @@
 #include "Particles/Location/ParticleModuleLocationBase.h"
 #include "ParticleModuleLocation.generated.h"
 
-UCLASS(editinlinenew, hidecategories=Object, MinimalAPI, meta=(DisplayName = "Initial Location"))
-class UParticleModuleLocation : public UParticleModuleLocationBase
+UCLASS(editinlinenew, hidecategories=Object, meta=(DisplayName = "Initial Location"))
+class ENGINE_API UParticleModuleLocation : public UParticleModuleLocationBase
 {
 	GENERATED_UCLASS_BODY()
 
diff --git a/Engine/Source/Runtime/Engine/Classes/Particles/Location/ParticleModuleLocationBase.h b/Engine/Source/Runtime/Engine/Classes/Particles/Location/ParticleModuleLocationBase.h
--- a/Engine/Source/Runtime/Engine/Classes/Particles/Location/ParticleModuleLocationBase.h
+++ b/Engine/Source/Runtime/Engine/Classes/Particles/Location/ParticleModuleLocationBase.h
@@ -6,7 +6,7 @@
 #include "ParticleModuleLocationBase.generated.h"
 
 UCLASS(editinlinenew, hidecategories=Object, abstract, meta=(DisplayName = "Location"))
-class UParticleModuleLocationBase : public UParticleModule
+class ENGINE_API UParticleModuleLocationBase : public UParticleModule
 {
 	GENERATED_UCLASS_BODY()
 
diff --git a/Engine/Source/Runtime/Engine/Classes/Particles/Location/ParticleModuleLocationBoneSocket.h b/Engine/Source/Runtime/Engine/Classes/Particles/Location/ParticleModuleLocationBoneSocket.h
--- a/Engine/Source/Runtime/Engine/Classes/Particles/Location/ParticleModuleLocationBoneSocket.h
+++ b/Engine/Source/Runtime/Engine/Classes/Particles/Location/ParticleModuleLocationBoneSocket.h
@@ -43,7 +43,7 @@ struct FLocationBoneSocketInfo
 };
 
 UCLASS(editinlinenew, hidecategories=Object, meta=(DisplayName = "Bone/Socket Location"))
-class UParticleModuleLocationBoneSocket : public UParticleModuleLocationBase
+class ENGINE_API UParticleModuleLocationBoneSocket : public UParticleModuleLocationBase
 {
 	GENERATED_UCLASS_BODY()
 
diff --git a/Engine/Source/Runtime/Engine/Classes/Particles/Location/ParticleModuleLocationDirect.h b/Engine/Source/Runtime/Engine/Classes/Particles/Location/ParticleModuleLocationDirect.h
--- a/Engine/Source/Runtime/Engine/Classes/Particles/Location/ParticleModuleLocationDirect.h
+++ b/Engine/Source/Runtime/Engine/Classes/Particles/Location/ParticleModuleLocationDirect.h
@@ -11,8 +11,8 @@
 #include "Particles/Location/ParticleModuleLocationBase.h"
 #include "ParticleModuleLocationDirect.generated.h"
 
-UCLASS(editinlinenew, hidecategories=Object, MinimalAPI, meta=(DisplayName = "Direct Location"))
-class UParticleModuleLocationDirect : public UParticleModuleLocationBase
+UCLASS(editinlinenew, hidecategories=Object, meta=(DisplayName = "Direct Location"))
+class ENGINE_API UParticleModuleLocationDirect : public UParticleModuleLocationBase
 {
 	GENERATED_UCLASS_BODY()
 
diff --git a/Engine/Source/Runtime/Engine/Classes/Particles/Location/ParticleModuleLocationEmitter.h b/Engine/Source/Runtime/Engine/Classes/Particles/Location/ParticleModuleLocationEmitter.h
--- a/Engine/Source/Runtime/Engine/Classes/Particles/Location/ParticleModuleLocationEmitter.h
+++ b/Engine/Source/Runtime/Engine/Classes/Particles/Location/ParticleModuleLocationEmitter.h
@@ -20,8 +20,8 @@ enum ELocationEmitterSelectionMethod
 	ELESM_MAX,
 };
 
-UCLASS(editinlinenew, hidecategories=Object, MinimalAPI, meta=(DisplayName = "Emitter Initial Location"))
-class UParticleModuleLocationEmitter : public UParticleModuleLocationBase
+UCLASS(editinlinenew, hidecategories=Object, meta=(DisplayName = "Emitter Initial Location"))
+class ENGINE_API UParticleModuleLocationEmitter : public UParticleModuleLocationBase
 {
 	GENERATED_UCLASS_BODY()
 
diff --git a/Engine/Source/Runtime/Engine/Classes/Particles/Location/ParticleModuleLocationEmitterDirect.h b/Engine/Source/Runtime/Engine/Classes/Particles/Location/ParticleModuleLocationEmitterDirect.h
--- a/Engine/Source/Runtime/Engine/Classes/Particles/Location/ParticleModuleLocationEmitterDirect.h
+++ b/Engine/Source/Runtime/Engine/Classes/Particles/Location/ParticleModuleLocationEmitterDirect.h
@@ -13,7 +13,7 @@
 #include "ParticleModuleLocationEmitterDirect.generated.h"
 
 UCLASS(editinlinenew, hidecategories=Object, meta=(DisplayName = "Emitter Direct Location"))
-class UParticleModuleLocationEmitterDirect : public UParticleModuleLocationBase
+class ENGINE_API UParticleModuleLocationEmitterDirect : public UParticleModuleLocationBase
 {
 	GENERATED_UCLASS_BODY()
 
diff --git a/Engine/Source/Runtime/Engine/Classes/Particles/Location/ParticleModuleLocationPrimitiveBase.h b/Engine/Source/Runtime/Engine/Classes/Particles/Location/ParticleModuleLocationPrimitiveBase.h
--- a/Engine/Source/Runtime/Engine/Classes/Particles/Location/ParticleModuleLocationPrimitiveBase.h
+++ b/Engine/Source/Runtime/Engine/Classes/Particles/Location/ParticleModuleLocationPrimitiveBase.h
@@ -10,7 +10,7 @@
 #include "ParticleModuleLocationPrimitiveBase.generated.h"
 
 UCLASS(editinlinenew, hidecategories=Object)
-class UParticleModuleLocationPrimitiveBase : public UParticleModuleLocationBase
+class ENGINE_API UParticleModuleLocationPrimitiveBase : public UParticleModuleLocationBase
 {
 	GENERATED_UCLASS_BODY()
 
diff --git a/Engine/Source/Runtime/Engine/Classes/Particles/Location/ParticleModuleLocationPrimitiveCylinder.h b/Engine/Source/Runtime/Engine/Classes/Particles/Location/ParticleModuleLocationPrimitiveCylinder.h
--- a/Engine/Source/Runtime/Engine/Classes/Particles/Location/ParticleModuleLocationPrimitiveCylinder.h
+++ b/Engine/Source/Runtime/Engine/Classes/Particles/Location/ParticleModuleLocationPrimitiveCylinder.h
@@ -18,8 +18,8 @@ enum CylinderHeightAxis
 	PMLPC_HEIGHTAXIS_MAX,
 };
 
-UCLASS(editinlinenew, hidecategories = Object, MinimalAPI, meta = (DisplayName = "Cylinder"))
-class UParticleModuleLocationPrimitiveCylinder : public UParticleModuleLocationPrimitiveBase
+UCLASS(editinlinenew, hidecategories = Object, meta = (DisplayName = "Cylinder"))
+class ENGINE_API UParticleModuleLocationPrimitiveCylinder : public UParticleModuleLocationPrimitiveBase
 {
 	GENERATED_UCLASS_BODY()
 
diff --git a/Engine/Source/Runtime/Engine/Classes/Particles/Location/ParticleModuleLocationPrimitiveCylinder_Seeded.h b/Engine/Source/Runtime/Engine/Classes/Particles/Location/ParticleModuleLocationPrimitiveCylinder_Seeded.h
--- a/Engine/Source/Runtime/Engine/Classes/Particles/Location/ParticleModuleLocationPrimitiveCylinder_Seeded.h
+++ b/Engine/Source/Runtime/Engine/Classes/Particles/Location/ParticleModuleLocationPrimitiveCylinder_Seeded.h
@@ -6,7 +6,7 @@
 #include "ParticleModuleLocationPrimitiveCylinder_Seeded.generated.h"
 
 UCLASS(editinlinenew, hidecategories=Object, meta=(DisplayName = "Cylinder (Seed)"))
-class UParticleModuleLocationPrimitiveCylinder_Seeded : public UParticleModuleLocationPrimitiveCylinder
+class ENGINE_API UParticleModuleLocationPrimitiveCylinder_Seeded : public UParticleModuleLocationPrimitiveCylinder
 {
 	GENERATED_UCLASS_BODY()
 
diff --git a/Engine/Source/Runtime/Engine/Classes/Particles/Location/ParticleModuleLocationPrimitiveSphere.h b/Engine/Source/Runtime/Engine/Classes/Particles/Location/ParticleModuleLocationPrimitiveSphere.h
--- a/Engine/Source/Runtime/Engine/Classes/Particles/Location/ParticleModuleLocationPrimitiveSphere.h
+++ b/Engine/Source/Runtime/Engine/Classes/Particles/Location/ParticleModuleLocationPrimitiveSphere.h
@@ -9,8 +9,8 @@
 #include "Particles/Location/ParticleModuleLocationPrimitiveBase.h"
 #include "ParticleModuleLocationPrimitiveSphere.generated.h"
 
-UCLASS(editinlinenew, hidecategories=Object, MinimalAPI, meta=(DisplayName = "Sphere"))
-class UParticleModuleLocationPrimitiveSphere : public UParticleModuleLocationPrimitiveBase
+UCLASS(editinlinenew, hidecategories=Object, meta=(DisplayName = "Sphere"))
+class ENGINE_API UParticleModuleLocationPrimitiveSphere : public UParticleModuleLocationPrimitiveBase
 {
 	GENERATED_UCLASS_BODY()
 
diff --git a/Engine/Source/Runtime/Engine/Classes/Particles/Location/ParticleModuleLocationPrimitiveSphere_Seeded.h b/Engine/Source/Runtime/Engine/Classes/Particles/Location/ParticleModuleLocationPrimitiveSphere_Seeded.h
--- a/Engine/Source/Runtime/Engine/Classes/Particles/Location/ParticleModuleLocationPrimitiveSphere_Seeded.h
+++ b/Engine/Source/Runtime/Engine/Classes/Particles/Location/ParticleModuleLocationPrimitiveSphere_Seeded.h
@@ -6,7 +6,7 @@
 #include "ParticleModuleLocationPrimitiveSphere_Seeded.generated.h"
 
 UCLASS(editinlinenew, hidecategories=Object, meta=(DisplayName = "Sphere (Seed)"))
-class UParticleModuleLocationPrimitiveSphere_Seeded : public UParticleModuleLocationPrimitiveSphere
+class ENGINE_API UParticleModuleLocationPrimitiveSphere_Seeded : public UParticleModuleLocationPrimitiveSphere
 {
 	GENERATED_UCLASS_BODY()
 
diff --git a/Engine/Source/Runtime/Engine/Classes/Particles/Location/ParticleModuleLocationPrimitiveTriangle.h b/Engine/Source/Runtime/Engine/Classes/Particles/Location/ParticleModuleLocationPrimitiveTriangle.h
--- a/Engine/Source/Runtime/Engine/Classes/Particles/Location/ParticleModuleLocationPrimitiveTriangle.h
+++ b/Engine/Source/Runtime/Engine/Classes/Particles/Location/ParticleModuleLocationPrimitiveTriangle.h
@@ -9,8 +9,8 @@
 #include "Particles/Location/ParticleModuleLocationBase.h"
 #include "ParticleModuleLocationPrimitiveTriangle.generated.h"
 
-UCLASS(editinlinenew, hidecategories=Object, MinimalAPI, meta=(DisplayName = "Triangle"))
-class UParticleModuleLocationPrimitiveTriangle : public UParticleModuleLocationBase
+UCLASS(editinlinenew, hidecategories=Object, meta=(DisplayName = "Triangle"))
+class ENGINE_API UParticleModuleLocationPrimitiveTriangle : public UParticleModuleLocationBase
 {
 	GENERATED_UCLASS_BODY()
 
diff --git a/Engine/Source/Runtime/Engine/Classes/Particles/Location/ParticleModuleLocationSkelVertSurface.h b/Engine/Source/Runtime/Engine/Classes/Particles/Location/ParticleModuleLocationSkelVertSurface.h
--- a/Engine/Source/Runtime/Engine/Classes/Particles/Location/ParticleModuleLocationSkelVertSurface.h
+++ b/Engine/Source/Runtime/Engine/Classes/Particles/Location/ParticleModuleLocationSkelVertSurface.h
@@ -14,7 +14,7 @@ enum ELocationSkelVertSurfaceSource
 };
 
 UCLASS(editinlinenew, hidecategories=Object, meta=(DisplayName = "Skel Vert/Surf Location"))
-class UParticleModuleLocationSkelVertSurface : public UParticleModuleLocationBase
+class ENGINE_API UParticleModuleLocationSkelVertSurface : public UParticleModuleLocationBase
 {
 	GENERATED_UCLASS_BODY()
 
diff --git a/Engine/Source/Runtime/Engine/Classes/Particles/Location/ParticleModuleLocationWorldOffset.h b/Engine/Source/Runtime/Engine/Classes/Particles/Location/ParticleModuleLocationWorldOffset.h
--- a/Engine/Source/Runtime/Engine/Classes/Particles/Location/ParticleModuleLocationWorldOffset.h
+++ b/Engine/Source/Runtime/Engine/Classes/Particles/Location/ParticleModuleLocationWorldOffset.h
@@ -6,7 +6,7 @@
 #include "ParticleModuleLocationWorldOffset.generated.h"
 
 UCLASS(editinlinenew, meta=(DisplayName = "World Offset"))
-class UParticleModuleLocationWorldOffset : public UParticleModuleLocation
+class ENGINE_API UParticleModuleLocationWorldOffset : public UParticleModuleLocation
 {
 	GENERATED_UCLASS_BODY()
 
diff --git a/Engine/Source/Runtime/Engine/Classes/Particles/Location/ParticleModuleLocationWorldOffset_Seeded.h b/Engine/Source/Runtime/Engine/Classes/Particles/Location/ParticleModuleLocationWorldOffset_Seeded.h
--- a/Engine/Source/Runtime/Engine/Classes/Particles/Location/ParticleModuleLocationWorldOffset_Seeded.h
+++ b/Engine/Source/Runtime/Engine/Classes/Particles/Location/ParticleModuleLocationWorldOffset_Seeded.h
@@ -6,7 +6,7 @@
 #include "ParticleModuleLocationWorldOffset_Seeded.generated.h"
 
 UCLASS(editinlinenew, meta=(DisplayName = "World Offset (Seed)"))
-class UParticleModuleLocationWorldOffset_Seeded : public UParticleModuleLocationWorldOffset
+class ENGINE_API UParticleModuleLocationWorldOffset_Seeded : public UParticleModuleLocationWorldOffset
 {
 	GENERATED_UCLASS_BODY()
 
diff --git a/Engine/Source/Runtime/Engine/Classes/Particles/Location/ParticleModuleLocation_Seeded.h b/Engine/Source/Runtime/Engine/Classes/Particles/Location/ParticleModuleLocation_Seeded.h
--- a/Engine/Source/Runtime/Engine/Classes/Particles/Location/ParticleModuleLocation_Seeded.h
+++ b/Engine/Source/Runtime/Engine/Classes/Particles/Location/ParticleModuleLocation_Seeded.h
@@ -6,7 +6,7 @@
 #include "ParticleModuleLocation_Seeded.generated.h"
 
 UCLASS(editinlinenew, hidecategories=Object, meta=(DisplayName = "Initial Location (Seed)"))
-class UParticleModuleLocation_Seeded : public UParticleModuleLocation
+class ENGINE_API UParticleModuleLocation_Seeded : public UParticleModuleLocation
 {
 	GENERATED_UCLASS_BODY()
 
diff --git a/Engine/Source/Runtime/Engine/Classes/Particles/Location/ParticleModuleSourceMovement.h b/Engine/Source/Runtime/Engine/Classes/Particles/Location/ParticleModuleSourceMovement.h
--- a/Engine/Source/Runtime/Engine/Classes/Particles/Location/ParticleModuleSourceMovement.h
+++ b/Engine/Source/Runtime/Engine/Classes/Particles/Location/ParticleModuleSourceMovement.h
@@ -6,7 +6,7 @@
 #include "ParticleModuleSourceMovement.generated.h"
 
 UCLASS(editinlinenew, hidecategories=Object, meta=(DisplayName = "Source Movement"))
-class UParticleModuleSourceMovement : public UParticleModuleLocationBase
+class ENGINE_API UParticleModuleSourceMovement : public UParticleModuleLocationBase
 {
 	GENERATED_UCLASS_BODY()
 
diff --git a/Engine/Source/Runtime/Engine/Classes/Particles/Modules/Location/ParticleModulePivotOffset.h b/Engine/Source/Runtime/Engine/Classes/Particles/Modules/Location/ParticleModulePivotOffset.h
--- a/Engine/Source/Runtime/Engine/Classes/Particles/Modules/Location/ParticleModulePivotOffset.h
+++ b/Engine/Source/Runtime/Engine/Classes/Particles/Modules/Location/ParticleModulePivotOffset.h
@@ -4,8 +4,8 @@
 #include "Particles/Location/ParticleModuleLocationBase.h"
 #include "ParticleModulePivotOffset.generated.h"
 
-UCLASS(editinlinenew, hidecategories=Object, MinimalAPI, DisplayName="Pivot Offset")
-class UParticleModulePivotOffset : public UParticleModuleLocationBase
+UCLASS(editinlinenew, hidecategories=Object, DisplayName="Pivot Offset")
+class ENGINE_API UParticleModulePivotOffset : public UParticleModuleLocationBase
 {
 	GENERATED_UCLASS_BODY()
 
diff --git a/Engine/Source/Runtime/Engine/Classes/Particles/ParticleSystem.h b/Engine/Source/Runtime/Engine/Classes/Particles/ParticleSystem.h
--- a/Engine/Source/Runtime/Engine/Classes/Particles/ParticleSystem.h
+++ b/Engine/Source/Runtime/Engine/Classes/Particles/ParticleSystem.h
@@ -455,7 +455,7 @@ public:
 	/**
 	 *	Set the distance for the given LOD index
 	 *
-	 *	@param	LODLevelIndex	The LOD level to set the distance ofr
+	 *	@param	LODLevelIndex	The LOD level to set the distance of
 	 *	@param	InDistance		The distance to set
 	 *
 	 *	@return	true			If successful
@@ -463,6 +463,16 @@ public:
 	 */
 	virtual bool SetLODDistance(int32 LODLevelIndex, float InDistance);
 
+	/**
+	*	Checks if any of the emitters have motion blur at a specific lod level.
+	*
+	*	@param	LODLevelIndex	The LOD level to check motion blur availability
+	*
+	*	@return	true			If any emitter has motion blur
+	*			false			None of the emitters have motion blur
+	*/
+	bool DoesAnyEmitterHaveMotionBlur(int32 LODLevelIndex) const;
+
 	/**
 	 * Builds all emitters in the particle system.
 	 */
diff --git a/Engine/Source/Runtime/Engine/Classes/Particles/TypeData/ParticleModuleTypeDataBase.h b/Engine/Source/Runtime/Engine/Classes/Particles/TypeData/ParticleModuleTypeDataBase.h
--- a/Engine/Source/Runtime/Engine/Classes/Particles/TypeData/ParticleModuleTypeDataBase.h
+++ b/Engine/Source/Runtime/Engine/Classes/Particles/TypeData/ParticleModuleTypeDataBase.h
@@ -40,6 +40,10 @@ class UParticleModuleTypeDataBase : public UParticleModule
 	virtual bool		SupportsSubUV() const	{ return false; }
 	// @todo document
 	virtual bool		IsAMeshEmitter() const	{ return false; }
+
+	/** Determine if motion blur is enabled for the owning emitter. */
+	virtual bool		IsMotionBlurEnabled() const  { return false; }
+
 };
 
 
diff --git a/Engine/Source/Runtime/Engine/Classes/Particles/TypeData/ParticleModuleTypeDataMesh.h b/Engine/Source/Runtime/Engine/Classes/Particles/TypeData/ParticleModuleTypeDataMesh.h
--- a/Engine/Source/Runtime/Engine/Classes/Particles/TypeData/ParticleModuleTypeDataMesh.h
+++ b/Engine/Source/Runtime/Engine/Classes/Particles/TypeData/ParticleModuleTypeDataMesh.h
@@ -89,6 +89,8 @@ class UParticleModuleTypeDataMesh : public UParticleModuleTypeDataBase
 	UPROPERTY(EditAnywhere, Category=Mesh)
 	uint32 bOverrideMaterial:1;
 
+	UPROPERTY(EditAnywhere, Category = Mesh)
+	uint32 bEnableMotionBlur : 1;
 
 	/** deprecated properties for initial orientation */
 	UPROPERTY()
@@ -204,6 +206,7 @@ class UParticleModuleTypeDataMesh : public UParticleModuleTypeDataBase
 	virtual bool	SupportsSpecificScreenAlignmentFlags() const override {	return true;	}	
 	virtual bool	SupportsSubUV() const override { return true; }
 	virtual bool	IsAMeshEmitter() const override { return true; }
+	virtual bool    IsMotionBlurEnabled() const override { return bEnableMotionBlur; }
 	//~ End UParticleModuleTypeDataBase Interface
 };
 
diff --git a/Engine/Source/Runtime/Engine/Private/Components/PrimitiveComponent.cpp b/Engine/Source/Runtime/Engine/Private/Components/PrimitiveComponent.cpp
--- a/Engine/Source/Runtime/Engine/Private/Components/PrimitiveComponent.cpp
+++ b/Engine/Source/Runtime/Engine/Private/Components/PrimitiveComponent.cpp
@@ -786,7 +786,7 @@ void UPrimitiveComponent::CheckForErrors()
 			->AddToken(FMapErrorToken::Create(FMapErrors::ShadowCasterUsingBoundsScale));
 	}
 
-	if (HasStaticLighting() && !HasValidSettingsForStaticLighting() && (!Owner || !Owner->IsA(AWorldSettings::StaticClass())))	// Ignore worldsettings
+	if (HasStaticLighting() && !HasValidSettingsForStaticLighting(true) && (!Owner || !Owner->IsA(AWorldSettings::StaticClass())))	// Ignore worldsettings
 	{
 		FMessageLog("MapCheck").Error()
 			->AddToken(FUObjectToken::Create(Owner))
