Post proc selection outline optimization

diff --git a/Engine/Shaders/PostProcessSelectionOutline.usf b/Engine/Shaders/PostProcessSelectionOutline.usf
--- a/Engine/Shaders/PostProcessSelectionOutline.usf
+++ b/Engine/Shaders/PostProcessSelectionOutline.usf
@@ -115,14 +115,14 @@ void MinMax(inout int2 Var, int Value)
 	Var.y = max(Var.y, Value);
 }
 
-float4 PerSample(int2 PixelPos, int SampleID, float3 DeviceZPlane, float2 DeviceZMinMax)
+float4 PerSample(int2 PixelPos, int SampleID, FPixelInfo CenterPixelInfo, float3 DeviceZPlane, float2 DeviceZMinMax)
 {
 	float4 Ret = 0;
 
 	// [0]:center, [1..4]:borders
 	FPixelInfo PixelInfo[5];
 	
-	PixelInfo[0] = GetPixelInfo(PixelPos, 0, 0,  0, DeviceZPlane, DeviceZMinMax);
+	PixelInfo[0] = CenterPixelInfo;
 
 	// diagonal cross is thicker than vertical/horizontal cross
 	PixelInfo[1] = GetPixelInfo(PixelPos, SampleID,  1,  1, DeviceZPlane, DeviceZMinMax);
@@ -212,47 +212,72 @@ void MainPS(noperspective float4 UVAndScreenPos : TEXCOORD0, float4 SvPosition :
 {
 	float2 UV = UVAndScreenPos.xy;
  	int2 PixelPos = (int2)SvPosition.xy;
+	
+	OutColor = Texture2DSample(PostprocessInput0, PostprocessInput0Sampler, UVAndScreenPos.xy);
 
-	float SceneDepth = CalcSceneDepth(UV);
+	// Scout outwards from the center pixel to determine if any neighboring pixels are selected
+	int ScountStep = 2;
+#if MSAA_SAMPLE_COUNT > 1
+	int4 StencilScout = int4(
+		EditorPrimitivesStencil.Load(PixelPos + int2( ScountStep,  0), 0).g,
+		EditorPrimitivesStencil.Load(PixelPos + int2(-ScountStep,  0), 0).g,
+		EditorPrimitivesStencil.Load(PixelPos + int2( 0,  ScountStep), 0).g,
+		EditorPrimitivesStencil.Load(PixelPos + int2( 0, -ScountStep), 0).g
+	);
+#else
+	int4 StencilScout = int4(
+		EditorPrimitivesStencil.Load(int3(PixelPos + int2( ScountStep,  0), 0)).g,
+		EditorPrimitivesStencil.Load(int3(PixelPos + int2(-ScountStep,  0), 0)).g,
+		EditorPrimitivesStencil.Load(int3(PixelPos + int2( 0,  ScountStep), 0)).g,
+		EditorPrimitivesStencil.Load(int3(PixelPos + int2( 0, -ScountStep), 0)).g
+	);
+#endif
+	int ScoutSum = dot(saturate(StencilScout), 1);
 
-	// This allows to reconstruct depth with a small pixel offset without many tetxure lookups (4xMSAA * 5 neighbors -> 20 samples)
-	// It's an approximation assuming the surface is a plane
-	float3 DeviceZPlane;
-	DeviceZPlane.z = Texture2DSampleLevel(SceneDepthTexture, SceneDepthTextureSampler, UV, 0).r;
+	// If this sum is zero, none of our neighbors are selected pixels and we can skip all the heavy processing
+	if (ScoutSum > 0) {
+		float SceneDepth = CalcSceneDepth(UV);
 
-	float Left = Texture2DSampleLevel(SceneDepthTexture, SceneDepthTextureSampler, UV + float2(-1, 0) * Frame.BufferSizeAndInvSize.zw, 0).r;
-	float Right= Texture2DSampleLevel(SceneDepthTexture, SceneDepthTextureSampler, UV + float2(1, 0) * Frame.BufferSizeAndInvSize.zw, 0).r;
-	float Top  = Texture2DSampleLevel(SceneDepthTexture, SceneDepthTextureSampler, UV + float2(0, -1) * Frame.BufferSizeAndInvSize.zw, 0).r;
-	float Bottom = Texture2DSampleLevel(SceneDepthTexture, SceneDepthTextureSampler, UV + float2(0, 1) * Frame.BufferSizeAndInvSize.zw, 0).r;
+		// This allows to reconstruct depth with a small pixel offset without many tetxure lookups (4xMSAA * 5 neighbors -> 20 samples)
+		// It's an approximation assuming the surface is a plane
+		float3 DeviceZPlane;
+		DeviceZPlane.z = Texture2DSampleLevel(SceneDepthTexture, SceneDepthTextureSampler, UV, 0).r;
 
-	float2 DeviceZMinMax;
-	DeviceZMinMax.x = min(DeviceZPlane.z, min(min(Left, Right), min(Top, Bottom)));
-	DeviceZMinMax.y = max(DeviceZPlane.z, max(max(Left, Right), max(Top, Bottom)));
+		float Left = Texture2DSampleLevel(SceneDepthTexture, SceneDepthTextureSampler, UV + float2(-1, 0) * Frame.BufferSizeAndInvSize.zw, 0).r;
+		float Right= Texture2DSampleLevel(SceneDepthTexture, SceneDepthTextureSampler, UV + float2(1, 0) * Frame.BufferSizeAndInvSize.zw, 0).r;
+		float Top  = Texture2DSampleLevel(SceneDepthTexture, SceneDepthTextureSampler, UV + float2(0, -1) * Frame.BufferSizeAndInvSize.zw, 0).r;
+		float Bottom = Texture2DSampleLevel(SceneDepthTexture, SceneDepthTextureSampler, UV + float2(0, 1) * Frame.BufferSizeAndInvSize.zw, 0).r;
 
-	DeviceZPlane.x = (Right - Left) / 2;
-	DeviceZPlane.y = (Bottom - Top) / 2;
+		float2 DeviceZMinMax;
+		DeviceZMinMax.x = min(DeviceZPlane.z, min(min(Left, Right), min(Top, Bottom)));
+		DeviceZMinMax.y = max(DeviceZPlane.z, max(max(Left, Right), max(Top, Bottom)));
 
-	// even stronger appromiation (faster but too many artifacts)
-//	DeviceZPlane.x = ddx(DeviceZPlane.z);
-//	DeviceZPlane.y = ddy(DeviceZPlane.z);
+		DeviceZPlane.x = (Right - Left) / 2;
+		DeviceZPlane.y = (Bottom - Top) / 2;
 
-	float4 Sum = 0;
-	for(int SampleID = 0; SampleID < MSAA_SAMPLE_COUNT; ++SampleID)
-	{
-		Sum += PerSample(PixelPos, SampleID, DeviceZPlane, DeviceZMinMax);
-	}
-	float4 Avg = Sum / MSAA_SAMPLE_COUNT;
+		// even stronger appromiation (faster but too many artifacts)
+	//	DeviceZPlane.x = ddx(DeviceZPlane.z);
+	//	DeviceZPlane.y = ddy(DeviceZPlane.z);
 
-	OutColor = Texture2DSample(PostprocessInput0, PostprocessInput0Sampler, UVAndScreenPos.xy);
+		FPixelInfo CenterPixelInfo = GetPixelInfo(PixelPos, 0, 0,  0, DeviceZPlane, DeviceZMinMax);
+
+		float4 Sum = 0;
+		for(int SampleID = 0; SampleID < MSAA_SAMPLE_COUNT; ++SampleID)
+		{
+			Sum += PerSample(PixelPos, SampleID, CenterPixelInfo, DeviceZPlane, DeviceZMinMax);
+		}
+		float4 Avg = Sum / MSAA_SAMPLE_COUNT;
 
-	// Dest color has gamma applied so we need to remove it before applying linear color
-	OutColor.rgb = pow( OutColor.rgb, 2.2f );
+		OutColor = Texture2DSample(PostprocessInput0, PostprocessInput0Sampler, UVAndScreenPos.xy);
 
-	OutColor.rgb = lerp(OutColor.rgb, Avg.rgb, Avg.a);
+		// Dest color has gamma applied so we need to remove it before applying linear color
+		OutColor.rgb = pow( OutColor.rgb, 2.2f );
 
-	// Re-apply gamma corrrection
-	OutColor.rgb = pow( OutColor.rgb, 1/2.2f );
+		OutColor.rgb = lerp(OutColor.rgb, Avg.rgb, Avg.a);
 
+		// Re-apply gamma corrrection
+		OutColor.rgb = pow( OutColor.rgb, 1/2.2f );
+	}
 /*
 	if(MSAA_SAMPLE_COUNT == 2)
 	{
