fix for Mac DDC generation failures

diff --git a/Engine/Shaders/GlobalDistanceFieldShared.usf b/Engine/Shaders/GlobalDistanceFieldShared.usf
--- a/Engine/Shaders/GlobalDistanceFieldShared.usf
+++ b/Engine/Shaders/GlobalDistanceFieldShared.usf
@@ -9,7 +9,7 @@
 float MaxGlobalDistance;
 
 //@todo - get the clipmap branches to compile under gl
-#if (SM5_PROFILE || PS4_PROFILE || METAL_SM5_PROFILE) && !COMPILER_HLSLCC
+#if SM5_PROFILE || PS4_PROFILE || METAL_SM5_PROFILE
 
 // Must match C++
 #define MAX_GLOBAL_DF_CLIPMAPS 4
