Bump Shaderversion to wipe DDC.  Fix Paragon decals. UE-28554

diff --git a/Engine/Shaders/ShaderVersion.usf b/Engine/Shaders/ShaderVersion.usf
--- a/Engine/Shaders/ShaderVersion.usf
+++ b/Engine/Shaders/ShaderVersion.usf
@@ -2,4 +2,4 @@
 // Each time the console command is executed it generates a new GUID. As this file is included
 // in common.usf (which should be included in any shader) it allows to invalidate the shader DDC.
 // 
-// GUID = 4262103F821247F4868CDF3249343BB  
+// GUID = 2F8BBA00C044A97B429F64CD2CFE7B7
