Copying //UE4/Dev-Rendering to Dev-Main (Source //UE4/Dev-Rendering@2932636)

==========================
MAJOR FEATURES + CHANGES
==========================

Change 2917472 on 2016/03/21 by [REDACTED_PERSON]

	DR - Fix SCW directcompile arguments, add -pipeline

Change 2919580 on 2016/03/23 by [REDACTED_PERSON]

	DR - HlslParser - Fix for used elements (sparrow's arrow was showing when it shouldn't)
	Arrays of input/outputs are now flattened so disjoint entries can be optimized out (and fixes a bug)
	#tests Run game with sparrow, test with slomo to check for gfx glitches

Change 2919660 on 2016/03/23 by [REDACTED_PERSON]

	DR - Latest vk changes (from dev mobile's 2916881 to 2919157)

Change 2919902 on 2016/03/23 by [REDACTED_PERSON]

	DR - Fix skeletal meshes decrementing stats twice

Change 2920020 on 2016/03/23 by [REDACTED_PERSON]

	EyeAdapation when used in material shader may not be initialized.

Change 2920071 on 2016/03/23 by [REDACTED_PERSON]

	DR - Remove old vk define
	- Started moving around direct calls to queue submit

Change 2920252 on 2016/03/23 by [REDACTED_PERSON]

	DR - Changes vk structs to classes

Change 2920314 on 2016/03/23 by [REDACTED_PERSON]

	Add -windowed to standalone game PIE command line to avoid PIE launching in full screen

Change 2920745 on 2016/03/24 by [REDACTED_PERSON]

	Texture streaming build now takes into account the material texcoord scales applied to the texture sampling.
	Also finds out which texcoord is being used when sampling textures (between 0 and 3 currently).
	TexCoord analysis debug view shaders is now working with SM4 ane SM5.
	StaticMeshComponents hold persistent data coming from the texture streaming build.
	#tests tested with different Paragon assets. Editor SM4 & SM5. Cooked maps

Change 2921335 on 2016/03/24 by [REDACTED_PERSON]

	Added missing static keyword for locally defined console variable.

Change 2921416 on 2016/03/24 by [REDACTED_PERSON]

	Revert enabling debugview shaders on non PC platforms (until properly tested and debugged)

Change 2921446 on 2016/03/24 by [REDACTED_PERSON]

	Planar reflection mesh

Change 2921530 on 2016/03/24 by [REDACTED_PERSON]

	Manual revert of Ronin planar reflections

Change 2921608 on 2016/03/24 by [REDACTED_PERSON]

	Updated texture streamer to take into account the new HLOD texture group.

Change 2921677 on 2016/03/24 by [REDACTED_PERSON]

	Distance Field Specular Occlusion
	* Prototype - disabled by default

Change 2921681 on 2016/03/24 by [REDACTED_PERSON]

	UnmappedTexelsPercentage is now 100 based

Change 2921682 on 2016/03/24 by [REDACTED_PERSON]

	Planar reflections
	* New actor and component
	* The scene is rendered to texture with a mirrored camera and a clip plane each frame
	* The reflection texture is then applied to opaque pixels in a deferred pass, with distance and angle from plane fades
	* Translucent materials apply the nearest reflection plane in the base pass
	* Planar reflections require the project setting 'Support global clip plane for Planar Reflections' to be enabled, since writing to SV_ClipDistance all the time adds about 15% BasePass GPU time on PS4
	* Fixed global distance field in materials which had been broken since moving global distance field properties into the view uniform buffer
	* Fixed PS4 removing system-value semantics when output from vertex shader and not read in next stage

Change 2921734 on 2016/03/24 by [REDACTED_PERSON]

	Fixed tessellated cube having wrong UVs

Change 2922063 on 2016/03/24 by [REDACTED_PERSON]

	Removed planar reflection debug code

Change 2922428 on 2016/03/25 by [REDACTED_PERSON]

	Delete FShaderPipeline objects when clearing TMaterialShaderMaps.

Change 2922803 on 2016/03/25 by [REDACTED_PERSON]

	DR - New cmd buffer management (disabled)
	- Move cmd buffer out of pending state and into context
	- Do not hardcode # cmd buffers
	- Move back buffer image mgmt into swapchain
	- Fixed some image layout transition bugs

Change 2923056 on 2016/03/25 by [REDACTED_PERSON]

	DR - Initial fix for canvas locking inside a render pass

diff --git a/Engine/Config/Windows/WindowsEngine.ini b/Engine/Config/Windows/WindowsEngine.ini
--- a/Engine/Config/Windows/WindowsEngine.ini
+++ b/Engine/Config/Windows/WindowsEngine.ini
@@ -7,6 +7,11 @@ PlatformHeadroomDB=0
 ; PoolSizeVRAMPercentage is how much percentage of GPU Dedicated VRAM should be used as a TexturePool cache for streaming textures (0 - unlimited streaming)
 PoolSizeVRAMPercentage=70
 
+[DebugViewModes]
+bEnableQuadOverdraw=true
+bEnableStreamingAccuracy=true
+bEnableTextureStreamingBuild=true
+
 [Engine.DeviceConfiguration]
 ConfigClass=/Engine/DeviceProfiles/SetDeviceProfileWindows.SetDeviceProfileWindows_C
 
diff --git a/Engine/Plugins/2D/Paper2D/Source/Paper2D/Classes/PaperFlipbookComponent.h b/Engine/Plugins/2D/Paper2D/Source/Paper2D/Classes/PaperFlipbookComponent.h
--- a/Engine/Plugins/2D/Paper2D/Source/Paper2D/Classes/PaperFlipbookComponent.h
+++ b/Engine/Plugins/2D/Paper2D/Source/Paper2D/Classes/PaperFlipbookComponent.h
@@ -191,7 +191,7 @@ public:
 	virtual void GetUsedTextures(TArray<UTexture*>& OutTextures, EMaterialQualityLevel::Type QualityLevel) override;
 	virtual UMaterialInterface* GetMaterial(int32 MaterialIndex) const override;
 	virtual void GetUsedMaterials(TArray<UMaterialInterface*>& OutMaterials) const override;
-	virtual void GetStreamingTextureInfo(TArray<FStreamingTexturePrimitiveInfo>& OutStreamingTextures) const override;
+	virtual void GetStreamingTextureInfo(FStreamingTextureLevelContext& LevelContext, TArray<FStreamingTexturePrimitiveInfo>& OutStreamingTextures) const override;
 	virtual int32 GetNumMaterials() const override;
 	virtual UBodySetup* GetBodySetup() override;
 	// End of UPrimitiveComponent interface
diff --git a/Engine/Plugins/2D/Paper2D/Source/Paper2D/Classes/PaperSpriteComponent.h b/Engine/Plugins/2D/Paper2D/Source/Paper2D/Classes/PaperSpriteComponent.h
--- a/Engine/Plugins/2D/Paper2D/Source/Paper2D/Classes/PaperSpriteComponent.h
+++ b/Engine/Plugins/2D/Paper2D/Source/Paper2D/Classes/PaperSpriteComponent.h
@@ -83,7 +83,7 @@ public:
 	virtual void GetUsedTextures(TArray<UTexture*>& OutTextures, EMaterialQualityLevel::Type QualityLevel) override;
 	virtual UMaterialInterface* GetMaterial(int32 MaterialIndex) const override;
 	virtual void GetUsedMaterials(TArray<UMaterialInterface*>& OutMaterials) const override;
-	virtual void GetStreamingTextureInfo(TArray<FStreamingTexturePrimitiveInfo>& OutStreamingTextures) const override;
+	virtual void GetStreamingTextureInfo(FStreamingTextureLevelContext& LevelContext, TArray<FStreamingTexturePrimitiveInfo>& OutStreamingTextures) const override;
 	virtual int32 GetNumMaterials() const override;
 	// End of UPrimitiveComponent interface
 
diff --git a/Engine/Plugins/2D/Paper2D/Source/Paper2D/Private/PaperFlipbookComponent.cpp b/Engine/Plugins/2D/Paper2D/Source/Paper2D/Private/PaperFlipbookComponent.cpp
--- a/Engine/Plugins/2D/Paper2D/Source/Paper2D/Private/PaperFlipbookComponent.cpp
+++ b/Engine/Plugins/2D/Paper2D/Source/Paper2D/Private/PaperFlipbookComponent.cpp
@@ -168,10 +168,10 @@ void UPaperFlipbookComponent::GetUsedMaterials(TArray<UMaterialInterface*>& OutM
 	return Super::GetUsedMaterials(OutMaterials);
 }
 
-void UPaperFlipbookComponent::GetStreamingTextureInfo(TArray<FStreamingTexturePrimitiveInfo>& OutStreamingTextures) const
+void UPaperFlipbookComponent::GetStreamingTextureInfo(FStreamingTextureLevelContext& LevelContext, TArray<FStreamingTexturePrimitiveInfo>& OutStreamingTextures) const
 {
 	//@TODO: PAPER2D: Need to support this for proper texture streaming
-	return Super::GetStreamingTextureInfo(OutStreamingTextures);
+	return Super::GetStreamingTextureInfo(LevelContext, OutStreamingTextures);
 }
 
 int32 UPaperFlipbookComponent::GetNumMaterials() const
diff --git a/Engine/Plugins/2D/Paper2D/Source/Paper2D/Private/PaperSpriteComponent.cpp b/Engine/Plugins/2D/Paper2D/Source/Paper2D/Private/PaperSpriteComponent.cpp
--- a/Engine/Plugins/2D/Paper2D/Source/Paper2D/Private/PaperSpriteComponent.cpp
+++ b/Engine/Plugins/2D/Paper2D/Source/Paper2D/Private/PaperSpriteComponent.cpp
@@ -272,10 +272,10 @@ void UPaperSpriteComponent::GetUsedMaterials(TArray<UMaterialInterface*>& OutMat
 	return Super::GetUsedMaterials(OutMaterials);
 }
 
-void UPaperSpriteComponent::GetStreamingTextureInfo(TArray<FStreamingTexturePrimitiveInfo>& OutStreamingTextures) const
+void UPaperSpriteComponent::GetStreamingTextureInfo(FStreamingTextureLevelContext& LevelContext, TArray<FStreamingTexturePrimitiveInfo>& OutStreamingTextures) const
 {
 	//@TODO: PAPER2D: Need to support this for proper texture streaming
-	return Super::GetStreamingTextureInfo(OutStreamingTextures);
+	return Super::GetStreamingTextureInfo(LevelContext, OutStreamingTextures);
 }
 
 int32 UPaperSpriteComponent::GetNumMaterials() const
diff --git a/Engine/Shaders/BasePassForForwardShadingPixelShader.usf b/Engine/Shaders/BasePassForForwardShadingPixelShader.usf
--- a/Engine/Shaders/BasePassForForwardShadingPixelShader.usf
+++ b/Engine/Shaders/BasePassForForwardShadingPixelShader.usf
@@ -23,10 +23,6 @@ float4 LightPositionAndInvRadius[MAX_DYNAMIC_POINT_LIGHTS];
 float4 LightColorAndFalloffExponent[MAX_DYNAMIC_POINT_LIGHTS];
 #endif
 
-// xyz: Reflection plane normal in world space. Zero vector if this is not a planar reflection capture.
-// w: Reflection plane height.
-float4 PlanarReflectionCaptureData;
-
 #ifndef FORWARD_QL_FORCE_FULLY_ROUGH
 #define FORWARD_QL_FORCE_FULLY_ROUGH 0
 #endif
@@ -79,15 +75,6 @@ half PhongApprox( half Roughness, half RoL )
 	// Total 7 instr
 }
 
-// Returns 1 if the position is below the reflection plane, 0 if it is above.
-half GetPlanarReflectionMask(half3 WorldPosition) {
-	const half3 PlaneOrigin = half3(PlanarReflectionCaptureData.xyz * PlanarReflectionCaptureData.w);
-	const half3 PlaneNormal = PlanarReflectionCaptureData.xyz;
-	const half3 V = WorldPosition - PlaneOrigin;
-	const half D = dot(PlaneNormal, V);
-	return (half)(D <= 0.0);
-}
-
 #if !FULLY_ROUGH
 #if HQ_REFLECTIONS
 
@@ -463,16 +450,8 @@ void Main(
 		#else
 			Color += (Shadow * NoL) * Frame.DirectionalLightColor.rgb * (DiffuseColor + SpecularColor * PhongApprox(Roughness, RoL));
 			half3 SpecularIBL = GetImageBasedReflectionLighting(MaterialParameters, Roughness);
-
-			#if MATERIAL_USE_PLANAR_REFLECTION	
-				//!! Ronin Hack to reuse emissive
-				const half3 EnvironmentReflection = SpecularIBL * IndirectIrradiance * SpecularColor; // Environment map has been prenormalized, scale by lightmap luminance
-				const half3 PlanarReflection = GetMaterialEmissive(PixelMaterialInputs) * SpecularColor;
-				const half PlanarReflectionMask = GetMaterialAmbientOcclusion(PixelMaterialInputs); // RT Alpha, 0 where there's reflection data
-				Color += lerp(PlanarReflection, EnvironmentReflection, PlanarReflectionMask);
-			#else
-				Color += SpecularIBL * IndirectIrradiance * SpecularColor;
-			#endif
+			// Environment map has been prenormalized, scale by lightmap luminance
+			Color += SpecularIBL * IndirectIrradiance * SpecularColor;
 		#endif
 
 		#if ENABLE_SKY_LIGHT
@@ -532,12 +511,9 @@ void Main(
 
 	#endif /* !MATERIAL_SHADINGMODEL_UNLIT */
 		 
+	half3 Emissive = GetMaterialEmissive(PixelMaterialInputs);
 
-	// Emissive is used for planar reflection contribution and applied with reflections.
-	#if !MATERIAL_USE_PLANAR_REFLECTION
-		half3 Emissive = GetMaterialEmissive(PixelMaterialInputs);
-		Color += Emissive;
-	#endif
+	Color += Emissive;
 
 	half4 VertexFog = half4(0, 0, 0, 1);
 
@@ -607,14 +583,4 @@ void Main(
 		OutColor *= ClipForEditorPrimitives(MaterialParameters);
 		clip(OutColor.a - GetMaterialOpacityMaskClipValue());
 	#endif
-
-	// Apply planar reflection pass mask
-	BRANCH
-	if (any(PlanarReflectionCaptureData.xyz))
-	{
-		// Ronin TODO: make this use view-relative position
-		const half PlanarReflectionMask = GetPlanarReflectionMask(MaterialParameters.AbsoluteWorldPosition);
-		OutColor.rgb *= PlanarReflectionMask;
-		OutColor.a = 1.0 - PlanarReflectionMask; // Alpha is cleared to 1
-	}
 }
diff --git a/Engine/Shaders/BasePassForForwardShadingVertexShader.usf b/Engine/Shaders/BasePassForForwardShadingVertexShader.usf
--- a/Engine/Shaders/BasePassForForwardShadingVertexShader.usf
+++ b/Engine/Shaders/BasePassForForwardShadingVertexShader.usf
@@ -17,10 +17,6 @@
 #include "VertexFactory.usf"
 #include "HeightFogCommon.usf"
 
-// xyz: Reflection plane normal in world space. Zero vector if this is not a planar reflection capture.
-// w: Reflection plane height.
-float4 PlanarReflectionCaptureData;
-
 struct FForwardShadingBasePassVSToPS
 {
 	FVertexFactoryInterpolantsVSToPS FactoryInterpolants;
@@ -31,18 +27,6 @@ struct FForwardShadingBasePassVSToPS
 #define FForwardShadingBasePassVSOutput FForwardShadingBasePassVSToPS
 #define VertexFactoryGetInterpolants VertexFactoryGetInterpolantsVSToPS
 
-// Reflect InP about an xy plane in world space.
-half3 GetPlanarReflectionPosition(const half3 InP) {
-	const half3 PlaneOrigin = half3(PlanarReflectionCaptureData.xyz * PlanarReflectionCaptureData.w);
-	const half3 PlaneNormal = PlanarReflectionCaptureData.xyz;
-
-	half3 P = InP - View.PreViewTranslation.xyz;
-	const half3 V = P - PlaneOrigin;
-	const half D = dot(PlaneNormal, V) * 2.0;
-	P -= PlaneNormal * D;
-	return P + View.PreViewTranslation.xyz;
-}
-
 /** Entry point for the base pass vertex shader. */
 void Main(
 	FVertexFactoryInput Input,
@@ -71,13 +55,6 @@ void Main(
 	
 	WorldPosition.xyz += WorldPositionOffset;
 
-	// Apply planar reflection
-	BRANCH
-	if (any(PlanarReflectionCaptureData.xyz))
-	{
-		WorldPosition.xyz = GetPlanarReflectionPosition(WorldPosition.xyz);
-	}
-
 	float4 RasterizedWorldPosition = VertexFactoryGetRasterizedWorldPosition(Input, VFIntermediates, WorldPosition);
 	Output.Position = mul(RasterizedWorldPosition, ResolvedView.TranslatedWorldToClip);
 	Output.BasePassInterpolants.PixelPosition = WorldPosition;
diff --git a/Engine/Shaders/BasePassPixelShader.usf b/Engine/Shaders/BasePassPixelShader.usf
--- a/Engine/Shaders/BasePassPixelShader.usf
+++ b/Engine/Shaders/BasePassPixelShader.usf
@@ -11,6 +11,7 @@
 #include "VertexFactory.usf"
 #include "LightmapCommon.usf"  
 #include "ReflectionEnvironmentShared.usf"
+#include "PlanarReflectionShared.usf"
 #include "BRDF.usf"
 #include "Random.usf"
 #include "LightAccumulator.usf"
@@ -143,7 +144,19 @@ half3 GetImageBasedReflectionLighting(FMaterialPixelParameters MaterialParameter
 	half NoV = saturate(dot(MaterialParameters.WorldNormal, MaterialParameters.CameraVector));
 	SpecularColor = EnvBRDFApprox(SpecularColor, Roughness, NoV);
 
-	return SpecularIBL.rgb * SpecularColor;
+	float3 SpecularLighting = SpecularIBL.rgb;
+
+	// Plane normal will be zero if the feature is disabled
+	BRANCH
+	if (abs(dot(ReflectionPlane.xyz, 1)) > .0001f)
+	{
+		// Reuse ReflectionCubemapSampler to avoid reducing the sampler count available to artists
+		float4 PlanarReflection = ComputePlanarReflections(MaterialParameters.AbsoluteWorldPosition, MaterialParameters.WorldNormal, ReflectionCubemapSampler);
+		// Planar reflections win over SSR and reflection environment
+		SpecularLighting = lerp(SpecularLighting, PlanarReflection.rgb, PlanarReflection.a);
+	}
+
+	return SpecularLighting * SpecularColor;
 }
 
 #endif
diff --git a/Engine/Shaders/BasePassVertexShader.usf b/Engine/Shaders/BasePassVertexShader.usf
--- a/Engine/Shaders/BasePassVertexShader.usf
+++ b/Engine/Shaders/BasePassVertexShader.usf
@@ -16,10 +16,13 @@ void Main(
 	FVertexFactoryInput Input,
 	OPTIONAL_VertexID
 	out FBasePassVSOutput Output
+#if USE_GLOBAL_CLIP_PLANE && !USING_TESSELLATION
+	, out float OutGlobalClipPlaneDistance : SV_ClipDistance
+#endif
 #if INSTANCED_STEREO
 	, uint InstanceId : SV_InstanceID
-	, out float4 OutClipDistance : SV_ClipDistance
-	, out float4 OutCullDistance : SV_CullDistance
+	, out float OutClipDistance : SV_ClipDistance1
+	, out float OutCullDistance : SV_CullDistance1
 #endif
 	)
 {	
@@ -75,6 +78,10 @@ void Main(
 		#endif
 	}
 
+#if USE_GLOBAL_CLIP_PLANE
+	OutGlobalClipPlaneDistance = dot(ResolvedView.GlobalClippingPlane, float4(WorldPosition.xyz - ResolvedView.PreViewTranslation.xyz, 1));
+#endif
+
 	#if USE_WORLD_POSITION_EXCLUDING_SHADER_OFFSETS
 		Output.BasePassInterpolants.PixelPositionExcludingWPO = WorldPositionExcludingWPO;
 	#endif
diff --git a/Engine/Shaders/Common.usf b/Engine/Shaders/Common.usf
--- a/Engine/Shaders/Common.usf
+++ b/Engine/Shaders/Common.usf
@@ -220,6 +220,10 @@ MaterialFloat StoreTexSample(float4 C, uint TextureReferenceIndex)
 // @todo metal mrt: No reason this can't work with Metal, once cross compiler is fixed
 #define SUPPORTS_INDEPENDENT_SAMPLERS (PS4_PROFILE || SM5_PROFILE || SM4_PROFILE || METAL_MRT_PROFILE || METAL_SM4_PROFILE || METAL_SM5_PROFILE)
 
+// Whether the platform supports a global clip plane through SV_ClipDistance
+// Ideally this would just be enabled for all SM4 and above feature level platforms, but not tested everywhere yet
+#define SUPPORT_GLOBAL_CLIP_PLANE (PS4_PROFILE || SM5_PROFILE || SM4_PROFILE)
+
 // This would need to be a #define in GLSL to ignore the SamplerState, however, it is currently a function call in HLSL
 // for type checking of the parameters - ironically the type checking is really only needed in GLSL!
 MaterialFloat4 Texture1DSample(Texture1D Tex, SamplerState Sampler, float UV)
@@ -1397,7 +1401,8 @@ MaterialFloat3 TransformTangentVectorToWorld(MaterialFloat3x3 TangentToWorld, Ma
 /** Transforms a vector from world space to tangent space */
 MaterialFloat3 TransformWorldVectorToTangent(MaterialFloat3x3 TangentToWorld, MaterialFloat3 InWorldVector)
 {
-	// Transform from world to tangent space
+	// Transform from world to tangent space with the transpose of TangentToWorld (achieved by swapping vector / matrix multiply order)
+	// Note that the transpose is only equal to the inverse for orthonormal matrices - aka only uniform scaling
 	return mul(TangentToWorld, InWorldVector);
 }
 
diff --git a/Engine/Shaders/DebugViewModeCommon.usf b/Engine/Shaders/DebugViewModeCommon.usf
--- a/Engine/Shaders/DebugViewModeCommon.usf
+++ b/Engine/Shaders/DebugViewModeCommon.usf
@@ -14,10 +14,11 @@ struct FDebugVSToPS
 	// TexCoords must first in order to be compatible with viewmode requiring using only SV_Position
 	float4 Position : SV_POSITION;
 	float4 VertexColor : TEXCOORD0;
-	float2 TexCoord0 : TEXCOORD1;
-	float3 TangentToWorld0 : TEXCOORD2;
-	float3 TangentToWorld1 : TEXCOORD3;
-	float3 TangentToWorld2 : TEXCOORD4;
+	float4 TexCoord01 : TEXCOORD1;
+	float4 TexCoord23 : TEXCOORD2;
+	float3 TangentToWorld0 : TEXCOORD3;
+	float3 TangentToWorld1 : TEXCOORD4;
+	float3 TangentToWorld2 : TEXCOORD5;
 #if INSTANCED_STEREO
 	nointerpolation uint EyeIndex : EYE_INDEX;
 #endif
@@ -27,10 +28,11 @@ struct FDebugPSIn
 {
 	float4 SvPosition : SV_POSITION;
 	float4 VertexColor : TEXCOORD0;
-	float2 TexCoord0 : TEXCOORD1;
-	float3 TangentToWorld0 : TEXCOORD2;
-	float3 TangentToWorld1 : TEXCOORD3;
-	float3 TangentToWorld2 : TEXCOORD4;
+	float4 TexCoord01 : TEXCOORD1;
+	float4 TexCoord23 : TEXCOORD2;
+	float3 TangentToWorld0 : TEXCOORD3;
+	float3 TangentToWorld1 : TEXCOORD4;
+	float3 TangentToWorld2 : TEXCOORD5;
 #if INSTANCED_STEREO
 	nointerpolation uint EyeIndex : EYE_INDEX;
 #endif
@@ -56,7 +58,22 @@ FMaterialPixelParameters GetMaterialPixelParameters(FDebugPSIn DebugInputs, floa
 	Result.TangentToWorld = half3x3(half3(DebugInputs.TangentToWorld0), half3(DebugInputs.TangentToWorld1), half3(DebugInputs.TangentToWorld2));
 
 #if NUM_MATERIAL_TEXCOORDS > 0
-	Result.TexCoords[0] = DebugInputs.TexCoord0;
+	Result.TexCoords[0].xy = DebugInputs.TexCoord01.xy;
+#endif
+#if NUM_MATERIAL_TEXCOORDS > 1
+	Result.TexCoords[1].xy = DebugInputs.TexCoord01.zw;
+#endif
+#if NUM_MATERIAL_TEXCOORDS > 2
+	Result.TexCoords[2].xy = DebugInputs.TexCoord23.xy;
+#endif
+#if NUM_MATERIAL_TEXCOORDS > 3
+	Result.TexCoords[3].xy = DebugInputs.TexCoord23.zw;
+#endif
+#if NUM_MATERIAL_TEXCOORDS > 4
+	for(int CoordinateIndex = 4; CoordinateIndex < NUM_MATERIAL_TEXCOORDS; CoordinateIndex++)
+	{
+		Result.TexCoords[CoordinateIndex] = DebugInputs.TexCoord01.xy;
+	}
 #endif
 
 	Result.TwoSidedSign = 1;
diff --git a/Engine/Shaders/DebugViewModeVertexShader.usf b/Engine/Shaders/DebugViewModeVertexShader.usf
--- a/Engine/Shaders/DebugViewModeVertexShader.usf
+++ b/Engine/Shaders/DebugViewModeVertexShader.usf
@@ -5,6 +5,12 @@
 =============================================================================*/
 
 #include "Common.usf"
+
+// This is required to pass the 4 texcoords to the pixel shaders, no matter what the material actually use.
+// This fits with FDebugVSToPS which supports 4 tex coords.
+#define MIN_MATERIAL_TEXCOORDS 4
+
+
 #include "Material.usf"
 #include "VertexFactory.usf"
 
@@ -56,9 +62,18 @@
 		Output.TangentToWorld2 = TangentToWorld[2];
 
 #if NUM_MATERIAL_TEXCOORDS_VERTEX > 0
-		Output.TexCoord0 = MaterialParameters.TexCoords[0];
+		Output.TexCoord23 = Output.TexCoord01 = MaterialParameters.TexCoords[0].xyxy;
 #else
-		Output.TexCoord0 = 0;
+		Output.TexCoord23 = Output.TexCoord01 = 0;
+#endif
+#if NUM_MATERIAL_TEXCOORDS_VERTEX > 1
+		Output.TexCoord01.zw = MaterialParameters.TexCoords[1];
+#endif
+#if NUM_MATERIAL_TEXCOORDS_VERTEX > 2
+		Output.TexCoord23.xy = MaterialParameters.TexCoords[2];
+#endif
+#if NUM_MATERIAL_TEXCOORDS_VERTEX > 3
+		Output.TexCoord23.zw = MaterialParameters.TexCoords[3];
 #endif
 
 #if INSTANCED_STEREO
@@ -77,10 +92,13 @@ void Main(
 	FVertexFactoryInput Input,
 	OPTIONAL_VertexID
 	out FDebugVSOutput Output
+#if USE_GLOBAL_CLIP_PLANE && !USING_TESSELLATION
+	, out float OutGlobalClipPlaneDistance : SV_ClipDistance
+#endif
 #if INSTANCED_STEREO
 	, uint InstanceId : SV_InstanceID
-	, out float4 OutClipDistance : SV_ClipDistance
-	, out float4 OutCullDistance : SV_CullDistance
+	, out float OutClipDistance : SV_ClipDistance1
+	, out float OutCullDistance : SV_CullDistance1
 #endif
 	)
 {
@@ -138,6 +156,10 @@ void Main(
 		#endif
 	}
 
+#if USE_GLOBAL_CLIP_PLANE
+	OutGlobalClipPlaneDistance = dot(ResolvedView.GlobalClippingPlane, float4(WorldPos.xyz - ResolvedView.PreViewTranslation.xyz, 1));
+#endif
+
 	Output.VertexColor = VertexParameters.VertexColor;
 
 	float3x3 TangentToWorld = VertexParameters.TangentToWorld;
@@ -146,9 +168,18 @@ void Main(
 	Output.TangentToWorld2 = float3(TangentToWorld[2]);
 
 #if NUM_MATERIAL_TEXCOORDS_VERTEX > 0
-	Output.TexCoord0 = VertexParameters.TexCoords[0];
+		Output.TexCoord23 = Output.TexCoord01 = VertexParameters.TexCoords[0].xyxy;
 #else
-	Output.TexCoord0 = 0;
+		Output.TexCoord23 = Output.TexCoord01 = 0;
+#endif
+#if NUM_MATERIAL_TEXCOORDS_VERTEX > 1
+		Output.TexCoord01.zw = VertexParameters.TexCoords[1];
+#endif
+#if NUM_MATERIAL_TEXCOORDS_VERTEX > 2
+		Output.TexCoord23.xy = VertexParameters.TexCoords[2];
+#endif
+#if NUM_MATERIAL_TEXCOORDS_VERTEX > 3
+		Output.TexCoord23.zw = VertexParameters.TexCoords[3];
 #endif
 
 #if INSTANCED_STEREO
diff --git a/Engine/Shaders/DepthOnlyVertexShader.usf b/Engine/Shaders/DepthOnlyVertexShader.usf
--- a/Engine/Shaders/DepthOnlyVertexShader.usf
+++ b/Engine/Shaders/DepthOnlyVertexShader.usf
@@ -80,10 +80,13 @@ void Main(
 	FVertexFactoryInput Input,
 	OPTIONAL_VertexID
 	out FDepthOnlyVSOutput Output
+#if USE_GLOBAL_CLIP_PLANE && !USING_TESSELLATION
+	, out float OutGlobalClipPlaneDistance : SV_ClipDistance
+#endif
 #if INSTANCED_STEREO
 	, uint InstanceId : SV_InstanceID
-	, out float4 OutClipDistance : SV_ClipDistance
-	, out float4 OutCullDistance : SV_CullDistance
+	, out float OutClipDistance : SV_ClipDistance1
+	, out float OutCullDistance : SV_CullDistance1
 #endif
 	)
 {
@@ -136,6 +139,10 @@ void Main(
 		Output.Position.z += InstancedStereoDepthBias * Output.Position.w * (float)bNeedsInstancedStereoBias;
 		#endif
 	}
+
+#if USE_GLOBAL_CLIP_PLANE
+	OutGlobalClipPlaneDistance = dot(ResolvedView.GlobalClippingPlane, float4(WorldPos.xyz - ResolvedView.PreViewTranslation.xyz, 1));
+#endif
 #endif
 
 #if !MATERIALBLENDING_SOLID || USING_TESSELLATION || OUTPUT_PIXEL_DEPTH_OFFSET
diff --git a/Engine/Shaders/DistanceFieldAOShared.usf b/Engine/Shaders/DistanceFieldAOShared.usf
--- a/Engine/Shaders/DistanceFieldAOShared.usf
+++ b/Engine/Shaders/DistanceFieldAOShared.usf
@@ -212,18 +212,28 @@ uint2 ComputeTileCoordinateFromScreenGrid(uint2 OutputCoordinate)
 	return TileCoordinate;
 }
 
-float2 GetBaseLevelScreenUVFromScreenGrid(uint2 OutputCoordinate)
+float2 GetBaseLevelScreenUVFromScreenGrid(uint2 OutputCoordinate, float JitterScale)
 {
-	float2 BaseLevelScreenUV = (OutputCoordinate * TRACE_DOWNSAMPLE_FACTOR + JitterOffset + float2(.5f, .5f)) * BaseLevelTexelSize;
+	float2 BaseLevelScreenUV = (OutputCoordinate * TRACE_DOWNSAMPLE_FACTOR + JitterOffset * JitterScale + float2(.5f, .5f)) * BaseLevelTexelSize;
 	return BaseLevelScreenUV;
 }
 
-float2 GetScreenUVFromScreenGrid(uint2 OutputCoordinate)
+float2 GetBaseLevelScreenUVFromScreenGrid(uint2 OutputCoordinate)
+{
+	return GetBaseLevelScreenUVFromScreenGrid(OutputCoordinate, 1);
+}
+
+float2 GetScreenUVFromScreenGrid(uint2 OutputCoordinate, float JitterScale)
 {
-	float2 ScreenUV = ((OutputCoordinate * TRACE_DOWNSAMPLE_FACTOR + JitterOffset) * AO_DOWNSAMPLE_FACTOR + Frame.ViewRectMin.xy + float2(.5f, .5f)) * Frame.BufferSizeAndInvSize.zw;
+	float2 ScreenUV = ((OutputCoordinate * TRACE_DOWNSAMPLE_FACTOR + JitterOffset * JitterScale) * AO_DOWNSAMPLE_FACTOR + Frame.ViewRectMin.xy + float2(.5f, .5f)) * Frame.BufferSizeAndInvSize.zw;
 	return ScreenUV;
 }
 
+float2 GetScreenUVFromScreenGrid(uint2 OutputCoordinate)
+{
+	return GetScreenUVFromScreenGrid(OutputCoordinate, 1);
+}
+
 RWBuffer<uint> RWTileHeadDataUnpacked;
 RWBuffer<uint> RWTileArrayData;
 
diff --git a/Engine/Shaders/DistanceFieldLightingPost.usf b/Engine/Shaders/DistanceFieldLightingPost.usf
--- a/Engine/Shaders/DistanceFieldLightingPost.usf
+++ b/Engine/Shaders/DistanceFieldLightingPost.usf
@@ -508,6 +508,9 @@ void FilterHistoryPS(
 #define IRRADIANCE_TARGET SV_Target1
 #endif
 
+Texture2D SpecularOcclusionTexture;
+SamplerState SpecularOcclusionSampler;
+
 float MinIndirectDiffuseOcclusion;
 
 /** Upsamples the AO results to full resolution using a bilateral filter. */
@@ -525,6 +528,7 @@ void AOUpsamplePS(
 	// Distance field AO was computed at 0,0 regardless of viewrect min
 	float2 DistanceFieldUVs = UVAndScreenPos.xy - Frame.ViewRectMin.xy * Frame.BufferSizeAndInvSize.zw;
 	float3 Irradiance = 0;
+	float SpecularOcclusion = 1;
 
 #define BILATERAL_UPSAMPLE 1
 #if BILATERAL_UPSAMPLE
@@ -576,12 +580,30 @@ void AOUpsamplePS(
 			* InvWeight;
 	#endif
 
+	#if SUPPORT_SPECULAR_OCCLUSION
+		float SpecularOcclusionValues00 = Texture2DSampleLevel(SpecularOcclusionTexture, SpecularOcclusionSampler, Corner00UV, 0).x;
+		float SpecularOcclusionValues10 = Texture2DSampleLevel(SpecularOcclusionTexture, SpecularOcclusionSampler, Corner00UV + float2(LowResTexelSize.x, 0), 0).x;
+		float SpecularOcclusionValues01 = Texture2DSampleLevel(SpecularOcclusionTexture, SpecularOcclusionSampler, Corner00UV + float2(0, LowResTexelSize.y), 0).x;
+		float SpecularOcclusionValues11 = Texture2DSampleLevel(SpecularOcclusionTexture, SpecularOcclusionSampler, Corner00UV + LowResTexelSize, 0).x;
+
+		SpecularOcclusion = 
+			(FinalWeights.x * SpecularOcclusionValues00 
+			+ FinalWeights.y * SpecularOcclusionValues10
+			+ FinalWeights.z * SpecularOcclusionValues01 
+			+ FinalWeights.w * SpecularOcclusionValues11)
+			* InvWeight;
+	#endif
+
 #else
 	float3 BentNormal = Texture2DSampleLevel(BentNormalAOTexture, BentNormalAOSampler, DistanceFieldUVs, 0).xyz;
 
 	#if SUPPORT_IRRADIANCE
 		Irradiance = Texture2DSampleLevel(IrradianceTexture, IrradianceSampler, DistanceFieldUVs, 0).xyz;
 	#endif
+
+	#if SUPPORT_SPECULAR_OCCLUSION
+		SpecularOcclusion = Texture2DSampleLevel(SpecularOcclusionTexture, SpecularOcclusionSampler, DistanceFieldUVs, 0).x;
+	#endif
 #endif
 
 #if MODULATE_SCENE_COLOR
@@ -594,7 +616,7 @@ void AOUpsamplePS(
 	//OutBentNormal = float4(BentNormal, 1);
 #else
 	#if OUTPUT_BENT_NORMAL
-		OutBentNormal = float4(BentNormal, 1);
+		OutBentNormal = float4(BentNormal, SpecularOcclusion);
 		#if SUPPORT_IRRADIANCE
 			OutIrradiance = float4(Irradiance, 1);
 		#endif
diff --git a/Engine/Shaders/DistanceFieldSpecularOcclusion.usf b/Engine/Shaders/DistanceFieldSpecularOcclusion.usf
new file mode 100644
--- /dev/null
+++ b/Engine/Shaders/DistanceFieldSpecularOcclusion.usf
@@ -0,0 +1,216 @@
+// Copyright 1998-2016 Epic Games, Inc. All Rights Reserved.
+
+/*=============================================================================
+	DistanceFieldSpecularOcclusion.usf
+=============================================================================*/
+
+#include "Common.usf"
+#include "DeferredShadingCommon.usf"
+#include "DistanceFieldLightingShared.usf"
+#include "DistanceFieldAOShared.usf"
+
+#define NUM_THREADS_PER_RECORD 4
+#define CONE_TRACE_DISPATCH_SIZEX (CULLED_TILE_SIZEX / TRACE_DOWNSAMPLE_FACTOR)
+#define CONE_TRACE_DISPATCH_SIZEY (CONE_TRACE_DISPATCH_SIZEX * NUM_THREADS_PER_RECORD)
+
+#define SHARED_VISIBILITY_DIM (CONE_TRACE_DISPATCH_SIZEX * CONE_TRACE_DISPATCH_SIZEX)
+groupshared uint SharedConeVisibility[SHARED_VISIBILITY_DIM];
+
+Buffer<float4> TileConeDepthRanges;
+RWTexture2D<float> RWSpecularOcclusion;
+
+#define CONE_TRACE_AGAINST_OBJECTS 1
+
+void ConeTraceAgainstTileCulledObjects(uint2 OutputCoordinate, uint SharedVisibilityBaseIndex, uint ThreadIndex, float3 WorldShadingPosition, float SceneDepth, float3 WorldNormal, float3 TangentX, float3 TangentY)
+{
+	uint2 TileCoordinate = ComputeTileCoordinateFromScreenGrid(OutputCoordinate);
+	uint4 TileHead = GetTileHead(TileCoordinate);
+
+#if USE_DEPTH_RANGE_LISTS
+		
+	uint TileIndex = TileCoordinate.y * TileListGroupSize.x + TileCoordinate.x;
+	float4 ConeAxisDepthRanges = TileConeDepthRanges.Load(TileIndex);
+	uint ListIndex = SceneDepth < ConeAxisDepthRanges.y ? 0 : 1;
+	uint NumObjectsAffectingTile = SceneDepth < ConeAxisDepthRanges.y ? TileHead.y : TileHead.z;
+#else
+	uint NumObjectsAffectingTile = TileHead.y + TileHead.z + TileHead.w;
+#endif
+
+	uint NumCulledObjects = NumObjectsAffectingTile;
+	float MaxWorldStepOffset = GetStepOffset(NUM_CONE_STEPS);
+	float InvMaxOcclusionDistance = 1.0f / AOObjectMaxDistance;
+
+	float3 CameraToPixel = normalize(WorldShadingPosition - View.WorldCameraOrigin);
+	float3 ReflectionVector = reflect(CameraToPixel, WorldNormal);
+	float TanConeHalfAngle = tan(Frame.GeneralPurposeTweak);
+
+	LOOP 
+	for (uint ListObjectIndex = ThreadIndex; ListObjectIndex < NumCulledObjects; ListObjectIndex += NUM_THREADS_PER_RECORD)
+	{
+		#if USE_DEPTH_RANGE_LISTS
+			uint ArrayIndex = ListObjectIndex;
+		#else
+			uint ListIndex = 0;
+			uint ArrayIndex = ListObjectIndex;
+
+			FLATTEN
+			if (ListObjectIndex >= TileHead.y + TileHead.z)
+			{
+				ListIndex = 2;
+				ArrayIndex = ListObjectIndex - TileHead.y - TileHead.z;
+			}
+			else if (ListObjectIndex >= TileHead.y)
+			{
+				ListIndex = 1;
+				ArrayIndex = ListObjectIndex - TileHead.y;
+			}
+		#endif
+
+		uint ObjectIndex = TileArrayData.Load((ArrayIndex * TileListGroupSize.x * TileListGroupSize.y + TileHead.x) * NUM_CULLED_OBJECT_LISTS + ListIndex);
+
+		float4 ObjectPositionAndRadius = LoadObjectPositionAndRadius(ObjectIndex);
+		float ObjectDistanceSq = dot(ObjectPositionAndRadius.xyz - WorldShadingPosition, ObjectPositionAndRadius.xyz - WorldShadingPosition);
+
+		BRANCH
+		// Skip tracing objects with a small projected angle 
+		if (ObjectPositionAndRadius.w * ObjectPositionAndRadius.w / ObjectDistanceSq > Square(.25f))
+		{
+			float3 LocalPositionExtent = LoadObjectLocalPositionExtent(ObjectIndex);
+			float4x4 WorldToVolume = LoadObjectWorldToVolume(ObjectIndex);
+			bool bGeneratedAsTwoSided;
+			float4 UVScaleAndVolumeScale = LoadObjectUVScale(ObjectIndex, bGeneratedAsTwoSided);
+			float3 VolumeShadingPosition = mul(float4(WorldShadingPosition, 1), WorldToVolume).xyz;
+			float BoxDistance = ComputeDistanceFromBoxToPoint(-LocalPositionExtent, LocalPositionExtent, VolumeShadingPosition) * UVScaleAndVolumeScale.w;
+
+			BRANCH
+			if (BoxDistance < AOObjectMaxDistance)
+			{
+				float3 UVAdd = LoadObjectUVAdd(ObjectIndex);
+
+				uint StartStepIndex = 0;
+			
+				#if !USE_DEPTH_RANGE_LISTS
+					FLATTEN
+					if (ListObjectIndex >= TileHead.y + TileHead.z)
+					{
+						StartStepIndex = 8;
+					}
+					else if (ListObjectIndex >= TileHead.y)
+					{
+						StartStepIndex = 5;
+					}
+				#endif
+
+				float ObjectOccluderRadius = length(LocalPositionExtent) * .5f * UVScaleAndVolumeScale.w;
+
+				//float3 ConeDirection = AOSamples2.SampleDirections[0].xyz;
+				//float3 RotatedConeDirection = ConeDirection.x * TangentX + ConeDirection.y * TangentY + ConeDirection.z * WorldNormal;
+				float3 RotatedConeDirection = ReflectionVector;
+
+				float MinVisibility = 1;
+				float WorldStepOffset = GetStepOffset(StartStepIndex);
+
+				LOOP
+				for (uint StepIndex = StartStepIndex; StepIndex < NUM_CONE_STEPS && WorldStepOffset < MaxWorldStepOffset; StepIndex++)
+				{
+					float3 WorldSamplePosition = WorldShadingPosition + RotatedConeDirection * WorldStepOffset;
+					float3 StepSamplePosition = mul(float4(WorldSamplePosition, 1), WorldToVolume).xyz;
+					float3 ClampedSamplePosition = clamp(StepSamplePosition, -LocalPositionExtent, LocalPositionExtent);
+					float DistanceToClamped = length(StepSamplePosition - ClampedSamplePosition);
+
+					float3 StepVolumeUV = DistanceFieldVolumePositionToUV(ClampedSamplePosition, UVScaleAndVolumeScale.xyz, UVAdd);
+					float DistanceToOccluder = (Texture3DSampleLevel(DistanceFieldTexture, DistanceFieldSampler, StepVolumeUV, 0).x + DistanceToClamped) * UVScaleAndVolumeScale.w;
+					float SphereRadius = WorldStepOffset * TanConeHalfAngle;
+					//@todo - have to bias away from surface further for this to work
+					float ShadingSphereRadius = SphereRadius * 1.0f;
+
+					// Derive visibility from 1d intersection
+					float Visibility = saturate(DistanceToOccluder / ShadingSphereRadius);
+
+					// Don't allow small objects to fully occlude a cone step
+					Visibility = max(Visibility, 1 - saturate(ObjectOccluderRadius / SphereRadius));
+
+					float OccluderDistanceFraction = (WorldStepOffset + DistanceToOccluder) * InvMaxOcclusionDistance;
+
+					// Fade out occlusion based on distance to occluder to avoid a discontinuity at the max AO distance
+					Visibility = max(Visibility, saturate(OccluderDistanceFraction * 1));
+
+					MinVisibility = min(MinVisibility, Visibility);
+
+					float MinStepScale = .3f;
+
+					float MinStepSize = MinStepScale * (GetStepOffset(StepIndex + 1) - GetStepOffset(StepIndex));
+					WorldStepOffset += max(DistanceToOccluder, MinStepSize);
+				}
+
+				InterlockedMin(SharedConeVisibility[SharedVisibilityBaseIndex], asuint(MinVisibility));
+			}
+		}
+	}
+}
+
+/** Traces cones of a hemisphere against intersecting object distance fields. */
+[numthreads(CONE_TRACE_DISPATCH_SIZEX, CONE_TRACE_DISPATCH_SIZEY, 1)]
+void ConeTraceSpecularOcclusionCS(
+	uint3 GroupId : SV_GroupID,
+	uint3 DispatchThreadId : SV_DispatchThreadID,
+    uint3 GroupThreadId : SV_GroupThreadID) 
+{
+	uint2 OutputCoordinate = uint2(DispatchThreadId.x, DispatchThreadId.y / NUM_THREADS_PER_RECORD);
+	uint2 RecordTileCoordinate = uint2(GroupThreadId.x, GroupThreadId.y / NUM_THREADS_PER_RECORD);
+	uint SharedVisibilityBaseIndex = (RecordTileCoordinate.y * CONE_TRACE_DISPATCH_SIZEX + RecordTileCoordinate.x) * 1;
+	uint ThreadIndex = GroupThreadId.y - RecordTileCoordinate.y * NUM_THREADS_PER_RECORD;
+	
+	uint FloatOneEncoded = asuint(1.0f);
+	
+	for (uint Index = GroupThreadId.y * CONE_TRACE_DISPATCH_SIZEX + GroupThreadId.x; Index < SHARED_VISIBILITY_DIM; Index += CONE_TRACE_DISPATCH_SIZEX * CONE_TRACE_DISPATCH_SIZEY)
+	{
+		SharedConeVisibility[Index] = FloatOneEncoded;
+	}
+
+	GroupMemoryBarrierWithGroupSync();
+
+	float2 BaseLevelScreenUV = GetBaseLevelScreenUVFromScreenGrid(OutputCoordinate, 0);
+
+	float3 WorldNormal;
+	float SceneDepth;
+	bool bHasDistanceFieldRepresentation;
+	bool bHasHeightfieldRepresentation;
+	GetDownsampledGBuffer(BaseLevelScreenUV, WorldNormal, SceneDepth, bHasDistanceFieldRepresentation, bHasHeightfieldRepresentation);
+
+	float3 TangentX;
+	float3 TangentY;
+	FindBestAxisVectors2(WorldNormal, TangentX, TangentY);
+	 
+	{
+		float2 ScreenUV = GetScreenUVFromScreenGrid(OutputCoordinate, 0);
+		float2 ScreenPosition = (ScreenUV.xy - View.ScreenPositionScaleBias.wz) / View.ScreenPositionScaleBias.xy;
+		
+		float4 HomogeneousWorldPosition = mul(float4(ScreenPosition * SceneDepth, SceneDepth, 1), View.ScreenToWorld);
+		float3 OpaqueWorldPosition = HomogeneousWorldPosition.xyz / HomogeneousWorldPosition.w;
+		float3 WorldShadingPosition = OpaqueWorldPosition;
+
+		#if CONE_TRACE_AGAINST_OBJECTS
+
+			ConeTraceAgainstTileCulledObjects(OutputCoordinate, SharedVisibilityBaseIndex, ThreadIndex, WorldShadingPosition, SceneDepth, WorldNormal, TangentX, TangentY);
+		
+		#endif
+	}
+
+	GroupMemoryBarrierWithGroupSync();
+
+	float ConeVisibility = asfloat(SharedConeVisibility[SharedVisibilityBaseIndex + 0]);
+	RWSpecularOcclusion[OutputCoordinate] = ConeVisibility;
+
+	/*
+	if (all(OutputCoordinate < ScreenGridConeVisibilitySize))
+	{
+		uint OutputBaseIndex = (OutputCoordinate.y * ScreenGridConeVisibilitySize.x + OutputCoordinate.x) * 1;
+
+		for (uint ConeIndex = ThreadIndex; ConeIndex < 1; ConeIndex += NUM_THREADS_PER_RECORD)
+		{
+			float ConeVisibility = asfloat(SharedConeVisibility[SharedVisibilityBaseIndex + ConeIndex]);
+			RWSpecularOcclusion[OutputBaseIndex + ConeIndex] = Frame.GeneralPurposeTweak;
+		}
+	}*/
+}
diff --git a/Engine/Shaders/DistanceFieldSurfaceCacheLighting.usf b/Engine/Shaders/DistanceFieldSurfaceCacheLighting.usf
--- a/Engine/Shaders/DistanceFieldSurfaceCacheLighting.usf
+++ b/Engine/Shaders/DistanceFieldSurfaceCacheLighting.usf
@@ -441,7 +441,8 @@ void RayTraceThroughGlobalDistanceField(
 	OutMaxRayTime = 1;
 
 	float3 GlobalVolumeCenter = GlobalVolumeCenterAndExtent[ClipmapIndex].xyz;
-	float GlobalVolumeExtent = GlobalVolumeCenterAndExtent[ClipmapIndex].w;
+	// Subtract one texel from the extent to avoid filtering from invalid texels
+	float GlobalVolumeExtent = GlobalVolumeCenterAndExtent[ClipmapIndex].w - GlobalVolumeCenterAndExtent[ClipmapIndex].w * GlobalVolumeTexelSize;
 	float3 VolumeRayStart = WorldRayStart - GlobalVolumeCenter;
 	float3 VolumeRayEnd = WorldRayEnd - GlobalVolumeCenter;
 	float3 VolumeRayDirection = VolumeRayEnd - VolumeRayStart;
diff --git a/Engine/Shaders/DistortAccumulateVS.usf b/Engine/Shaders/DistortAccumulateVS.usf
--- a/Engine/Shaders/DistortAccumulateVS.usf
+++ b/Engine/Shaders/DistortAccumulateVS.usf
@@ -68,6 +68,9 @@ void Main(
 	FVertexFactoryInput Input,
 	OPTIONAL_VertexID
 	out VS_OUTPUT_TYPE Output
+#if USE_GLOBAL_CLIP_PLANE && !USING_TESSELLATION
+	, out float OutGlobalClipPlaneDistance : SV_ClipDistance
+#endif
 	)
 {
 	ResolvedView = ResolveView();
@@ -94,6 +97,10 @@ void Main(
 	Output.Position = mul(WorldPosition, View.TranslatedWorldToClip);
 	Output.FactoryInterpolants = VertexFactoryGetInterpolantsVSToPS(Input, VFIntermediates, VertexParameters);
 
+#if USE_GLOBAL_CLIP_PLANE
+	OutGlobalClipPlaneDistance = dot(ResolvedView.GlobalClippingPlane, float4(WorldPosition.xyz - ResolvedView.PreViewTranslation.xyz, 1));
+#endif
+
 	Output.PixelPosition = WorldPosition;
 
 #endif // USING_TESSELLATION
diff --git a/Engine/Shaders/FlatTessellation.usf b/Engine/Shaders/FlatTessellation.usf
--- a/Engine/Shaders/FlatTessellation.usf
+++ b/Engine/Shaders/FlatTessellation.usf
@@ -131,9 +131,12 @@ struct FFlatTessellationHSToDS
 		FHullShaderConstantOutput HSConstantData,
 		const OutputPatch<FFlatTessellationHSToDS, 3> I,
 		float3 BarycentricCoords : SV_DomainLocation 
+#if USE_GLOBAL_CLIP_PLANE
+		, out float OutGlobalClipPlaneDistance : SV_ClipDistance
+#endif
 #if INSTANCED_STEREO
-		, out float4 OutClipDistance : SV_ClipDistance
-		, out float4 OutCullDistance : SV_CullDistance
+		, out float OutClipDistance : SV_ClipDistance1
+		, out float OutCullDistance : SV_CullDistance1
 #endif
 		)
 	{
@@ -199,6 +202,10 @@ struct FFlatTessellationHSToDS
 
 		WorldPosition.xyz += WorldDisplacement;
 
+		#if USE_GLOBAL_CLIP_PLANE
+			OutGlobalClipPlaneDistance = dot(ResolvedView.GlobalClippingPlane, float4(WorldPosition.xyz - ResolvedView.PreViewTranslation.xyz, 1));
+		#endif
+
 		FPassSpecificVSToPS Result = PassFinalizeTessellationOutput(Interp, WorldPosition,MaterialParameters);
 
 		#if INSTANCED_STEREO
diff --git a/Engine/Shaders/GlobalDistanceField.usf b/Engine/Shaders/GlobalDistanceField.usf
--- a/Engine/Shaders/GlobalDistanceField.usf
+++ b/Engine/Shaders/GlobalDistanceField.usf
@@ -9,6 +9,7 @@
 #include "DistanceFieldLightingShared.usf"
 #include "DistanceFieldAOShared.usf"
 #include "GlobalDistanceFieldShared.usf"
+#include "HeightfieldLightingShared.usf"
 
 RWBuffer<float4> RWCulledObjectBounds;
 RWBuffer<float4> RWCulledObjectData;
@@ -296,4 +297,57 @@ void CompositeObjectDistanceFieldsCS(
 		uint3 WriteCoordinate = (uint3)(frac(GlobalUV) * GlobalVolumeDimension);
 		RWGlobalDistanceFieldTexture[WriteCoordinate] = MinDistance;
 	}
+}
+
+#ifndef COMPOSITE_HEIGHTFIELDS_THREADGROUP_SIZE
+#define COMPOSITE_HEIGHTFIELDS_THREADGROUP_SIZE 1
+#endif
+
+[numthreads(COMPOSITE_HEIGHTFIELDS_THREADGROUP_SIZE, COMPOSITE_HEIGHTFIELDS_THREADGROUP_SIZE, 1)]
+void CompositeHeightfieldsIntoGlobalDistanceFieldCS(
+	uint3 GroupId : SV_GroupID,
+	uint3 DispatchThreadId : SV_DispatchThreadID,
+	uint3 GroupThreadId : SV_GroupThreadID)
+{
+	float GlobalVolumeExtent = GlobalVolumeCenterAndExtent[ClipmapIndex].w;
+	// We only dispatch one depth layer, so the z value is incorrect
+	float3 LookupWorldPosition = UpdateRegionVolumeMin + (DispatchThreadId.xyz + .5f) * VolumeTexelSize * GlobalVolumeExtent * 2;
+
+	// Search through all heightfields to find the one that encompases this UV
+	LOOP
+	for (uint HeightfieldIndex = 0; HeightfieldIndex < NumHeightfields; HeightfieldIndex++)
+	{
+		float3 LocalPosition = mul(float4(LookupWorldPosition, 1), GetWorldToLocal(HeightfieldIndex)).xyz;
+		float4 MinMaxHeightfieldUV;
+		float2 HeightfieldUV = GetHeightfieldUV(HeightfieldIndex, LocalPosition.xy, MinMaxHeightfieldUV);
+
+		BRANCH
+		if (all(HeightfieldUV > MinMaxHeightfieldUV.xy) && all(HeightfieldUV < MinMaxHeightfieldUV.zw))
+		{
+			float3 WorldHeightfieldNormal;
+			float3 WorldHeightfieldShadingPosition = GetHeightfieldWorldPositionAndNormal(HeightfieldIndex, LocalPosition.xy, HeightfieldUV, WorldHeightfieldNormal);
+
+			if (all(DispatchThreadId.xy < UpdateRegionSize.xy))
+			{
+				// Compute distance for all Z values of the update region
+				LOOP
+				for (uint ZIndex = 0; ZIndex < UpdateRegionSize.z; ZIndex++)
+				{
+					float3 WorldPosition;
+					WorldPosition.xy = LookupWorldPosition.xy;
+					WorldPosition.z = UpdateRegionVolumeMin.z + (ZIndex + .5f) * VolumeTexelSize.z * GlobalVolumeExtent * 2;
+
+					// Project the vertical height vector onto the normal of the heightfield directly below the point we are computing the distance field for, use the perpendicular distance
+					float DistanceToHeightfieldPlane = dot(WorldHeightfieldNormal, WorldPosition - WorldHeightfieldShadingPosition);
+					float MinDistance = clamp(DistanceToHeightfieldPlane, -AOGlobalMaxSphereQueryRadius, AOGlobalMaxSphereQueryRadius);
+
+					float3 GlobalUV = ComputeGlobalUV(WorldPosition, ClipmapIndex);
+					uint3 WriteCoordinate = (uint3)(frac(GlobalUV) * GlobalVolumeDimension);
+					RWGlobalDistanceFieldTexture[WriteCoordinate] = min(MinDistance, RWGlobalDistanceFieldTexture[WriteCoordinate]);
+				}
+			}
+
+			break;
+		}
+	}
 }
\ No newline at end of file
diff --git a/Engine/Shaders/GlobalDistanceFieldShared.usf b/Engine/Shaders/GlobalDistanceFieldShared.usf
--- a/Engine/Shaders/GlobalDistanceFieldShared.usf
+++ b/Engine/Shaders/GlobalDistanceFieldShared.usf
@@ -6,7 +6,11 @@
 
 #if FEATURE_LEVEL >= FEATURE_LEVEL_SM5
 
+#if IS_MATERIAL_SHADER
+#define MaxGlobalDistance				Frame.MaxGlobalDistance_UB
+#else
 float MaxGlobalDistance;
+#endif
 
 //@todo - get the clipmap branches to compile under gl
 #if SM5_PROFILE || PS4_PROFILE || METAL_SM5_PROFILE
@@ -30,37 +34,49 @@ float MaxGlobalDistance;
 	#define GlobalVolumeTexelSize			Frame.GlobalVolumeTexelSize_UB
 #else
 	// these are only used for the precomputation shaders; which don't have a frame UB
-Texture3D GlobalDistanceFieldTexture0;
-Texture3D GlobalDistanceFieldTexture1;
-Texture3D GlobalDistanceFieldTexture2;
-Texture3D GlobalDistanceFieldTexture3;
-SamplerState GlobalDistanceFieldSampler0;
-SamplerState GlobalDistanceFieldSampler1;
-SamplerState GlobalDistanceFieldSampler2;
-SamplerState GlobalDistanceFieldSampler3;
-float4 GlobalVolumeCenterAndExtent[MAX_GLOBAL_DF_CLIPMAPS];
-float4 GlobalVolumeWorldToUVAddAndMul[MAX_GLOBAL_DF_CLIPMAPS];
-float GlobalVolumeDimension;
-float GlobalVolumeTexelSize;
+	Texture3D GlobalDistanceFieldTexture0;
+	Texture3D GlobalDistanceFieldTexture1;
+	Texture3D GlobalDistanceFieldTexture2;
+	Texture3D GlobalDistanceFieldTexture3;
+	SamplerState GlobalDistanceFieldSampler0;
+	SamplerState GlobalDistanceFieldSampler1;
+	SamplerState GlobalDistanceFieldSampler2;
+	SamplerState GlobalDistanceFieldSampler3;
+	float4 GlobalVolumeCenterAndExtent[MAX_GLOBAL_DF_CLIPMAPS];
+	float4 GlobalVolumeWorldToUVAddAndMul[MAX_GLOBAL_DF_CLIPMAPS];
+	float GlobalVolumeDimension;
+	float GlobalVolumeTexelSize;
+#endif
+
+#if SUPPORTS_INDEPENDENT_SAMPLERS
+	#define SharedGlobalDistanceFieldSampler0 GlobalDistanceFieldSampler0
+	#define SharedGlobalDistanceFieldSampler1 GlobalDistanceFieldSampler0
+	#define SharedGlobalDistanceFieldSampler2 GlobalDistanceFieldSampler0
+	#define SharedGlobalDistanceFieldSampler3 GlobalDistanceFieldSampler0
+#else
+	#define SharedGlobalDistanceFieldSampler0 GlobalDistanceFieldSampler0
+	#define SharedGlobalDistanceFieldSampler1 GlobalDistanceFieldSampler1
+	#define SharedGlobalDistanceFieldSampler2 GlobalDistanceFieldSampler2
+	#define SharedGlobalDistanceFieldSampler3 GlobalDistanceFieldSampler3
 #endif
 
 float4 SampleGlobalDistanceField(int ClipmapIndex, float3 UV)
 {
 	if (ClipmapIndex == 0)
 	{
-		return Texture3DSampleLevel(GlobalDistanceFieldTexture0, GlobalDistanceFieldSampler0, UV, 0);
+		return Texture3DSampleLevel(GlobalDistanceFieldTexture0, SharedGlobalDistanceFieldSampler0, UV, 0);
 	}
 	else if (ClipmapIndex == 1)
 	{
-		return Texture3DSampleLevel(GlobalDistanceFieldTexture1, GlobalDistanceFieldSampler1, UV, 0);
+		return Texture3DSampleLevel(GlobalDistanceFieldTexture1, SharedGlobalDistanceFieldSampler1, UV, 0);
 	}
 	else if (ClipmapIndex == 2)
 	{
-		return Texture3DSampleLevel(GlobalDistanceFieldTexture2, GlobalDistanceFieldSampler2, UV, 0);
+		return Texture3DSampleLevel(GlobalDistanceFieldTexture2, SharedGlobalDistanceFieldSampler2, UV, 0);
 	}
 	else
 	{
-		return Texture3DSampleLevel(GlobalDistanceFieldTexture3, GlobalDistanceFieldSampler3, UV, 0);
+		return Texture3DSampleLevel(GlobalDistanceFieldTexture3, SharedGlobalDistanceFieldSampler3, UV, 0);
 	}
 }
 
@@ -71,25 +87,26 @@ float3 ComputeGlobalUV(float3 WorldPosition, uint ClipmapIndex)
 	return WorldPosition * WorldToUVAddAndMul.www + WorldToUVAddAndMul.xyz;
 }
 
-float GetDistanceToNearestSurfaceGlobalClipmap(float3 WorldPosition, uint ClipmapIndex)
+float GetDistanceToNearestSurfaceGlobalClipmap(float3 WorldPosition, uint ClipmapIndex, float OuterClipmapFade)
 {
 	float3 GlobalUV = ComputeGlobalUV(WorldPosition, ClipmapIndex);
 	float DistanceToSurface = 0;
 	if (ClipmapIndex == 0)
 	{
-		DistanceToSurface = Texture3DSampleLevel(GlobalDistanceFieldTexture0, GlobalDistanceFieldSampler0, GlobalUV, 0).x;
+		DistanceToSurface = Texture3DSampleLevel(GlobalDistanceFieldTexture0, SharedGlobalDistanceFieldSampler0, GlobalUV, 0).x;
 	}
 	else if (ClipmapIndex == 1)
 	{
-		DistanceToSurface = Texture3DSampleLevel(GlobalDistanceFieldTexture1, GlobalDistanceFieldSampler1, GlobalUV, 0).x;
+		DistanceToSurface = Texture3DSampleLevel(GlobalDistanceFieldTexture1, SharedGlobalDistanceFieldSampler1, GlobalUV, 0).x;
 	}
 	else if (ClipmapIndex == 2)
 	{
-		DistanceToSurface = Texture3DSampleLevel(GlobalDistanceFieldTexture2, GlobalDistanceFieldSampler2, GlobalUV, 0).x;
+		DistanceToSurface = Texture3DSampleLevel(GlobalDistanceFieldTexture2, SharedGlobalDistanceFieldSampler2, GlobalUV, 0).x;
 	}
 	else if (ClipmapIndex == 3)
 	{
-		DistanceToSurface = Texture3DSampleLevel(GlobalDistanceFieldTexture3, GlobalDistanceFieldSampler3, GlobalUV, 0).x;
+		DistanceToSurface = Texture3DSampleLevel(GlobalDistanceFieldTexture3, SharedGlobalDistanceFieldSampler3, GlobalUV, 0).x;
+		DistanceToSurface = lerp(MaxGlobalDistance, DistanceToSurface, OuterClipmapFade);
 	}
 	return DistanceToSurface;
 }
@@ -103,7 +120,7 @@ float GetDistanceToNearestSurfaceGlobal(float3 WorldPosition)
 	BRANCH
 	if (DistanceFromClipmap > GlobalVolumeCenterAndExtent[0].w * GlobalVolumeTexelSize)
 	{
-		DistanceToSurface = GetDistanceToNearestSurfaceGlobalClipmap(WorldPosition, 0);
+		DistanceToSurface = GetDistanceToNearestSurfaceGlobalClipmap(WorldPosition, 0, 0);
 	}
 	else
 	{
@@ -112,7 +129,7 @@ float GetDistanceToNearestSurfaceGlobal(float3 WorldPosition)
 		BRANCH
 		if (DistanceFromClipmap > GlobalVolumeCenterAndExtent[1].w * GlobalVolumeTexelSize)
 		{
-			DistanceToSurface = GetDistanceToNearestSurfaceGlobalClipmap(WorldPosition, 1);
+			DistanceToSurface = GetDistanceToNearestSurfaceGlobalClipmap(WorldPosition, 1, 0);
 		}
 		else
 		{
@@ -122,11 +139,13 @@ float GetDistanceToNearestSurfaceGlobal(float3 WorldPosition)
 			BRANCH
 			if (DistanceFromClipmap > GlobalVolumeCenterAndExtent[2].w * GlobalVolumeTexelSize)
 			{
-				DistanceToSurface = GetDistanceToNearestSurfaceGlobalClipmap(WorldPosition, 2);
+				DistanceToSurface = GetDistanceToNearestSurfaceGlobalClipmap(WorldPosition, 2, 0);
 			}
 			else if (DistanceFromLastClipmap > GlobalVolumeCenterAndExtent[3].w * GlobalVolumeTexelSize)
 			{
-				DistanceToSurface = GetDistanceToNearestSurfaceGlobalClipmap(WorldPosition, 3);
+				// Fade to max distance in the last 10% of the largest clipmap
+				float OuterClipmapFade = saturate(DistanceFromLastClipmap * 10 * GlobalVolumeWorldToUVAddAndMul[3].w);
+				DistanceToSurface = GetDistanceToNearestSurfaceGlobalClipmap(WorldPosition, 3, OuterClipmapFade);
 			}
 		}
 	}
@@ -147,39 +166,39 @@ float3 GetDistanceFieldGradientGlobalClipmap(float3 WorldPosition, uint ClipmapI
 
 	if (ClipmapIndex == 0)
 	{
-		R = Texture3DSampleLevel(GlobalDistanceFieldTexture0, GlobalDistanceFieldSampler0, float3(GlobalUV.x + GlobalVolumeTexelSize, GlobalUV.y, GlobalUV.z), 0).x;
-		L = Texture3DSampleLevel(GlobalDistanceFieldTexture0, GlobalDistanceFieldSampler0, float3(GlobalUV.x - GlobalVolumeTexelSize, GlobalUV.y, GlobalUV.z), 0).x;
-		F = Texture3DSampleLevel(GlobalDistanceFieldTexture0, GlobalDistanceFieldSampler0, float3(GlobalUV.x, GlobalUV.y + GlobalVolumeTexelSize, GlobalUV.z), 0).x;
-		B = Texture3DSampleLevel(GlobalDistanceFieldTexture0, GlobalDistanceFieldSampler0, float3(GlobalUV.x, GlobalUV.y - GlobalVolumeTexelSize, GlobalUV.z), 0).x;
-		U = Texture3DSampleLevel(GlobalDistanceFieldTexture0, GlobalDistanceFieldSampler0, float3(GlobalUV.x, GlobalUV.y, GlobalUV.z + GlobalVolumeTexelSize), 0).x;
-		D = Texture3DSampleLevel(GlobalDistanceFieldTexture0, GlobalDistanceFieldSampler0, float3(GlobalUV.x, GlobalUV.y, GlobalUV.z - GlobalVolumeTexelSize), 0).x;
+		R = Texture3DSampleLevel(GlobalDistanceFieldTexture0, SharedGlobalDistanceFieldSampler0, float3(GlobalUV.x + GlobalVolumeTexelSize, GlobalUV.y, GlobalUV.z), 0).x;
+		L = Texture3DSampleLevel(GlobalDistanceFieldTexture0, SharedGlobalDistanceFieldSampler0, float3(GlobalUV.x - GlobalVolumeTexelSize, GlobalUV.y, GlobalUV.z), 0).x;
+		F = Texture3DSampleLevel(GlobalDistanceFieldTexture0, SharedGlobalDistanceFieldSampler0, float3(GlobalUV.x, GlobalUV.y + GlobalVolumeTexelSize, GlobalUV.z), 0).x;
+		B = Texture3DSampleLevel(GlobalDistanceFieldTexture0, SharedGlobalDistanceFieldSampler0, float3(GlobalUV.x, GlobalUV.y - GlobalVolumeTexelSize, GlobalUV.z), 0).x;
+		U = Texture3DSampleLevel(GlobalDistanceFieldTexture0, SharedGlobalDistanceFieldSampler0, float3(GlobalUV.x, GlobalUV.y, GlobalUV.z + GlobalVolumeTexelSize), 0).x;
+		D = Texture3DSampleLevel(GlobalDistanceFieldTexture0, SharedGlobalDistanceFieldSampler0, float3(GlobalUV.x, GlobalUV.y, GlobalUV.z - GlobalVolumeTexelSize), 0).x;
 	}
 	else if (ClipmapIndex == 1)
 	{
-		R = Texture3DSampleLevel(GlobalDistanceFieldTexture1, GlobalDistanceFieldSampler1, float3(GlobalUV.x + GlobalVolumeTexelSize, GlobalUV.y, GlobalUV.z), 0).x;
-		L = Texture3DSampleLevel(GlobalDistanceFieldTexture1, GlobalDistanceFieldSampler1, float3(GlobalUV.x - GlobalVolumeTexelSize, GlobalUV.y, GlobalUV.z), 0).x;
-		F = Texture3DSampleLevel(GlobalDistanceFieldTexture1, GlobalDistanceFieldSampler1, float3(GlobalUV.x, GlobalUV.y + GlobalVolumeTexelSize, GlobalUV.z), 0).x;
-		B = Texture3DSampleLevel(GlobalDistanceFieldTexture1, GlobalDistanceFieldSampler1, float3(GlobalUV.x, GlobalUV.y - GlobalVolumeTexelSize, GlobalUV.z), 0).x;
-		U = Texture3DSampleLevel(GlobalDistanceFieldTexture1, GlobalDistanceFieldSampler1, float3(GlobalUV.x, GlobalUV.y, GlobalUV.z + GlobalVolumeTexelSize), 0).x;
-		D = Texture3DSampleLevel(GlobalDistanceFieldTexture1, GlobalDistanceFieldSampler1, float3(GlobalUV.x, GlobalUV.y, GlobalUV.z - GlobalVolumeTexelSize), 0).x;
+		R = Texture3DSampleLevel(GlobalDistanceFieldTexture1, SharedGlobalDistanceFieldSampler1, float3(GlobalUV.x + GlobalVolumeTexelSize, GlobalUV.y, GlobalUV.z), 0).x;
+		L = Texture3DSampleLevel(GlobalDistanceFieldTexture1, SharedGlobalDistanceFieldSampler1, float3(GlobalUV.x - GlobalVolumeTexelSize, GlobalUV.y, GlobalUV.z), 0).x;
+		F = Texture3DSampleLevel(GlobalDistanceFieldTexture1, SharedGlobalDistanceFieldSampler1, float3(GlobalUV.x, GlobalUV.y + GlobalVolumeTexelSize, GlobalUV.z), 0).x;
+		B = Texture3DSampleLevel(GlobalDistanceFieldTexture1, SharedGlobalDistanceFieldSampler1, float3(GlobalUV.x, GlobalUV.y - GlobalVolumeTexelSize, GlobalUV.z), 0).x;
+		U = Texture3DSampleLevel(GlobalDistanceFieldTexture1, SharedGlobalDistanceFieldSampler1, float3(GlobalUV.x, GlobalUV.y, GlobalUV.z + GlobalVolumeTexelSize), 0).x;
+		D = Texture3DSampleLevel(GlobalDistanceFieldTexture1, SharedGlobalDistanceFieldSampler1, float3(GlobalUV.x, GlobalUV.y, GlobalUV.z - GlobalVolumeTexelSize), 0).x;
 	}
 	else if (ClipmapIndex == 2)
 	{
-		R = Texture3DSampleLevel(GlobalDistanceFieldTexture2, GlobalDistanceFieldSampler2, float3(GlobalUV.x + GlobalVolumeTexelSize, GlobalUV.y, GlobalUV.z), 0).x;
-		L = Texture3DSampleLevel(GlobalDistanceFieldTexture2, GlobalDistanceFieldSampler2, float3(GlobalUV.x - GlobalVolumeTexelSize, GlobalUV.y, GlobalUV.z), 0).x;
-		F = Texture3DSampleLevel(GlobalDistanceFieldTexture2, GlobalDistanceFieldSampler2, float3(GlobalUV.x, GlobalUV.y + GlobalVolumeTexelSize, GlobalUV.z), 0).x;
-		B = Texture3DSampleLevel(GlobalDistanceFieldTexture2, GlobalDistanceFieldSampler2, float3(GlobalUV.x, GlobalUV.y - GlobalVolumeTexelSize, GlobalUV.z), 0).x;
-		U = Texture3DSampleLevel(GlobalDistanceFieldTexture2, GlobalDistanceFieldSampler2, float3(GlobalUV.x, GlobalUV.y, GlobalUV.z + GlobalVolumeTexelSize), 0).x;
-		D = Texture3DSampleLevel(GlobalDistanceFieldTexture2, GlobalDistanceFieldSampler2, float3(GlobalUV.x, GlobalUV.y, GlobalUV.z - GlobalVolumeTexelSize), 0).x;
+		R = Texture3DSampleLevel(GlobalDistanceFieldTexture2, SharedGlobalDistanceFieldSampler2, float3(GlobalUV.x + GlobalVolumeTexelSize, GlobalUV.y, GlobalUV.z), 0).x;
+		L = Texture3DSampleLevel(GlobalDistanceFieldTexture2, SharedGlobalDistanceFieldSampler2, float3(GlobalUV.x - GlobalVolumeTexelSize, GlobalUV.y, GlobalUV.z), 0).x;
+		F = Texture3DSampleLevel(GlobalDistanceFieldTexture2, SharedGlobalDistanceFieldSampler2, float3(GlobalUV.x, GlobalUV.y + GlobalVolumeTexelSize, GlobalUV.z), 0).x;
+		B = Texture3DSampleLevel(GlobalDistanceFieldTexture2, SharedGlobalDistanceFieldSampler2, float3(GlobalUV.x, GlobalUV.y - GlobalVolumeTexelSize, GlobalUV.z), 0).x;
+		U = Texture3DSampleLevel(GlobalDistanceFieldTexture2, SharedGlobalDistanceFieldSampler2, float3(GlobalUV.x, GlobalUV.y, GlobalUV.z + GlobalVolumeTexelSize), 0).x;
+		D = Texture3DSampleLevel(GlobalDistanceFieldTexture2, SharedGlobalDistanceFieldSampler2, float3(GlobalUV.x, GlobalUV.y, GlobalUV.z - GlobalVolumeTexelSize), 0).x;
 	}
 	else if (ClipmapIndex == 3)
 	{
-		R = Texture3DSampleLevel(GlobalDistanceFieldTexture3, GlobalDistanceFieldSampler3, float3(GlobalUV.x + GlobalVolumeTexelSize, GlobalUV.y, GlobalUV.z), 0).x;
-		L = Texture3DSampleLevel(GlobalDistanceFieldTexture3, GlobalDistanceFieldSampler3, float3(GlobalUV.x - GlobalVolumeTexelSize, GlobalUV.y, GlobalUV.z), 0).x;
-		F = Texture3DSampleLevel(GlobalDistanceFieldTexture3, GlobalDistanceFieldSampler3, float3(GlobalUV.x, GlobalUV.y + GlobalVolumeTexelSize, GlobalUV.z), 0).x;
-		B = Texture3DSampleLevel(GlobalDistanceFieldTexture3, GlobalDistanceFieldSampler3, float3(GlobalUV.x, GlobalUV.y - GlobalVolumeTexelSize, GlobalUV.z), 0).x;
-		U = Texture3DSampleLevel(GlobalDistanceFieldTexture3, GlobalDistanceFieldSampler3, float3(GlobalUV.x, GlobalUV.y, GlobalUV.z + GlobalVolumeTexelSize), 0).x;
-		D = Texture3DSampleLevel(GlobalDistanceFieldTexture3, GlobalDistanceFieldSampler3, float3(GlobalUV.x, GlobalUV.y, GlobalUV.z - GlobalVolumeTexelSize), 0).x;
+		R = Texture3DSampleLevel(GlobalDistanceFieldTexture3, SharedGlobalDistanceFieldSampler3, float3(GlobalUV.x + GlobalVolumeTexelSize, GlobalUV.y, GlobalUV.z), 0).x;
+		L = Texture3DSampleLevel(GlobalDistanceFieldTexture3, SharedGlobalDistanceFieldSampler3, float3(GlobalUV.x - GlobalVolumeTexelSize, GlobalUV.y, GlobalUV.z), 0).x;
+		F = Texture3DSampleLevel(GlobalDistanceFieldTexture3, SharedGlobalDistanceFieldSampler3, float3(GlobalUV.x, GlobalUV.y + GlobalVolumeTexelSize, GlobalUV.z), 0).x;
+		B = Texture3DSampleLevel(GlobalDistanceFieldTexture3, SharedGlobalDistanceFieldSampler3, float3(GlobalUV.x, GlobalUV.y - GlobalVolumeTexelSize, GlobalUV.z), 0).x;
+		U = Texture3DSampleLevel(GlobalDistanceFieldTexture3, SharedGlobalDistanceFieldSampler3, float3(GlobalUV.x, GlobalUV.y, GlobalUV.z + GlobalVolumeTexelSize), 0).x;
+		D = Texture3DSampleLevel(GlobalDistanceFieldTexture3, SharedGlobalDistanceFieldSampler3, float3(GlobalUV.x, GlobalUV.y, GlobalUV.z - GlobalVolumeTexelSize), 0).x;
 	}
 		
 	float Extent = GlobalVolumeCenterAndExtent[ClipmapIndex].w;
diff --git a/Engine/Shaders/GpuSkinCacheComputeShader.usf b/Engine/Shaders/GpuSkinCacheComputeShader.usf
--- a/Engine/Shaders/GpuSkinCacheComputeShader.usf
+++ b/Engine/Shaders/GpuSkinCacheComputeShader.usf
@@ -99,20 +99,18 @@ FBoneMatrix CalcBoneMatrix( FVertexUnpacked VertexUnpacked )
 Buffer<float> SkinStreamUniformBuffer;
 Buffer<float> SkinStreamInputBuffer;
 
-RWBuffer<float> SkinCacheBuffer;
+RWBuffer<float> SkinCacheBufferUAV;
 
-uint SkinCacheInputStreamFloatOffset;
-uint SkinCacheOutputBufferFloatOffset;
-uint SkinCacheVertexStride;
-uint SkinCacheVertexCount;
+uint InputStreamStart;
+uint SkinCacheStart;
+uint InputStreamStride;
+uint NumVertices;
 
 float3 MeshOrigin;
 float3 MeshExtension;
 // see r.SkinCache.Debug, default is 1.0f
 float SkinCacheDebug;
 
-//uint4 DebugParameter;
-
 FVertexUnpacked UnpackedVertex(uint VertexIndex, uint Offset )
 {
 	FVertexUnpacked Unpacked;
@@ -173,16 +171,16 @@ void SkinCacheUpdateBatchCS(uint3 GroupID : SV_GroupID,
  uint3 GroupThreadID : SV_GroupThreadID)
 {
 	uint VertexIndex = DispatchThreadID.x;
-	if (VertexIndex >= SkinCacheVertexCount)
+		if (VertexIndex >= NumVertices)
 	{
 		return;
 	}
 
-	uint FloatsPerVertex = SkinCacheVertexStride / 4;
+	uint FloatsPerVertex = InputStreamStride / 4;
 	uint VBOffset = VertexIndex * FloatsPerVertex;
 	uint RWOffset = VertexIndex * GPUSKIN_RWBUFFER_NUM_FLOATS;
-	uint InputOffset = SkinCacheInputStreamFloatOffset + VBOffset;
-	uint OutputOffset = SkinCacheOutputBufferFloatOffset + RWOffset;
+	uint InputOffset = InputStreamStart + VBOffset;
+	uint OutputOffset = SkinCacheStart + RWOffset;
 
 	FVertexUnpacked	Unpacked = UnpackedVertex(VertexIndex, InputOffset);
 
@@ -222,22 +220,24 @@ void SkinCacheUpdateBatchCS(uint3 GroupID : SV_GroupID,
 	float3 TangentX = mul((float3x3)Intermediates.BlendMatrix, LocalTangentX);
 	float3 TangentZ = mul((float3x3)Intermediates.BlendMatrix, LocalTangentZ);
 
-	SkinCacheBuffer[OutputOffset + GPUSKIN_RWBUFFER_OFFSET_POSITION + 0] = SPos.x;
-	SkinCacheBuffer[OutputOffset + GPUSKIN_RWBUFFER_OFFSET_POSITION + 1] = SPos.y;
-	SkinCacheBuffer[OutputOffset + GPUSKIN_RWBUFFER_OFFSET_POSITION + 2] = SPos.z;
+	SkinCacheBufferUAV[OutputOffset + GPUSKIN_RWBUFFER_OFFSET_POSITION + 0] = SPos.x;
+	SkinCacheBufferUAV[OutputOffset + GPUSKIN_RWBUFFER_OFFSET_POSITION + 1] = SPos.y;
+	SkinCacheBufferUAV[OutputOffset + GPUSKIN_RWBUFFER_OFFSET_POSITION + 2] = SPos.z;
 	// we keep the w component (what is in it?)
 	uint TangentXCompressed = PackS8x3N(TangentX);
-	SkinCacheBuffer[OutputOffset + GPUSKIN_RWBUFFER_OFFSET_TANGENT_X] = asfloat(TangentXCompressed | (asuint(SkinStreamInputBuffer[InputOffset + GPUSKIN_VB_OFFSET_TANGENT_X]) & 0xff000000));
+//	uint TangentXCompressed = PackS8x3N(float3(0,0,0));
+	SkinCacheBufferUAV[OutputOffset + GPUSKIN_RWBUFFER_OFFSET_TANGENT_X] = asfloat(TangentXCompressed | (asuint(SkinStreamInputBuffer[InputOffset + GPUSKIN_VB_OFFSET_TANGENT_X]) & 0xff000000));
 	// we keep the w component which contains the sign of tangent basis determinant (left or right handed)
 	uint TangentZCompressed = PackS8x3N(TangentZ.xyz);
-	SkinCacheBuffer[OutputOffset + GPUSKIN_RWBUFFER_OFFSET_TANGENT_Z] = asfloat(TangentZCompressed | (asuint(SkinStreamInputBuffer[ InputOffset + GPUSKIN_VB_OFFSET_TANGENT_Z]) & 0xff000000));
+//	uint TangentZCompressed = PackS8x3N(float3(0,0,0));
+	SkinCacheBufferUAV[OutputOffset + GPUSKIN_RWBUFFER_OFFSET_TANGENT_Z] = asfloat(TangentZCompressed | (asuint(SkinStreamInputBuffer[ InputOffset + GPUSKIN_VB_OFFSET_TANGENT_Z]) & 0xff000000));
 
 #if GPUSKIN_DEBUG_PASSTHROUGH
 	// Passthrough debug code
-	SkinCacheBuffer[OutputOffset + GPUSKIN_RWBUFFER_OFFSET_POSITION + 0] = SkinStreamInputBuffer[InputOffset + GPUSKIN_VB_OFFSET_POSITION + 0];
-	SkinCacheBuffer[OutputOffset + GPUSKIN_RWBUFFER_OFFSET_POSITION + 1] = SkinStreamInputBuffer[InputOffset + GPUSKIN_VB_OFFSET_POSITION + 1];
-	SkinCacheBuffer[OutputOffset + GPUSKIN_RWBUFFER_OFFSET_POSITION + 2] = SkinStreamInputBuffer[InputOffset + GPUSKIN_VB_OFFSET_POSITION + 2];
-	SkinCacheBuffer[OutputOffset + GPUSKIN_RWBUFFER_OFFSET_TANGENT_X] = SkinStreamInputBuffer[InputOffset + GPUSKIN_VB_OFFSET_TANGENT_X];
-	SkinCacheBuffer[OutputOffset + GPUSKIN_RWBUFFER_OFFSET_TANGENT_Z] = SkinStreamInputBuffer[InputOffset + GPUSKIN_VB_OFFSET_TANGENT_Z];
+	SkinCacheBufferUAV[OutputOffset + GPUSKIN_RWBUFFER_OFFSET_POSITION + 0] = SkinStreamInputBuffer[InputOffset + GPUSKIN_VB_OFFSET_POSITION + 0];
+	SkinCacheBufferUAV[OutputOffset + GPUSKIN_RWBUFFER_OFFSET_POSITION + 1] = SkinStreamInputBuffer[InputOffset + GPUSKIN_VB_OFFSET_POSITION + 1];
+	SkinCacheBufferUAV[OutputOffset + GPUSKIN_RWBUFFER_OFFSET_POSITION + 2] = SkinStreamInputBuffer[InputOffset + GPUSKIN_VB_OFFSET_POSITION + 2];
+	SkinCacheBufferUAV[OutputOffset + GPUSKIN_RWBUFFER_OFFSET_TANGENT_X] = SkinStreamInputBuffer[InputOffset + GPUSKIN_VB_OFFSET_TANGENT_X];
+	SkinCacheBufferUAV[OutputOffset + GPUSKIN_RWBUFFER_OFFSET_TANGENT_Z] = SkinStreamInputBuffer[InputOffset + GPUSKIN_VB_OFFSET_TANGENT_Z];
 #endif
 }
diff --git a/Engine/Shaders/GpuSkinCommon.usf b/Engine/Shaders/GpuSkinCommon.usf
--- a/Engine/Shaders/GpuSkinCommon.usf
+++ b/Engine/Shaders/GpuSkinCommon.usf
@@ -4,10 +4,6 @@
 	GpuSkinCacheCommon.usf: Common properties and defines for GPU Skinning Cache
 =============================================================================*/
 
-#ifndef GPUSKIN_USE_EXTRA_INFLUENCES
-#define GPUSKIN_USE_EXTRA_INFLUENCES 0
-#endif
-
 // These offsets are in uint32/floats; they attempt to match TGPUSkinVertexBase<> & derived classes:
 //	TGPUSkinVertexBase<false>
 //	TGPUSkinVertexBase<true>
diff --git a/Engine/Shaders/HeightfieldLighting.usf b/Engine/Shaders/HeightfieldLighting.usf
--- a/Engine/Shaders/HeightfieldLighting.usf
+++ b/Engine/Shaders/HeightfieldLighting.usf
@@ -17,111 +17,7 @@
 #include "LightFunctionCommon.usf"
 #endif
 
-Texture2D HeightfieldTexture;
-SamplerState HeightfieldSampler;
-Texture2D DiffuseColorTexture;
-SamplerState DiffuseColorSampler;
-uint NumHeightfields;
-Buffer<float4> HeightfieldDescriptions;
-
-// In float4's, must match C++
-#define HEIGHTFIELD_DATA_STRIDE 12
-
-float4 GetHeightfieldScaleBias(uint Index, out bool bHasSubsections)
-{
-	float4 HeightfieldVector = HeightfieldDescriptions[Index * HEIGHTFIELD_DATA_STRIDE + 0];
-	// Sign bit stored presence of subsections
-	bHasSubsections = HeightfieldVector.x < 0;
-	HeightfieldVector.x = abs(HeightfieldVector.x);
-	return HeightfieldVector;
-}
-
-float4 GetHeightfieldMinMaxUV(uint Index)
-{
-	return HeightfieldDescriptions[Index * HEIGHTFIELD_DATA_STRIDE + 1];
-}
-
-float4 GetHeightfieldGlobalScaleBias(uint Index)
-{
-	return HeightfieldDescriptions[Index * HEIGHTFIELD_DATA_STRIDE + 2];
-}
-
-float2 GetHeightfieldSize(uint Index)
-{
-	return HeightfieldDescriptions[Index * HEIGHTFIELD_DATA_STRIDE + 3].xy;
-}
-
-float2 GetGlobalTexelSize(uint Index)
-{
-	return HeightfieldDescriptions[Index * HEIGHTFIELD_DATA_STRIDE + 3].zw;
-}
-
-float4x4 GetWorldToLocal(uint Index)
-{
-	return float4x4(
-		HeightfieldDescriptions[Index * HEIGHTFIELD_DATA_STRIDE + 4],
-		HeightfieldDescriptions[Index * HEIGHTFIELD_DATA_STRIDE + 5],
-		HeightfieldDescriptions[Index * HEIGHTFIELD_DATA_STRIDE + 6],
-		HeightfieldDescriptions[Index * HEIGHTFIELD_DATA_STRIDE + 7]);
-}
-
-float4x4 GetLocalToWorld(uint Index)
-{
-	return float4x4(
-		HeightfieldDescriptions[Index * HEIGHTFIELD_DATA_STRIDE + 8],
-		HeightfieldDescriptions[Index * HEIGHTFIELD_DATA_STRIDE + 9],
-		HeightfieldDescriptions[Index * HEIGHTFIELD_DATA_STRIDE + 10],
-		HeightfieldDescriptions[Index * HEIGHTFIELD_DATA_STRIDE + 11]);
-}
-
-float2 GetHeightfieldUV(uint HeightfieldIndex, float2 LocalPosition, out float4 MinMaxHeightfieldUV)
-{
-	bool bHasSubsections;
-	float4 HeightfieldScaleBias = GetHeightfieldScaleBias(HeightfieldIndex, bHasSubsections);
-	float2 HeightfieldUV = LocalPosition * HeightfieldScaleBias.xy + HeightfieldScaleBias.zw;
-	MinMaxHeightfieldUV = GetHeightfieldMinMaxUV(HeightfieldIndex);
-
-	if (bHasSubsections)
-	{
-		float2 CenterUV = .5f * (MinMaxHeightfieldUV.xy + MinMaxHeightfieldUV.zw);
-		// One row of heights in the heightfield are duplicated between subsections
-		// The only number of subsections supported is 2
-		// If we are in subsection 0, no adjustment is needed, but if we are in subsection 1 we need to shift up by one texel
-		HeightfieldUV += HeightfieldUV + .5f * HeightfieldScaleBias.xy > CenterUV ? HeightfieldScaleBias.xy : 0;
-	}
-
-	return HeightfieldUV;
-}
-
-float2 GetHeightfieldUV(uint HeightfieldIndex, float2 LocalPosition)
-{
-	float4 MinMaxHeightfieldUV;
-	return GetHeightfieldUV(HeightfieldIndex, LocalPosition, MinMaxHeightfieldUV);
-}
-
-float3 GetHeightfieldWorldPositionAndNormal(uint HeightfieldIndex, float2 LocalPosition, float2 HeightfieldUV, out float3 WorldNormal)
-{
-	//@todo - don't call GetHeightfieldScaleBias twice
-	bool bHasSubsections;
-	float4 HeightfieldScaleBias = GetHeightfieldScaleBias(HeightfieldIndex, bHasSubsections);
-
-	float4 SampleValue = Texture2DSampleLevel(HeightfieldTexture, HeightfieldSampler, HeightfieldUV + .5f * HeightfieldScaleBias.xy, 0);
-	float HeightfieldHeight = DecodePackedHeight(SampleValue.xy);
-
-	float2 SampleNormal = float2(SampleValue.b, SampleValue.a) * float2(2.0,2.0) - float2(1.0,1.0);
-	WorldNormal = float3(SampleNormal, sqrt(max(1.0 - dot(SampleNormal, SampleNormal), 0.0)));
-
-	float3 HeightfieldShadingPosition = float3(LocalPosition, HeightfieldHeight);
-	float3 WorldHeightfieldShadingPosition = mul(float4(HeightfieldShadingPosition, 1), GetLocalToWorld(HeightfieldIndex)).xyz;
-
-	return WorldHeightfieldShadingPosition;
-}
-
-float3 GetHeightfieldWorldPosition(uint HeightfieldIndex, float2 LocalPosition, float2 HeightfieldUV)
-{
-	float3 WorldNormal;
-	return GetHeightfieldWorldPositionAndNormal(HeightfieldIndex, LocalPosition, HeightfieldUV, WorldNormal);
-}
+#include "HeightfieldLightingShared.usf"
 
 Texture2D GlobalHeightfieldTexture;
 Texture2D GlobalNormalTexture;
@@ -1064,3 +960,4 @@ void PreCullTrianglesToHeightfieldsCS(
 		}
 	}
 }
+
diff --git a/Engine/Shaders/HeightfieldLightingShared.usf b/Engine/Shaders/HeightfieldLightingShared.usf
new file mode 100644
--- /dev/null
+++ b/Engine/Shaders/HeightfieldLightingShared.usf
@@ -0,0 +1,111 @@
+// Copyright 1998-2016 Epic Games, Inc. All Rights Reserved.
+
+/*=============================================================================
+	HeightfieldLightingShared.usf
+=============================================================================*/
+
+Texture2D HeightfieldTexture;
+SamplerState HeightfieldSampler;
+Texture2D DiffuseColorTexture;
+SamplerState DiffuseColorSampler;
+uint NumHeightfields;
+Buffer<float4> HeightfieldDescriptions;
+
+// In float4's, must match C++
+#define HEIGHTFIELD_DATA_STRIDE 12
+
+float4 GetHeightfieldScaleBias(uint Index, out bool bHasSubsections)
+{
+	float4 HeightfieldVector = HeightfieldDescriptions[Index * HEIGHTFIELD_DATA_STRIDE + 0];
+	// Sign bit stored presence of subsections
+	bHasSubsections = HeightfieldVector.x < 0;
+	HeightfieldVector.x = abs(HeightfieldVector.x);
+	return HeightfieldVector;
+}
+
+float4 GetHeightfieldMinMaxUV(uint Index)
+{
+	return HeightfieldDescriptions[Index * HEIGHTFIELD_DATA_STRIDE + 1];
+}
+
+float4 GetHeightfieldGlobalScaleBias(uint Index)
+{
+	return HeightfieldDescriptions[Index * HEIGHTFIELD_DATA_STRIDE + 2];
+}
+
+float2 GetHeightfieldSize(uint Index)
+{
+	return HeightfieldDescriptions[Index * HEIGHTFIELD_DATA_STRIDE + 3].xy;
+}
+
+float2 GetGlobalTexelSize(uint Index)
+{
+	return HeightfieldDescriptions[Index * HEIGHTFIELD_DATA_STRIDE + 3].zw;
+}
+
+float4x4 GetWorldToLocal(uint Index)
+{
+	return float4x4(
+		HeightfieldDescriptions[Index * HEIGHTFIELD_DATA_STRIDE + 4],
+		HeightfieldDescriptions[Index * HEIGHTFIELD_DATA_STRIDE + 5],
+		HeightfieldDescriptions[Index * HEIGHTFIELD_DATA_STRIDE + 6],
+		HeightfieldDescriptions[Index * HEIGHTFIELD_DATA_STRIDE + 7]);
+}
+
+float4x4 GetLocalToWorld(uint Index)
+{
+	return float4x4(
+		HeightfieldDescriptions[Index * HEIGHTFIELD_DATA_STRIDE + 8],
+		HeightfieldDescriptions[Index * HEIGHTFIELD_DATA_STRIDE + 9],
+		HeightfieldDescriptions[Index * HEIGHTFIELD_DATA_STRIDE + 10],
+		HeightfieldDescriptions[Index * HEIGHTFIELD_DATA_STRIDE + 11]);
+}
+
+float2 GetHeightfieldUV(uint HeightfieldIndex, float2 LocalPosition, out float4 MinMaxHeightfieldUV)
+{
+	bool bHasSubsections;
+	float4 HeightfieldScaleBias = GetHeightfieldScaleBias(HeightfieldIndex, bHasSubsections);
+	float2 HeightfieldUV = LocalPosition * HeightfieldScaleBias.xy + HeightfieldScaleBias.zw;
+	MinMaxHeightfieldUV = GetHeightfieldMinMaxUV(HeightfieldIndex);
+
+	if (bHasSubsections)
+	{
+		float2 CenterUV = .5f * (MinMaxHeightfieldUV.xy + MinMaxHeightfieldUV.zw);
+		// One row of heights in the heightfield are duplicated between subsections
+		// The only number of subsections supported is 2
+		// If we are in subsection 0, no adjustment is needed, but if we are in subsection 1 we need to shift up by one texel
+		HeightfieldUV += HeightfieldUV + .5f * HeightfieldScaleBias.xy > CenterUV ? HeightfieldScaleBias.xy : 0;
+	}
+
+	return HeightfieldUV;
+}
+
+float2 GetHeightfieldUV(uint HeightfieldIndex, float2 LocalPosition)
+{
+	float4 MinMaxHeightfieldUV;
+	return GetHeightfieldUV(HeightfieldIndex, LocalPosition, MinMaxHeightfieldUV);
+}
+
+float3 GetHeightfieldWorldPositionAndNormal(uint HeightfieldIndex, float2 LocalPosition, float2 HeightfieldUV, out float3 WorldNormal)
+{
+	//@todo - don't call GetHeightfieldScaleBias twice
+	bool bHasSubsections;
+	float4 HeightfieldScaleBias = GetHeightfieldScaleBias(HeightfieldIndex, bHasSubsections);
+
+	float4 SampleValue = Texture2DSampleLevel(HeightfieldTexture, HeightfieldSampler, HeightfieldUV + .5f * HeightfieldScaleBias.xy, 0);
+	float HeightfieldHeight = DecodePackedHeight(SampleValue.xy);
+
+	float2 SampleNormal = float2(SampleValue.b, SampleValue.a) * float2(2.0,2.0) - float2(1.0,1.0);
+	WorldNormal = float3(SampleNormal, sqrt(max(1.0 - dot(SampleNormal, SampleNormal), 0.0)));
+
+	float3 HeightfieldShadingPosition = float3(LocalPosition, HeightfieldHeight);
+	float3 WorldHeightfieldShadingPosition = mul(float4(HeightfieldShadingPosition, 1), GetLocalToWorld(HeightfieldIndex)).xyz;
+
+	return WorldHeightfieldShadingPosition;
+}
+
+float3 GetHeightfieldWorldPosition(uint HeightfieldIndex, float2 LocalPosition, float2 HeightfieldUV)
+{
+	float3 WorldNormal;
+	return GetHeightfieldWorldPositionAndNormal(HeightfieldIndex, LocalPosition, HeightfieldUV, WorldNormal);
+}
diff --git a/Engine/Shaders/HitProxyVertexShader.usf b/Engine/Shaders/HitProxyVertexShader.usf
--- a/Engine/Shaders/HitProxyVertexShader.usf
+++ b/Engine/Shaders/HitProxyVertexShader.usf
@@ -89,6 +89,9 @@ void Main(
 	FVertexFactoryInput Input,
 	OPTIONAL_VertexID
 	out FHitProxyVSOutput Output
+#if USE_GLOBAL_CLIP_PLANE && !USING_TESSELLATION
+	, out float OutGlobalClipPlaneDistance : SV_ClipDistance
+#endif
 	)
 {
 	ResolvedView = ResolveView();
@@ -115,6 +118,10 @@ void Main(
 		Output.Position = mul(RasterizedWorldPosition, View.TranslatedWorldToClip);
 	}
 
+#if USE_GLOBAL_CLIP_PLANE
+	OutGlobalClipPlaneDistance = dot(ResolvedView.GlobalClippingPlane, float4(WorldPosition.xyz - ResolvedView.PreViewTranslation.xyz, 1));
+#endif
+
 	Output.FactoryInterpolants = VertexFactoryGetInterpolantsVSToPS(Input, VFIntermediates, VertexParameters);
 	Output.PixelPosition = WorldPosition;
 	
diff --git a/Engine/Shaders/InstancedStereo.usf b/Engine/Shaders/InstancedStereo.usf
--- a/Engine/Shaders/InstancedStereo.usf
+++ b/Engine/Shaders/InstancedStereo.usf
@@ -44,6 +44,7 @@ struct ViewDependentState {
 	float4x4 PrevInvViewProj;
 	float4x4 PrevScreenToTranslatedWorld;
 	float4x4 ClipToPrevClip;
+	float4 GlobalClippingPlane;
 };
 
 // Used when instanced stereo is disabled and for shaders which don't support or need it.
@@ -89,6 +90,7 @@ ViewDependentState ResolveView()
 	Result.PrevInvViewProj = View.PrevInvViewProj;
 	Result.PrevScreenToTranslatedWorld = View.PrevScreenToTranslatedWorld;
 	Result.ClipToPrevClip = View.ClipToPrevClip;
+	Result.GlobalClippingPlane = View.GlobalClippingPlane;
 
 	return Result;
 }
@@ -168,6 +170,7 @@ ViewDependentState ResolveView(uint ViewIndex)
 		Result.PrevInvViewProj = View.PrevInvViewProj;
 		Result.PrevScreenToTranslatedWorld = View.PrevScreenToTranslatedWorld;
 		Result.ClipToPrevClip = View.ClipToPrevClip;
+		Result.GlobalClippingPlane = View.GlobalClippingPlane;
 	}
 	else
 	{
@@ -208,6 +211,7 @@ ViewDependentState ResolveView(uint ViewIndex)
 		Result.PrevInvViewProj = InstancedView.PrevInvViewProj;
 		Result.PrevScreenToTranslatedWorld = InstancedView.PrevScreenToTranslatedWorld;
 		Result.ClipToPrevClip = InstancedView.ClipToPrevClip;
+		Result.GlobalClippingPlane = View.GlobalClippingPlane;
 	}
 
 	return Result;
diff --git a/Engine/Shaders/LightMapDensityShader.usf b/Engine/Shaders/LightMapDensityShader.usf
--- a/Engine/Shaders/LightMapDensityShader.usf
+++ b/Engine/Shaders/LightMapDensityShader.usf
@@ -71,6 +71,9 @@ void MainVertexShader(
 	FVertexFactoryInput Input,
 	OPTIONAL_VertexID
 	out FLightMapDensityVSOutput Output
+#if USE_GLOBAL_CLIP_PLANE && !USING_TESSELLATION
+	, out float OutGlobalClipPlaneDistance : SV_ClipDistance
+#endif
 	)
 {
 	ResolvedView = ResolveView();
@@ -94,6 +97,10 @@ void MainVertexShader(
 		float4 RasterizedWorldPosition = VertexFactoryGetRasterizedWorldPosition(Input, VFIntermediates, Output.WorldPosition);
 		Output.Position = mul(RasterizedWorldPosition, ResolvedView.TranslatedWorldToClip);
 	}
+
+#if USE_GLOBAL_CLIP_PLANE
+	OutGlobalClipPlaneDistance = dot(ResolvedView.GlobalClippingPlane, float4(WorldPosition.xyz - ResolvedView.PreViewTranslation.xyz, 1));
+#endif
 	Output.FactoryInterpolants = VertexFactoryGetInterpolantsVSToPS(Input, VFIntermediates, VertexParameters);
 
 #endif // USING_TESSELLATION
diff --git a/Engine/Shaders/MaterialTemplate.usf b/Engine/Shaders/MaterialTemplate.usf
--- a/Engine/Shaders/MaterialTemplate.usf
+++ b/Engine/Shaders/MaterialTemplate.usf
@@ -18,6 +18,10 @@
 
 #define NUM_MATERIAL_TEXCOORDS_VERTEX %s
 #define NUM_MATERIAL_TEXCOORDS %s
+
+#ifdef MIN_MATERIAL_TEXCOORDS 
+	#include "MinMaterialTexCoords.usf"
+#endif
  
 #if MATERIAL_ATMOSPHERIC_FOG
 	#include "AtmosphereCommon.usf"
@@ -1791,6 +1795,9 @@ half3x3 AssembleTangentToWorld( half3 TangentToWorld0, half4 TangentToWorld2 )
 	#define OPTIONAL_OutDepth
 #endif
 
+// Using SV_ClipDistance has overhead (15% slower base pass in triangle bound test scene on PS4) so projects have to opt-in
+#define USE_GLOBAL_CLIP_PLANE (SUPPORT_GLOBAL_CLIP_PLANE && PROJECT_ALLOW_GLOBAL_CLIP_PLANE)
+
 float ApplyPixelDepthOffsetToMaterialParameters(inout FMaterialPixelParameters MaterialParameters, FPixelMaterialInputs PixelMaterialInputs, out float OutDepth)
 {
 	// Can only push into the scene for conservative depth writes
diff --git a/Engine/Shaders/MinMaterialTexCoords.usf b/Engine/Shaders/MinMaterialTexCoords.usf
new file mode 100644
--- /dev/null
+++ b/Engine/Shaders/MinMaterialTexCoords.usf
@@ -0,0 +1,18 @@
+// Copyright 1998-2016 Epic Games, Inc. All Rights Reserved.
+
+/*=============================================================================
+	MinMaterialTexCoords.usf: Definition allowing implementing MinMaterialTexCoords.
+=============================================================================*/
+
+// Allow some shader types to change the number of coords used. 
+
+#if MIN_MATERIAL_TEXCOORDS > NUM_MATERIAL_TEXCOORDS_VERTEX
+	#undef NUM_MATERIAL_TEXCOORDS_VERTEX
+	#define NUM_MATERIAL_TEXCOORDS_VERTEX MIN_MATERIAL_TEXCOORDS
+#endif
+
+#if MIN_MATERIAL_TEXCOORDS > NUM_MATERIAL_TEXCOORDS
+	#undef NUM_MATERIAL_TEXCOORDS
+	#define NUM_MATERIAL_TEXCOORDS MIN_MATERIAL_TEXCOORDS
+#endif
+
diff --git a/Engine/Shaders/MissingShaderPixelShader.usf b/Engine/Shaders/MissingShaderPixelShader.usf
--- a/Engine/Shaders/MissingShaderPixelShader.usf
+++ b/Engine/Shaders/MissingShaderPixelShader.usf
@@ -23,7 +23,7 @@ void Main(
 	ResolvedView = ResolveView();
 #endif
 
-	int2 CheckBoardCoord = DebugInputs.TexCoord0.xy * 64;
+	int2 CheckBoardCoord = DebugInputs.TexCoord01.xy * 64;
 	float CheckBoard = (CheckBoardCoord.x + CheckBoardCoord.y) % 2;
 
 	float3 Result = float3(CheckBoard, CheckBoard, CheckBoard);
diff --git a/Engine/Shaders/NumMaterialTexCoordsOverride.usf b/Engine/Shaders/NumMaterialTexCoordsOverride.usf
new file mode 100644
--- /dev/null
+++ b/Engine/Shaders/NumMaterialTexCoordsOverride.usf
@@ -0,0 +1,22 @@
+// Copyright 1998-2016 Epic Games, Inc. All Rights Reserved.
+
+/*=============================================================================
+	NumMaterialTexCoordsOverride.usf: Definition allowing to override the number 
+	of texcoord required by the material.
+=============================================================================*/
+
+// Allow some shader types to override the number of coords used. Note that even though the shader type is the one including "Material.usf",
+// the strings here only depend on the material asset, and can't be changed per shader type or per vertex factory at the CPU level.
+#ifdef NUM_MATERIAL_TEXCOORDS_OVERRIDE
+
+#if NUM_MATERIAL_TEXCOORDS_OVERRIDE > NUM_MATERIAL_TEXCOORDS_VERTEX
+	#undef NUM_MATERIAL_TEXCOORDS_VERTEX
+	#define NUM_MATERIAL_TEXCOORDS_VERTEX NUM_MATERIAL_TEXCOORDS_OVERRIDE
+#endif
+
+#if NUM_MATERIAL_TEXCOORDS_OVERRIDE > NUM_MATERIAL_TEXCOORDS
+	#undef NUM_MATERIAL_TEXCOORDS
+	#define NUM_MATERIAL_TEXCOORDS NUM_MATERIAL_TEXCOORDS_OVERRIDE
+#endif
+
+#endif
diff --git a/Engine/Shaders/PNTriangles.usf b/Engine/Shaders/PNTriangles.usf
--- a/Engine/Shaders/PNTriangles.usf
+++ b/Engine/Shaders/PNTriangles.usf
@@ -169,9 +169,12 @@ float4 ComputeControlPoint(float4 PositionA, float4 PositionB, float4 NormalA)
 		FHullShaderConstantOutput HSConstantData,
 		const OutputPatch<FPNTessellationHSToDS, OUT_CONTROL_POINTS> I,
 		float3 BarycentricCoords : SV_DomainLocation
+#if USE_GLOBAL_CLIP_PLANE
+		, out float OutGlobalClipPlaneDistance : SV_ClipDistance
+#endif
 #if INSTANCED_STEREO
-		, out float4 OutClipDistance : SV_ClipDistance
-		, out float4 OutCullDistance : SV_CullDistance
+		, out float OutClipDistance : SV_ClipDistance1
+		, out float OutCullDistance : SV_CullDistance1
 #endif
 		)
 	{
@@ -264,7 +267,11 @@ float4 ComputeControlPoint(float4 PositionA, float4 PositionB, float4 NormalA)
 		// world space displacement value.  This will likely use the tangent2world transform.
 		float3 WorldDisplacement = GetMaterialWorldDisplacement(DisplacementMaterialParameters);
 		WorldPosition.xyz += WorldDisplacement;
-																									
+										
+		#if USE_GLOBAL_CLIP_PLANE
+			OutGlobalClipPlaneDistance = dot(ResolvedView.GlobalClippingPlane, float4(WorldPosition.xyz - ResolvedView.PreViewTranslation.xyz, 1));
+		#endif
+
 		FPassSpecificVSToPS Result = PassFinalizeTessellationOutput(Interp, WorldPosition, MaterialParameters);
 
 		#if INSTANCED_STEREO
diff --git a/Engine/Shaders/PlanarReflectionShaders.usf b/Engine/Shaders/PlanarReflectionShaders.usf
new file mode 100644
--- /dev/null
+++ b/Engine/Shaders/PlanarReflectionShaders.usf
@@ -0,0 +1,31 @@
+// Copyright 1998-2016 Epic Games, Inc. All Rights Reserved.
+
+/*=============================================================================
+ PlanarReflectionShaders.usf
+=============================================================================*/
+
+#include "Common.usf"
+#include "DeferredShadingCommon.usf"
+#include "PlanarReflectionShared.usf"
+
+void PlanarReflectionPS(
+	float2 InUV : TEXCOORD0,
+	float3 ScreenVector : TEXCOORD1,
+	float4 SVPos : SV_POSITION,
+	out float4 OutColor : SV_Target0
+	)
+{
+	OutColor = 0;
+
+	float SceneDepth = CalcSceneDepth(InUV);
+	float3 WorldPosition = ScreenVector * SceneDepth + View.WorldCameraOrigin;
+		
+	FScreenSpaceData ScreenSpaceData = GetScreenSpaceData(InUV);
+
+	BRANCH
+	// Skip unlit pixels (skybox)
+	if (ScreenSpaceData.GBuffer.ShadingModelID > 0)
+	{
+		OutColor = ComputePlanarReflections(WorldPosition, ScreenSpaceData.GBuffer.WorldNormal, PlanarReflectionSampler);
+	}
+}
diff --git a/Engine/Shaders/PlanarReflectionShared.usf b/Engine/Shaders/PlanarReflectionShared.usf
new file mode 100644
--- /dev/null
+++ b/Engine/Shaders/PlanarReflectionShared.usf
@@ -0,0 +1,67 @@
+// Copyright 1998-2016 Epic Games, Inc. All Rights Reserved.
+
+/*=============================================================================
+ PlanarReflectionShared.usf
+=============================================================================*/
+
+// World space reflection plane, normalized
+float4 ReflectionPlane;
+
+// Used to transform normals into the mirrored space
+float3x4 InverseTransposeMirrorMatrix;
+
+// x = scale for plane distance fading, y = bias for plane distance fading, z = controls normal distortion strength
+float3 PlanarReflectionParameters;
+
+// x = scale for angle fading, y = bias for angle fading
+float2 PlanarReflectionParameters2;
+
+// Projection matrix that was used to capture PlanarReflectionTexture
+float4x4 ProjectionWithExtraFOV;
+
+Texture2D PlanarReflectionTexture;
+SamplerState PlanarReflectionSampler;
+
+float4 ComputePlanarReflections(float3 WorldPosition, float3 WorldNormal, SamplerState SharedClampSampler)
+{
+	float4 OutPlanarReflection = 0;
+
+	float PlaneDistance = dot(ReflectionPlane, float4(WorldPosition, -1));
+	float DistanceFade = 1 - saturate(abs(PlaneDistance) * PlanarReflectionParameters.x + PlanarReflectionParameters.y);
+
+	BRANCH
+	if (DistanceFade > 0)
+	{
+		// CameraToPixel in the main view is what we used as ReflectionVector when rendering the reflection pass to PlanarReflectionTexture
+		float3 CameraToPixel = normalize(WorldPosition - View.WorldCameraOrigin);
+		// Reflect the effective ReflectionVector in mirrored space to get the original camera vector
+		float3 MirroredCameraVector = reflect(CameraToPixel, -ReflectionPlane.xyz);
+		// Transform the GBuffer normal into mirrored space
+		float3 MirroredNormal = mul(WorldNormal, (float3x3)InverseTransposeMirrorMatrix);
+		// Reflect the original camera vector across the GBuffer normal in mirrored space
+		float3 MirroredReflectionVectorOffNormal = reflect(MirroredCameraVector, MirroredNormal);
+		// At this point we have a new reflection vector off of the GBuffer normal, and we need to approximate its intersection with the scene
+		// An accurate intersection would ray trace the planar reflection depth buffer
+		// As an approximation we are just intersecting with a user defined sphere
+		float3 VirtualReflectionSpherePosition = WorldPosition + MirroredReflectionVectorOffNormal * PlanarReflectionParameters.z;
+		// Transform the intersection position into view space
+		float3 ViewVirtualReflectionSpherePosition = mul(float4(VirtualReflectionSpherePosition + View.PreViewTranslation.xyz, 1), View.TranslatedWorldToView).xyz;
+		// Transform the intersection position into clip space using the same projection matrix used to render PlanarReflectionTexture
+		float4 ClipVirtualReflectionSpherePosition = mul(float4(ViewVirtualReflectionSpherePosition, 1), ProjectionWithExtraFOV);
+		// Compute viewport UVs from clip space position.  
+		// Note not using View.ScreenPositionScaleBias because PlanarReflectionTexture is exactly sized with the viewport, unlike the scene render targets
+		float2 ViewportUV = ClipVirtualReflectionSpherePosition.xy / ClipVirtualReflectionSpherePosition.w * float2(.5f, -.5f) + .5f;
+
+		float AngleFade = saturate(dot(ReflectionPlane.xyz, WorldNormal) * PlanarReflectionParameters2.x + PlanarReflectionParameters2.y);
+		float FinalFade = DistanceFade * AngleFade;
+
+#if SUPPORTS_INDEPENDENT_SAMPLERS
+		OutPlanarReflection.rgb = Texture2DSampleLevel(PlanarReflectionTexture, SharedClampSampler, ViewportUV, 0).rgb * FinalFade;
+#else
+		OutPlanarReflection.rgb = Texture2DSampleLevel(PlanarReflectionTexture, PlanarReflectionSampler, ViewportUV, 0).rgb * FinalFade;
+#endif
+		OutPlanarReflection.a = FinalFade;
+	}
+
+	return OutPlanarReflection;
+}
diff --git a/Engine/Shaders/PositionOnlyDepthVertexShader.usf b/Engine/Shaders/PositionOnlyDepthVertexShader.usf
--- a/Engine/Shaders/PositionOnlyDepthVertexShader.usf
+++ b/Engine/Shaders/PositionOnlyDepthVertexShader.usf
@@ -15,10 +15,13 @@ bool bNeedsInstancedStereoBias;
 void Main(
 	FPositionOnlyVertexFactoryInput Input,
 	out float4 OutPosition : SV_POSITION
+#if USE_GLOBAL_CLIP_PLANE
+	, out float OutGlobalClipPlaneDistance : SV_ClipDistance
+#endif
 #if INSTANCED_STEREO
 	, uint InstanceId : SV_InstanceID
-	, out float4 OutClipDistance : SV_ClipDistance
-	, out float4 OutCullDistance : SV_CullDistance
+	, out float OutClipDistance : SV_ClipDistance1
+	, out float OutCullDistance : SV_CullDistance1
 #endif
 	)
 {
@@ -53,4 +56,8 @@ void Main(
 		OutPosition.z += InstancedStereoDepthBias * OutPosition.w * (float)bNeedsInstancedStereoBias;
 		#endif
 	}
+
+#if USE_GLOBAL_CLIP_PLANE
+	OutGlobalClipPlaneDistance = dot(ResolvedView.GlobalClippingPlane, float4(WorldPos.xyz - ResolvedView.PreViewTranslation.xyz, 1));
+#endif
 }
diff --git a/Engine/Shaders/PostProcessCircleDOF.usf b/Engine/Shaders/PostProcessCircleDOF.usf
--- a/Engine/Shaders/PostProcessCircleDOF.usf
+++ b/Engine/Shaders/PostProcessCircleDOF.usf
@@ -13,6 +13,19 @@
 // 0:off / 1:on (use with "vis CircleDOF0")
 #define DEBUG_SAMPLE_PATTERN 0
 
+// 0:old, 1:new(fixed high res screenshots), does not make a difference
+#define DILATION_RESOLUTION_INDEPENDENT 0
+
+// Note: Frame.CircleDOFParams.w = Frame.ViewSizeAndInvSize.x / 1920
+
+// is wrapping the DepthToCoc() function in this file for faster iteration
+// @return half res pixel radius
+float DepthToCoc2(float SceneDepth)
+{
+//	return ((SceneDepth > 8000) ? 5.0f : 2.6f) * 0.5f * Frame.CircleDOFParams.w;
+	return DepthToCoc(SceneDepth);
+}
+
 // pixel shader entry point
 void CircleSetupPS(noperspective float4 UVAndScreenPos : TEXCOORD0, out float4 OutColor0 : SV_Target0)
 {
@@ -33,7 +46,7 @@ void CircleSetupPS(noperspective float4 UVAndScreenPos : TEXCOORD0, out float4 O
 	CX = min(256 * 256, CX);
 	CY = min(256 * 256, CY);
 
-	float4 CocQuad = float4(DepthToCoc(DepthQuad.x), DepthToCoc(DepthQuad.y), DepthToCoc(DepthQuad.z), DepthToCoc(DepthQuad.w));
+	float4 CocQuad = float4(DepthToCoc2(DepthQuad.x), DepthToCoc2(DepthQuad.y), DepthToCoc2(DepthQuad.z), DepthToCoc2(DepthQuad.w));
 
 #if ENABLE_FAR_BLUR == 0
 	CocQuad = min(CocQuad, 0);
@@ -195,11 +208,25 @@ void CircleDilatePS(float4 UVAndScreenPos : TEXCOORD0, out float OutColor : SV_T
 		Sm.z = PostprocessInput0.SampleLevel(PostprocessInput0Sampler, UV.xy, 0, int2(0,1)).a;
 		Sm.w = PostprocessInput0.SampleLevel(PostprocessInput0Sampler, UV.xy, 0, int2(1,1)).a;
 	#else
-		float2 UV = UVAndScreenPos.xy + 1.0*PostprocessInput0Size.zw;
-		float4 Sd = PostprocessInput0.GatherAlpha(PostprocessInput0Sampler, UV, int2(-2,-2));
-		float4 Sg = PostprocessInput0.GatherAlpha(PostprocessInput0Sampler, UV, int2(0,-2));
-		float4 Sj = PostprocessInput0.GatherAlpha(PostprocessInput0Sampler, UV, int2(-2,0));
-		float4 Sm = PostprocessInput0.GatherAlpha(PostprocessInput0Sampler, UV, int2(0,0));
+		#if DILATION_RESOLUTION_INDEPENDENT
+			float factor = Frame.ViewSizeAndInvSize.x / 1980.0f;
+
+			float2 A = 1 + float2(-2,-2);
+			float2 B = 1 + float2(0,-2);
+			float2 C = 1 + float2(-2,0);
+			float2 D = 1 + float2(0,0);
+
+			float4 Sd = PostprocessInput0.GatherAlpha(PostprocessInput0Sampler, UVAndScreenPos.xy + A * factor * PostprocessInput0Size.zw);
+			float4 Sg = PostprocessInput0.GatherAlpha(PostprocessInput0Sampler, UVAndScreenPos.xy + B * factor * PostprocessInput0Size.zw);
+			float4 Sj = PostprocessInput0.GatherAlpha(PostprocessInput0Sampler, UVAndScreenPos.xy + C * factor * PostprocessInput0Size.zw);
+			float4 Sm = PostprocessInput0.GatherAlpha(PostprocessInput0Sampler, UVAndScreenPos.xy + D * factor * PostprocessInput0Size.zw);
+		#else 
+			float2 UV = UVAndScreenPos.xy + 1.0*PostprocessInput0Size.zw;
+			float4 Sd = PostprocessInput0.GatherAlpha(PostprocessInput0Sampler, UV, int2(-2,-2));
+			float4 Sg = PostprocessInput0.GatherAlpha(PostprocessInput0Sampler, UV, int2(0,-2));
+			float4 Sj = PostprocessInput0.GatherAlpha(PostprocessInput0Sampler, UV, int2(-2,0));
+			float4 Sm = PostprocessInput0.GatherAlpha(PostprocessInput0Sampler, UV, int2(0,0));
+		#endif
 	#endif
 	
 	// Make sure near is only near blur.
@@ -221,7 +248,7 @@ void Circle4Samples(float4 UVAndScreenPos : TEXCOORD0, out float4 OutColor0 : SV
 	//
 
 	// Fixed maximum search size (in terms of Circle of Confusion radius).
-	// Higher than 8 is too noizy for 4 samples.
+	// Higher than 8 is too noizy for 4 samples. (actual radius is scaled by width/1920)
 	float Coc = 8.0 * Frame.CircleDOFParams.w;
 
 	float LocalRand = RandomOffset.x + InRand;
@@ -233,7 +260,7 @@ void Circle4Samples(float4 UVAndScreenPos : TEXCOORD0, out float4 OutColor0 : SV
 	float RadiusBase = NoizNorm(UVAndScreenPos.xy, 0.013 * LocalRand);
 
 #if DEBUG_SAMPLE_PATTERN
-	Coc = 100;
+	Coc = 100 * Frame.CircleDOFParams.w;
 #endif
 
 	// Radius 
@@ -269,7 +296,8 @@ void Circle4Samples(float4 UVAndScreenPos : TEXCOORD0, out float4 OutColor0 : SV
 	float D4 = PostprocessInput1.SampleLevel(PostprocessInput1Sampler, UV4, 0).x;
 
 	float NearCoc = 65536.0;
-	float Feather = 2.0;
+	// >0 causes DepthBlur to disappear in low resolutions, unclear why it was 2
+	float Feather = 0.0f;
 	if(abs(D1)+Feather > R1) NearCoc = min(NearCoc, D1);
 	if(abs(D2)+Feather > R2) NearCoc = min(NearCoc, D2);
 	if(abs(D3)+Feather > R3) NearCoc = min(NearCoc, D3);
@@ -575,8 +603,8 @@ float4 Recombine2Samples(float4 UVAndScreenPos, float PixCoc, float InRand)
 
 #if QUALITY
 		// we don't have depth in the alpha channel
-		CA.a = DepthToCoc(CalcSceneDepth(UVAndScreenPos.xy + VP));
-		CB.a = DepthToCoc(CalcSceneDepth(UVAndScreenPos.xy - VP));
+		CA.a = DepthToCoc2(CalcSceneDepth(UVAndScreenPos.xy + VP));
+		CB.a = DepthToCoc2(CalcSceneDepth(UVAndScreenPos.xy - VP));
 
 		// Weight the two samples to avoid forground into background bleed.
 		float IFeather0 = 1.0/4.0;
@@ -666,7 +694,7 @@ void MainCircleRecombinePS(in float4 UVAndScreenPos : TEXCOORD0, out float4 OutC
 {
 	// Circle of confusion size for the pixel.
 	float PixDepth = CalcSceneDepth(UVAndScreenPos.xy);
-	float PixCoc = DepthToCoc(PixDepth);
+	float PixCoc = DepthToCoc2(PixDepth);
 
 	float4 HalfRes = PostprocessInput1.SampleLevel(PostprocessInput1Sampler, UVAndScreenPos.xy, 0);
 	
diff --git a/Engine/Shaders/RecomputeTangentsCommon.usf b/Engine/Shaders/RecomputeTangentsCommon.usf
--- a/Engine/Shaders/RecomputeTangentsCommon.usf
+++ b/Engine/Shaders/RecomputeTangentsCommon.usf
@@ -7,13 +7,48 @@
 #include "Common.usf"
 #include "GpuSkinCommon.usf"
 
-// intermediate tangent accumulation buffer, int se we can use atomics and to get deterministic results no matter what order
-// uses as output on MainCS(), used and input and output (to clear) for ResolveCS()
-RWBuffer<int> IntermediateAccumBufferUAV;
-
-// How many units the compute shader should process
-// MainCS(): triangles count, ResolveCS(): vertex count
-uint UnitCount;
+// constants -----------------------------
 
 // could be larger for better precision but we need to be conservative as we sum up many triangles weighed by the angle (potentially unbound)
 #define TANGENT_RANGE 0x7fff
+
+// chunk -----------------------------
+
+//
+uint NumVertices;
+//
+uint NumTriangles;
+
+// vertex input ----------------------------
+
+//
+Buffer<float> SkinStreamInputBuffer;
+// in bytes
+uint InputStreamStride;
+// in floats
+uint InputStreamStart;
+
+// index input ----------------------------
+
+// start index in IndexBuffer[]
+uint IndexBufferOffset;
+// todo: test 16bit and 32 bit 
+/// only needed for MainCS()
+Buffer<uint> IndexBuffer;
+
+// SkinCache -----------------------------
+
+// input vertex data (SkinCache output)
+// Buffer<float> GPUSkinCacheBuffer;
+// used as input in MainCS()
+Buffer<float> GPUSkinCacheBuffer;
+// used as output in ResolveCS()
+RWBuffer<float> SkinCacheBufferUAV;
+// start offset within SkinCacheBufferUAV[], in floats
+uint SkinCacheStart;
+
+// intermediate ------------------------------
+
+// intermediate tangent accumulation buffer, int se we can use atomics and to get deterministic results no matter what order
+// uses as output on MainCS(), used and input and output (to clear) for ResolveCS()
+RWBuffer<int> IntermediateAccumBufferUAV;
diff --git a/Engine/Shaders/RecomputeTangentsPerTrianglePass.usf b/Engine/Shaders/RecomputeTangentsPerTrianglePass.usf
--- a/Engine/Shaders/RecomputeTangentsPerTrianglePass.usf
+++ b/Engine/Shaders/RecomputeTangentsPerTrianglePass.usf
@@ -7,23 +7,6 @@
 #include "RecomputeTangentsCommon.usf"
 #include "FastMath.usf"
 
-// start offset within GPUSkinCacheBuffer[]
-uint GPUSkinCacheFloatOffset;
-// input vertex data (SkinCache output)
-// Buffer<float> GPUSkinCacheBuffer;
-// used as input in MainCS()
-Buffer<float> GPUSkinCacheBuffer;
-
-uint SkinCacheVertexStride;
-uint SkinCacheInputStreamFloatOffset;
-Buffer<float> SkinStreamInputBuffer;
-
-// start index in IndexBuffer[]
-uint IndexBufferOffset;
-// todo: test 16bit and 32 bit 
-/// only needed for MainCS()
-Buffer<uint> IndexBuffer;
-
 struct FVertexData
 {
 	// vertex index in the vertexbuffer
@@ -54,21 +37,21 @@ FVertexData GetVertex(uint VertexIndex)
 	// is filled out later
 	Ret.Angle = 0;
 
+	// no start offset as we reuse the same buffer always from 0 on
+	Ret.VertexIndex = VertexIndex;
+
 	{
 		// start offset within GPUSkinCacheBuffer[]
-		uint OutputOffset = VertexIndex * GPUSKIN_RWBUFFER_NUM_FLOATS + GPUSkinCacheFloatOffset;
+		uint OutputOffset = VertexIndex * GPUSKIN_RWBUFFER_NUM_FLOATS + SkinCacheStart;
 	
-		// no start offset as we reuse the same buffer always from 0 on
-		Ret.VertexIndex = VertexIndex;
-
 		Ret.Position.x = GPUSkinCacheBuffer[OutputOffset + GPUSKIN_RWBUFFER_OFFSET_POSITION + 0];
 		Ret.Position.y = GPUSkinCacheBuffer[OutputOffset + GPUSKIN_RWBUFFER_OFFSET_POSITION + 1];
 		Ret.Position.z = GPUSkinCacheBuffer[OutputOffset + GPUSKIN_RWBUFFER_OFFSET_POSITION + 2];
 	}
 
 	{
-		uint FloatsPerVertex = SkinCacheVertexStride / 4;
-		uint InputOffset = VertexIndex * FloatsPerVertex + SkinCacheInputStreamFloatOffset;
+		uint FloatsPerVertex = InputStreamStride / 4;
+		uint InputOffset = VertexIndex * FloatsPerVertex + InputStreamStart;
 
 		if(FULL_PRECISION_UV)
 		{
@@ -95,7 +78,7 @@ void MainCS(
 	uint TriangleIndex = DispatchThreadId;
 
 	// We need this to restrict processing of triangles to the valid range (We run a Dispatch with a coarse granularity and would process too many elements without that)
-	if(TriangleIndex < UnitCount)
+	if(TriangleIndex < NumTriangles)
 	{
 		// 3 triangle corners
 		FVertexData Corner[3];
@@ -105,7 +88,8 @@ void MainCS(
 
 			UNROLL for(uint CornerId = 0; CornerId < 3; ++CornerId)
 			{
-				uint VertexIndex = IndexBuffer[TriangleIndexBufferStart + CornerId]; //+ SkinCacheInputStreamFloatOffset;
+				uint FloatsPerVertex = InputStreamStride / 4;
+				uint VertexIndex = IndexBuffer[TriangleIndexBufferStart + CornerId] - InputStreamStart / FloatsPerVertex;
 
 				Corner[CornerId] = GetVertex(VertexIndex);
 			}
diff --git a/Engine/Shaders/RecomputeTangentsPerVertexPass.usf b/Engine/Shaders/RecomputeTangentsPerVertexPass.usf
--- a/Engine/Shaders/RecomputeTangentsPerVertexPass.usf
+++ b/Engine/Shaders/RecomputeTangentsPerVertexPass.usf
@@ -6,31 +6,34 @@
 
 #include "RecomputeTangentsCommon.usf"
 
-// start offset within GPUSkinCacheBufferUAV[]
-uint GPUSkinCacheFloatOffset;
-// used as output in ResolveCS()
-RWBuffer<float> GPUSkinCacheBufferUAV;
-
-uint SkinCacheInputStreamFloatOffset;
+// 0:fast and simple / 1:higher quality lighting with extreme mesh distortions
+#define ORTHOGONALIZE 1
 
 [numthreads(THREADGROUP_SIZEX, 1, 1)]
 void MainCS(
 	uint DispatchThreadId : SV_DispatchThreadID)	// DispatchThreadId = GroupId * int2(dimx,dimy) + GroupThreadId
 {
-	if(DispatchThreadId < UnitCount)
+	if(DispatchThreadId < NumVertices)
 	{
-		uint VertexIndex = DispatchThreadId + SkinCacheInputStreamFloatOffset;
 		// -1..1 range, normalized
 		float3 TangentZ, TangentX;
 		// the sign defines the handyness of the tangent matrix
 		float Orientation;
 		{
 			// no start offset as we reuse the same buffer always from 0 on
-			uint Index = VertexIndex * INTERMEDIATE_ACCUM_BUFFER_NUM_INTS;
+			uint Index = DispatchThreadId * INTERMEDIATE_ACCUM_BUFFER_NUM_INTS;
 
 			// we don't have to scale down as we normalize anyway
 			TangentZ = normalize(float3(IntermediateAccumBufferUAV[Index + 0], IntermediateAccumBufferUAV[Index + 1], IntermediateAccumBufferUAV[Index + 2]));
 			TangentX = normalize(float3(IntermediateAccumBufferUAV[Index + 3], IntermediateAccumBufferUAV[Index + 4], IntermediateAccumBufferUAV[Index + 5]));
+			
+#if ORTHOGONALIZE
+			// todo: can be optimized
+			// derive the new tangent by orthonormalizing the new normal against
+			// the base tangent vector (assuming these are normalized)
+			TangentX = normalize(TangentX - (dot(TangentX, TangentZ) * TangentZ));
+#endif
+
 			Orientation = IntermediateAccumBufferUAV[Index + 6];
 
 			// clear input UAV to 0 again to avoid another Dispatch() call
@@ -40,13 +43,20 @@ void MainCS(
 			}
 		}
 
-		uint Offset = VertexIndex * GPUSKIN_RWBUFFER_NUM_FLOATS + GPUSkinCacheFloatOffset;
+		uint FloatsPerVertex = InputStreamStride / 4;
+		uint BaseVertexIndex = InputStreamStart / FloatsPerVertex;
+
+		uint VertexIndex = DispatchThreadId;
 
 		uint TangentXCompressed = PackS8x3N(TangentX);
 		uint TangentZCompressed = PackS8x3N(TangentZ);
 		uint HandynessValue = (Orientation < 0) ? 0xff000000 : 0;
 
-		GPUSkinCacheBufferUAV[Offset + GPUSKIN_RWBUFFER_OFFSET_TANGENT_X] = asfloat(TangentXCompressed);
-		GPUSkinCacheBufferUAV[Offset + GPUSKIN_RWBUFFER_OFFSET_TANGENT_Z] = asfloat(TangentZCompressed | HandynessValue);
+		// debug (darken all affected vertices)
+//		TangentXCompressed = PackS8x3N(float3(0,0,0)); TangentZCompressed = PackS8x3N(float3(0,0,0)); HandynessValue = 0;		// do not checkin !!!!!!!!!!!!!!!!!!!!!!!!!!!!!!!!!!!!!!!!!!!!!!
+
+		uint Offset = VertexIndex * GPUSKIN_RWBUFFER_NUM_FLOATS + SkinCacheStart;
+		SkinCacheBufferUAV[Offset + GPUSKIN_RWBUFFER_OFFSET_TANGENT_X] = asfloat(TangentXCompressed);
+		SkinCacheBufferUAV[Offset + GPUSKIN_RWBUFFER_OFFSET_TANGENT_Z] = asfloat(TangentZCompressed | HandynessValue);
 	}
-}
\ No newline at end of file
+}
diff --git a/Engine/Shaders/ShaderVersion.usf b/Engine/Shaders/ShaderVersion.usf
--- a/Engine/Shaders/ShaderVersion.usf
+++ b/Engine/Shaders/ShaderVersion.usf
@@ -2,4 +2,4 @@
 // Each time the console command is executed it generates a new GUID. As this file is included
 // in common.usf (which should be included in any shader) it allows to invalidate the shader DDC.
 // 
-// GUID = 2F8BBA00C044A97B429F64CD2CFE7B8
+// GUID = E2CCCE226CFF4FD4AC5543BFC426587F
diff --git a/Engine/Shaders/SkyLightingShared.usf b/Engine/Shaders/SkyLightingShared.usf
--- a/Engine/Shaders/SkyLightingShared.usf
+++ b/Engine/Shaders/SkyLightingShared.usf
@@ -31,9 +31,14 @@ void GetDistanceFieldAOSpecularOcclusion(float2 ScreenUV, float3 ReflectionVecto
 	BRANCH
 	if (ApplyBentNormalAO > 0)
 	{
-		float3 BentNormalAO = Texture2DSampleLevel(BentNormalAOTexture, BentNormalAOSampler, ScreenUV, 0).xyz;
+		float4 TextureValue = Texture2DSampleLevel(BentNormalAOTexture, BentNormalAOSampler, ScreenUV, 0);
+		float3 BentNormalAO = TextureValue.xyz;
 		float BentNormalLength = length(BentNormalAO);
 
+#define USE_CONE_TRACED_SPECULAR_OCCLUSION 0
+#if USE_CONE_TRACED_SPECULAR_OCCLUSION
+		IndirectSpecularOcclusion = TextureValue.w;
+#else
 		BRANCH
 		if (bTwoSidedFoliage)
 		{
@@ -50,6 +55,7 @@ void GetDistanceFieldAOSpecularOcclusion(float2 ScreenUV, float3 ReflectionVecto
 			IndirectSpecularOcclusion = lerp(0, IndirectSpecularOcclusion, saturate((UnoccludedAngle - .1f) / .2f));
 			IndirectSpecularOcclusion = lerp(IndirectSpecularOcclusion, 1, MinSkySpecularOcclusion);
 		}
+#endif
 
 		IndirectDiffuseOcclusion = lerp(BentNormalLength, 1, MinSkySpecularOcclusion);
 	}
