Resave Engine MaterialFunctions to clear bad object flag.  Fixes game materials that use these in some cases.

diff --git a/Engine/Shaders/ShaderVersion.usf b/Engine/Shaders/ShaderVersion.usf
--- a/Engine/Shaders/ShaderVersion.usf
+++ b/Engine/Shaders/ShaderVersion.usf
@@ -2,4 +2,4 @@
 // Each time the console command is executed it generates a new GUID. As this file is included
 // in common.usf (which should be included in any shader) it allows to invalidate the shader DDC.
 // 
-// GUID = A874AF26469A172B6598B0B299B3EE0E
\ No newline at end of file
+// GUID = C3536D1F9D6C4AE78A1260B4F4A38EFE
\ No newline at end of file
