Merging shader compile fix, CL 2945519.

diff --git a/Engine/Shaders/ShadowDepthPixelShader.usf b/Engine/Shaders/ShadowDepthPixelShader.usf
--- a/Engine/Shaders/ShadowDepthPixelShader.usf
+++ b/Engine/Shaders/ShadowDepthPixelShader.usf
@@ -21,7 +21,7 @@
 // to increase AO contribution. The AO calculation works best with closed geometry and small-medium scale objects
 #define LPV_GV_BOOST				1
 
-#if REFLECTIVE_SHADOW_MAP && !MATERIAL_BLOCK_GI && !NO_LPV
+#if REFLECTIVE_SHADOW_MAP && !NO_LPV
 	#include "LPVWriteVplCommon.usf"
 
 	#if RSM_TRANSMISSION
