Copied CL 2967537: Fix skylight reflections on translucent objects and SM4 shader compilation errors.

diff --git a/Engine/Shaders/BasePassPixelShader.usf b/Engine/Shaders/BasePassPixelShader.usf
--- a/Engine/Shaders/BasePassPixelShader.usf
+++ b/Engine/Shaders/BasePassPixelShader.usf
@@ -53,17 +53,18 @@ SamplerState	PrevSceneColorSampler;
 	/** Prenormalized capture of the scene that's closest to the object being rendered, used for reflection environment on translucency. */
 	TextureCubeArray ReflectionCubemap;
 	SamplerState ReflectionCubemapSampler;
-	int CubemapArrayIndex;
-	float4 ReflectionPositionAndRadius;
-	float ReflectionShape;
-	float4x4 BoxTransform;
-	float4 BoxScales;
-	float3 CaptureOffset;
+	int CubemapArrayIndex;	
 #else
 	TextureCube ReflectionCubemap;
 	SamplerState ReflectionCubemapSampler;
 #endif
 
+float4 ReflectionPositionAndRadius;
+float ReflectionShape;
+float4x4 BoxTransform;
+float4 BoxScales;
+float3 CaptureOffset;
+
 uint MortonCode( uint x )
 {
 	//x = (x ^ (x <<  8)) & 0x00ff00ff;
@@ -83,7 +84,7 @@ half3 GetImageBasedReflectionLighting(FMaterialPixelParameters MaterialParameter
 	half NoV = saturate(dot(N, V));
 
 	float RoughnessSq = Roughness * Roughness;
-	float AO = 1.0f;
+	float AO = 0.0f;
 	float SpecularOcclusion = GetSpecularOcclusion(NoV, RoughnessSq, AO);
 
 	float4 SpecularIBL = float4(0,0,0,SpecularOcclusion);
