Copying //UE4/Dev-Rendering to //UE4/Dev-Main (Source: //UE4/Dev-Rendering @ 2981742)

==========================
MAJOR FEATURES + CHANGES
==========================

Change 2967522 on 2016/05/05 by [REDACTED_PERSON]

	fixed UE-30208 Ensure taking high resolution screenshot
	renamed cvar from r.SaveUncompressedEXRFrame to r.SaveEXR.CompressionQuality and change meaning

Change 2967541 on 2016/05/05 by [REDACTED_PERSON]

	added cvar help entry to the editor menu

Change 2967544 on 2016/05/05 by [REDACTED_PERSON]

	* added ensures to track down UE-22073 [CrashReport] Crash when changing Editor Window Size - FDeferredShadingSceneRenderer::ClearTranslucentVolumeLighting()
	* fixed typo
	* improved r.DisplayInternals

Change 2967685 on 2016/05/05 by [REDACTED_PERSON]

	DR - Fix 2013 compile

Change 2967883 on 2016/05/05 by [REDACTED_PERSON]

	Making inherit rotation work for mesh particles

Change 2967948 on 2016/05/05 by [REDACTED_PERSON]

	MaterialParameterCollections now create default resources (uniform buffers) which are used when no valid FScene is present (eg DrawTile while exporting materials to lightmass)

Change 2967954 on 2016/05/05 by [REDACTED_PERSON]

	New refraction mode 'Pixel Normal Offset', which uses the difference between the per-pixel normal and the per-vertex normal to compute the refraction offset.  This is useful for refraction from large water surfaces without reading outside of screen bounds as much.

Change 2967970 on 2016/05/05 by [REDACTED_PERSON]

	Improvement for depth buffer collision; reduce tunneling by sampling midpoint instead of predicted position.

Change 2968125 on 2016/05/05 by [REDACTED_PERSON]

	Hopeful fix for mac compile error with enum

Change 2968126 on 2016/05/05 by [REDACTED_PERSON]

	Forward declare to fix compile error

Change 2969070 on 2016/05/06 by [REDACTED_PERSON]

	fixed Clang Compile error

Change 2969224 on 2016/05/06 by [REDACTED_PERSON]

	fixed Clang

Change 2969275 on 2016/05/06 by [REDACTED_PERSON]

	DR - Fix for simulation time from BP not matching material editor expression time

Change 2969413 on 2016/05/06 by [REDACTED_PERSON]

	Simple forward shading path
	* Support for a directional light in the base pass, which can be shadowed either by distance field shadowmaps or a single shadow factor from the Indirect Lighting Cache
	* Binding and outputting to GBuffers is disabled when using simple forward shading, along with dependencies of the GBuffers
	* Specular, normal maps and the 3rd SH band of skylighting are all hardcoded off.  Materials are forced fully rough.
	* Enabled with 'r.SimpleForwardShading 1', which can be changed at runtime
	* Added r.SupportSimpleForwardShading, which determines if the necessary shaders are compiled
	* Removed r.SimpleDynamicLighting, whose functionality (unshadowed directional light in the base pass) is a subset of simple forward shading, and was unused

Change 2969414 on 2016/05/06 by [REDACTED_PERSON]

	Added cvar 'r.Fog' for scalability

Change 2969415 on 2016/05/06 by [REDACTED_PERSON]

	Low end scalability setup for Paragon

Change 2969583 on 2016/05/06 by [REDACTED_PERSON]

	Effectively disabled the WRITES_PRECSHADOWFACTOR_TO_GBUFFER optimization, as all combinations use the GBuffer target
	* This fixes bSingleSampleShadowFromStationaryLights which was broken in cl 2863775 when selective GBuffer outputs was enabled

Change 2969597 on 2016/05/06 by [REDACTED_PERSON]

	Adding emitter tick times to Cascade viewport overlays

Change 2969599 on 2016/05/06 by [REDACTED_PERSON]

	Added r.Shadow.ForceSingleSampleShadowingFromStationary for scalability

Change 2969795 on 2016/05/06 by [REDACTED_PERSON]

	Made r.SimpleForwardShading a scalability cvar

Change 2969796 on 2016/05/06 by [REDACTED_PERSON]

	Enabled simple forward shading for the lowest shadow quality in Paragon
	* Moved Orion.Character.NonLocalHeroPreshadows to be under shadow quality
	* Enabling r.Shadow.ForceSingleSampleShadowingFromStationary when shadows are disabled through r.ShadowQuality

Change 2972338 on 2016/05/10 by [REDACTED_PERSON]

	Added ImageValidator V0.11 to Perforce (for automated screenshot comparison)

Change 2972450 on 2016/05/10 by [REDACTED_PERSON]

	fixed sorting issues, added column sort

Change 2972564 on 2016/05/10 by [REDACTED_PERSON]

	fixed PixelCountToFail

Change 2972572 on 2016/05/10 by [REDACTED_PERSON]

	added color

Change 2972667 on 2016/05/10 by [REDACTED_PERSON]

	Capsule shadows support an off-center projection - fixes stereo rendering
	Indirect lighting cache is updated for primitives with a visible indirect capsule shadow - fixes indirect shadows disappearing when the caster is off-screen in the Static skylight case

Change 2972687 on 2016/05/10 by [REDACTED_PERSON]

	split data from view

Change 2972713 on 2016/05/10 by [REDACTED_PERSON]

	background validation

Change 2972742 on 2016/05/10 by [REDACTED_PERSON]

	do not lock files in folder

Change 2972837 on 2016/05/10 by [REDACTED_PERSON]

	Comment for tiled reflection culling edge failure case

Change 2972891 on 2016/05/10 by [REDACTED_PERSON]

	Simple forward shading only supported on PC now

Change 2972927 on 2016/05/10 by [REDACTED_PERSON]

	fixed comparison, added directory watcher

Change 2972952 on 2016/05/10 by [REDACTED_PERSON]

	tooltips

Change 2972965 on 2016/05/10 by [REDACTED_PERSON]

	readme

Change 2972986 on 2016/05/10 by [REDACTED_PERSON]

	polish

Change 2973015 on 2016/05/10 by [REDACTED_PERSON]

	Clamp max cubemap resolutions for reflection capture component.

Change 2973067 on 2016/05/10 by [REDACTED_PERSON]

	Removed temporary streaming data from ULevel. Now the data goes directly in FStreamingManagerTexture.
	Removed deprecated classes and structures associated with old data.
	Uniformized dynamic primitives and static primitives streaming logic.
	Static primitive instances are now bound to level component and have visibility information.
	MeshTexCoordSize Debug now shows live update when tweaking "StreamingDistanceMultiplier"

Change 2973129 on 2016/05/10 by [REDACTED_PERSON]

	fixed crash (uninit var)

Change 2973167 on 2016/05/10 by [REDACTED_PERSON]

	Fixing mac compiler warnings; removed beam tangent from particle source and target, as it's treated separately anyway.

Change 2973654 on 2016/05/11 by [REDACTED_PERSON]

	FIxed bad resolve from main merge

Change 2973787 on 2016/05/11 by [REDACTED_PERSON]

	PR #2080: Added 7 additional blueprint nodes and assoicated functions to Particle System Compoenent (Contributed by [REDACTED_PERSON])

	#2080

Change 2973794 on 2016/05/11 by [REDACTED_PERSON]

	GitHub 1646 : Increase Wireframe Linear Color Value for Particles in Cascade Preview viewport

	#1646

Change 2973814 on 2016/05/11 by [REDACTED_PERSON]

	Fixed D3D warning.

Change 2973868 on 2016/05/11 by [REDACTED_PERSON]

	fixed compiler warning

Change 2974104 on 2016/05/11 by [REDACTED_PERSON]

	Optimized TextureStreaming ASync task by improving cache efficency.
	Removed ForcedFullyLoad texture sets and merged them with the other entries (to prevent checking in two data structures).

Change 2974163 on 2016/05/11 by [REDACTED_PERSON]

	Fixed Editor Build

Change 2974263 on 2016/05/11 by [REDACTED_PERSON]

	Fix for particle beam source/target selection

Change 2974313 on 2016/05/11 by [REDACTED_PERSON]

	nicer icon

Change 2974632 on 2016/05/11 by [REDACTED_PERSON]

	zoom (mouse wheel) and pan support

Change 2974651 on 2016/05/11 by [REDACTED_PERSON]

	UE-30565 Blend Radius does not override Priority for blendable post process material array
	#test:PC

Change 2974994 on 2016/05/12 by [REDACTED_PERSON]

	Fixed uninitialized members in texture streaming when updating dynamic primitives.

Change 2975514 on 2016/05/12 by [REDACTED_PERSON]

	Getting mesh particle transform down to pixel shader, making available in TransformVector/TransformPosition nodes

	#udn 288158, and requests from several artists

Change 2975970 on 2016/05/12 by [REDACTED_PERSON]

	Enable setting high precision normals/tangets as a part of static mesh import.

Change 2976134 on 2016/05/12 by [REDACTED_PERSON]

	Moved RHI prefix out of command list macros so you can find callers in a Find in Files for RHIBlendState, etc

Change 2976135 on 2016/05/12 by [REDACTED_PERSON]

	Added material property bAllowNegativeEmissiveColor which is useful when rendering to an offscreen render target

Change 2976139 on 2016/05/12 by [REDACTED_PERSON]

	Blueprint can now be used to draw to a TextureRenderTarget2D  with DrawMaterialToRenderTarget
	Canvas can now draw to a TextureRenderTarget2D in a Blueprint with BeginDrawCanvasToRenderTarget / EndDrawCanvasToRenderTarget

Change 2976199 on 2016/05/12 by [REDACTED_PERSON]

	Improved screen space ray cast intersection heuristic.

Change 2976203 on 2016/05/12 by [REDACTED_PERSON]

	Fixed hair NaNs

Change 2976214 on 2016/05/12 by [REDACTED_PERSON]

	Fix compiler error on Clang

Change 2976239 on 2016/05/12 by [REDACTED_PERSON]

	Improved hair shadows

Change 2976240 on 2016/05/12 by [REDACTED_PERSON]

	Made recompileshaders specific file work again. Kind of hacky but better than nothing

Change 2976290 on 2016/05/12 by [REDACTED_PERSON]

	Few more clang compiler fixes.

Change 2976517 on 2016/05/13 by [REDACTED_PERSON]

	Final set of fixes for compile breaks.

Change 2976535 on 2016/05/13 by [REDACTED_PERSON]

	Fix mono builds

Change 2977270 on 2016/05/13 by [REDACTED_PERSON]

	Incorrectly used 7 as max texcoords. Corrected to #define MAX_TEX_COORDS

Change 2980140 on 2016/05/17 by [REDACTED_PERSON]

	Merging //UE4/Dev-Main@2980127 to Dev-Rendering (//UE4/Dev-Rendering)

Change 2980211 on 2016/05/17 by [REDACTED_PERSON]

	EC: Build DDC in Dev-Rendering every 3 hours.

Change 2980876 on 2016/05/17 by [REDACTED_PERSON]

	Metal shader compile fix

diff --git a/Engine/Shaders/BasePassCommon.usf b/Engine/Shaders/BasePassCommon.usf
--- a/Engine/Shaders/BasePassCommon.usf
+++ b/Engine/Shaders/BasePassCommon.usf
@@ -13,13 +13,14 @@
 #define NEEDS_BASEPASS_FOGGING			(ENABLE_TRANSLUCENCY_VERTEX_FOG && ((MATERIALBLENDING_TRANSLUCENT || MATERIALBLENDING_ADDITIVE || MATERIALBLENDING_MODULATE) && !MATERIAL_USES_SCENE_COLOR_COPY))
 #define NEEDS_LIGHTMAP					(NEEDS_LIGHTMAP_COORDINATE)
 
-#define USES_GBUFFER						(FEATURE_LEVEL >= FEATURE_LEVEL_SM4 && (MATERIALBLENDING_SOLID || MATERIALBLENDING_MASKED) && !NO_GBUFFER)
+#define USES_GBUFFER						(FEATURE_LEVEL >= FEATURE_LEVEL_SM4 && (MATERIALBLENDING_SOLID || MATERIALBLENDING_MASKED) && !SIMPLE_FORWARD_SHADING)
 
 // Only some shader models actually need custom data.
 #define WRITES_CUSTOMDATA_TO_GBUFFER		(USES_GBUFFER && (!SELECTIVE_BASEPASS_OUTPUTS || MATERIAL_SHADINGMODEL_SUBSURFACE || MATERIAL_SHADINGMODEL_PREINTEGRATED_SKIN || MATERIAL_SHADINGMODEL_SUBSURFACE_PROFILE || MATERIAL_SHADINGMODEL_CLEAR_COAT || MATERIAL_SHADINGMODEL_TWOSIDED_FOLIAGE || MATERIAL_SHADINGMODEL_HAIR || MATERIAL_SHADINGMODEL_CLOTH || MATERIAL_SHADINGMODEL_EYE))
 
 // Based on GetPrecomputedShadowMasks()
-#define WRITES_PRECSHADOWFACTOR_TO_GBUFFER	(USES_GBUFFER && ALLOW_STATIC_LIGHTING && (!SELECTIVE_BASEPASS_OUTPUTS || (STATICLIGHTING_TEXTUREMASK && STATICLIGHTING_SIGNEDDISTANCEFIELD)))
+// Note: WRITES_PRECSHADOWFACTOR_TO_GBUFFER is currently disabled because we use the precomputed shadow factor GBuffer outside of STATICLIGHTING_TEXTUREMASK to store UseSingleSampleShadowFromStationaryLights
+#define WRITES_PRECSHADOWFACTOR_TO_GBUFFER	(USES_GBUFFER && ALLOW_STATIC_LIGHTING && (!SELECTIVE_BASEPASS_OUTPUTS || (1)))
 #define WRITES_PRECSHADOWFACTOR_ZERO		(USES_GBUFFER && ALLOW_STATIC_LIGHTING && (!SELECTIVE_BASEPASS_OUTPUTS || HQ_TEXTURE_LIGHTMAP || LQ_TEXTURE_LIGHTMAP))
 
 // If a primitive has static lighting, we assume it is not moving. If it is, it will be rerendered in an extra renderpass.
diff --git a/Engine/Shaders/BasePassForForwardShadingPixelShader.usf b/Engine/Shaders/BasePassForForwardShadingPixelShader.usf
--- a/Engine/Shaders/BasePassForForwardShadingPixelShader.usf
+++ b/Engine/Shaders/BasePassForForwardShadingPixelShader.usf
@@ -411,10 +411,6 @@ void Main(
 
 	#if !MATERIAL_SHADINGMODEL_UNLIT
 		half Shadow = GetPrimaryPrecomputedShadowMask(Interpolants).r;
-
-		#if SIMPLE_DYNAMIC_LIGHTING
-			Shadow = PrecomputedLightingBuffer.DirectionalLightShadowing;
-		#endif
 			
 		#if DIRECTIONAL_LIGHT_CSM
 		// Cascaded Shadow Map
diff --git a/Engine/Shaders/BasePassPixelShader.usf b/Engine/Shaders/BasePassPixelShader.usf
--- a/Engine/Shaders/BasePassPixelShader.usf
+++ b/Engine/Shaders/BasePassPixelShader.usf
@@ -32,7 +32,7 @@
 	 
 #endif
 
-#if MATERIAL_SHADINGMODEL_HAIR
+#if MATERIAL_SHADINGMODEL_HAIR || SIMPLE_FORWARD_DIRECTIONAL_LIGHT
 #include "ShadingModels.usf"
 #endif
 
@@ -45,6 +45,7 @@ SamplerState	PrevSceneColorSampler;
 	#define COMPILER_GLSL 0
 #endif
 
+#define FORCE_FULLY_ROUGH (SIMPLE_FORWARD_SHADING)
 #define EDITOR_ALPHA2COVERAGE (USE_EDITOR_COMPOSITING && FEATURE_LEVEL >= FEATURE_LEVEL_SM5 && !COMPILER_GLSL)
 
 #if TRANSLUCENCY_LIGHTING_SURFACE || TRANSLUCENCY_LIGHTING_SURFACE_PERPIXEL || MATERIAL_SHADINGMODEL_CLEAR_COAT
@@ -357,6 +358,12 @@ float3 GetTranslucencyLighting(
 	return InterpolatedLighting;
 }
 
+#if SIMPLE_FORWARD_SHADING
+	#define GetEffectiveSkySHDiffuse GetSkySHDiffuseSimple
+#else
+	#define GetEffectiveSkySHDiffuse GetSkySHDiffuse
+#endif
+
 /** Computes sky diffuse lighting, including precomputed shadowing. */
 void GetSkyLighting(float3 WorldNormal, float2 LightmapUV, out float3 OutDiffuseLighting, out float3 OutSubsurfaceLighting)
 {
@@ -411,13 +418,13 @@ void GetSkyLighting(float3 WorldNormal, float2 LightmapUV, out float3 OutDiffuse
 	#endif
 			
 	// Compute the preconvolved incoming lighting with the bent normal direction
-	float3 DiffuseLookup = GetSkySHDiffuse(SkyLightingNormal) * Frame.SkyLightColor.rgb;
+	float3 DiffuseLookup = GetEffectiveSkySHDiffuse(SkyLightingNormal) * Frame.SkyLightColor.rgb;
 
 	// Apply AO to the sky diffuse
 	OutDiffuseLighting += DiffuseLookup * (SkyVisibility * GeometryTerm);
 
 	#if MATERIAL_SHADINGMODEL_TWOSIDED_FOLIAGE
-		float3 BackfaceDiffuseLookup = GetSkySHDiffuse(-WorldNormal) * Frame.SkyLightColor.rgb;
+		float3 BackfaceDiffuseLookup = GetEffectiveSkySHDiffuse(-WorldNormal) * Frame.SkyLightColor.rgb;
 		OutSubsurfaceLighting += BackfaceDiffuseLookup * SkyVisibility;
 	#endif
 
@@ -567,6 +574,23 @@ void GetPrecomputedIndirectLightingAndSkyLight(
 	#endif
 }
 
+#if SIMPLE_FORWARD_DIRECTIONAL_LIGHT
+
+float3 GetSimpleForwardLightingDirectionalLight(float3 DiffuseColor, float3 SpecularColor, float Roughness, float3 WorldNormal, float3 CameraVector)
+{
+	float3 V = CameraVector;
+	float3 N = WorldNormal;
+	float3 L = Frame.DirectionalLightDirection;
+	float NoL = saturate( dot( N, L ) );
+
+	float3 LightColor = Frame.DirectionalLightColor.rgb * PI;
+	
+	// Not computing specular, material was forced fully rough
+	return LightColor * NoL * StandardShading(DiffuseColor, SpecularColor, Roughness.xxx, 1.0f, L, V, N, float2(1, 0));
+}
+
+#endif
+
 #if USE_EDITOR_COMPOSITING
 	bool bEnableEditorPrimitiveDepthTest;
 	int MSAASampleCount;
@@ -839,6 +863,12 @@ void Main(
 	}
 #endif
 
+	#if FORCE_FULLY_ROUGH
+		// Factors derived from EnvBRDFApprox( SpecularColor, 1, 1 ) == SpecularColor * 0.4524 - 0.0024
+		DiffuseColor += SpecularColor * 0.45;
+		SpecularColor = 0;
+	#endif
+
 	half Opacity = GetMaterialOpacity(PixelMaterialInputs);
 
 	FGBufferData GBuffer = (FGBufferData)0;
@@ -985,7 +1015,7 @@ void Main(
 			EffectiveDiffuseColor = PI * HairShading( GBuffer, L, V, N, 1, 0, 0.2, uint2(0,0) );
 		}
 		#endif
-			
+
 		float3 DiffuseIndirectLighting;
 		float3 SubsurfaceIndirectLighting;
 		GetPrecomputedIndirectLightingAndSkyLight(MaterialParameters, Interpolants, BasePassInterpolants, DiffuseDir, DiffuseIndirectLighting, SubsurfaceIndirectLighting, IndirectIrradiance);
@@ -993,17 +1023,24 @@ void Main(
 		Color += (DiffuseIndirectLighting * EffectiveDiffuseColor + SubsurfaceIndirectLighting * SubsurfaceColor) * MaterialAO;
 
 
-		#if SIMPLE_DYNAMIC_LIGHTING
-			// always unshadowed so BiasedNDotL is not needed
-			half Lambert = saturate(dot(MaterialParameters.WorldNormal, Frame.DirectionalLightDirection));
-			Color += DiffuseColor * Lambert * Frame.DirectionalLightColor.rgb;
+		#if SIMPLE_FORWARD_DIRECTIONAL_LIGHT
+
+			float3 DirectionalLighting = GetSimpleForwardLightingDirectionalLight(
+				EffectiveDiffuseColor,
+				SpecularColor,
+				Roughness,
+				MaterialParameters.WorldNormal, 
+				MaterialParameters.CameraVector);
+			
+			#if STATICLIGHTING_SIGNEDDISTANCEFIELD
+				half4 PrecomputedShadowMasks = GetPrecomputedShadowMasks(Interpolants);
+				DirectionalLighting *= PrecomputedShadowMasks.x;
+			#elif CACHED_POINT_INDIRECT_LIGHTING
+				DirectionalLighting *= PrecomputedLightingBuffer.DirectionalLightShadowing;
+			#endif
+			
+			Color += DirectionalLighting;
 
-			Color += GetMaterialHemisphereLightTransferFull(
-				DiffuseColor,
-				MaterialParameters,
-				Frame.UpperSkyColor.rgb,
-				Frame.LowerSkyColor.rgb
-				);
 		#endif
 	#endif
 
@@ -1014,7 +1051,7 @@ void Main(
 	#endif
 		
 	// Volume lighting for lit translucency
-	#if (MATERIAL_SHADINGMODEL_DEFAULT_LIT || MATERIAL_SHADINGMODEL_SUBSURFACE) && (MATERIALBLENDING_TRANSLUCENT || MATERIALBLENDING_ADDITIVE)
+	#if (MATERIAL_SHADINGMODEL_DEFAULT_LIT || MATERIAL_SHADINGMODEL_SUBSURFACE) && (MATERIALBLENDING_TRANSLUCENT || MATERIALBLENDING_ADDITIVE) && !SIMPLE_FORWARD_SHADING
 		Color += GetTranslucencyLighting(MaterialParameters, PixelMaterialInputs, BasePassInterpolants, GBuffer, DiffuseColor, SpecularColor, IndirectIrradiance);
 	#endif
 
diff --git a/Engine/Shaders/CapsuleShadowShaders.usf b/Engine/Shaders/CapsuleShadowShaders.usf
--- a/Engine/Shaders/CapsuleShadowShaders.usf
+++ b/Engine/Shaders/CapsuleShadowShaders.usf
@@ -187,24 +187,68 @@ void SetupTileCullingData(
 	float3 ViewTileMin2;
 	float3 ViewTileMax2;
 
-	float ExpandRadius = 0;
-
-	float2 TanViewFOV = float2(1 / View.ViewToClip[0][0], 1 / View.ViewToClip[1][1]);
-	// tan(FOV) = HalfUnitPlaneWidth / 1, so TanViewFOV * 2 is the size of the whole unit view plane
-	// We are operating on a subset of that defined by ScissorRectMinAndSize
-	float2 TileSize = TanViewFOV * 2 * ScissorRectMinAndSize.zw / ((float2)Frame.ViewSizeAndInvSize.xy * NumGroups);
-	float2 UnitPlaneMin = -TanViewFOV + TanViewFOV * 2 * (ScissorRectMinAndSize.xy - Frame.ViewRectMin.xy) * Frame.ViewSizeAndInvSize.zw;
-	float2 UnitPlaneTileMin = (GroupId.xy * TileSize + UnitPlaneMin) * float2(1, -1);
-	float2 UnitPlaneTileMax = ((GroupId.xy + 1) * TileSize + UnitPlaneMin) * float2(1, -1);
-
-	ViewTileMin.xy = min(MinTileZ * UnitPlaneTileMin, MaxTileZ2 * UnitPlaneTileMin) - ExpandRadius;
-	ViewTileMax.xy = max(MinTileZ * UnitPlaneTileMax, MaxTileZ2 * UnitPlaneTileMax) + ExpandRadius;
-	ViewTileMin.z = MinTileZ - ExpandRadius;
-	ViewTileMax.z = MaxTileZ2 + ExpandRadius;
-	ViewTileMin2.xy = min(MinTileZ2 * UnitPlaneTileMin, MaxTileZ * UnitPlaneTileMin) - ExpandRadius;
-	ViewTileMax2.xy = max(MinTileZ2 * UnitPlaneTileMax, MaxTileZ * UnitPlaneTileMax) + ExpandRadius;
-	ViewTileMin2.z = MinTileZ2 - ExpandRadius;
-	ViewTileMax2.z = MaxTileZ + ExpandRadius;
+	bool bCenteredProjection = abs(View.ViewToClip[1][0]) < .00001f && abs(View.ViewToClip[2][0]) < .00001f;
+
+	BRANCH
+	// Off center projection path uses 37 more asm instructions
+	if (bCenteredProjection)
+	{
+		float2 TanViewFOV = GetTanHalfFieldOfView();
+		// tan(FOV) = HalfUnitPlaneWidth / 1, so TanViewFOV * 2 is the size of the whole unit view plane
+		// We are operating on a subset of that defined by ScissorRectMinAndSize
+		float2 TileSize = TanViewFOV * 2 * ScissorRectMinAndSize.zw / ((float2)Frame.ViewSizeAndInvSize.xy * NumGroups);
+		float2 UnitPlaneMin = -TanViewFOV + TanViewFOV * 2 * (ScissorRectMinAndSize.xy - Frame.ViewRectMin.xy) * Frame.ViewSizeAndInvSize.zw;
+
+		float2 UnitPlaneTileMin = (GroupId.xy * TileSize + UnitPlaneMin) * float2(1, -1);
+		float2 UnitPlaneTileMax = ((GroupId.xy + 1) * TileSize + UnitPlaneMin) * float2(1, -1);
+
+		ViewTileMin.xy = min(MinTileZ * UnitPlaneTileMin, MaxTileZ2 * UnitPlaneTileMin);
+		ViewTileMax.xy = max(MinTileZ * UnitPlaneTileMax, MaxTileZ2 * UnitPlaneTileMax);
+		ViewTileMin.z = MinTileZ;
+		ViewTileMax.z = MaxTileZ2;
+		ViewTileMin2.xy = min(MinTileZ2 * UnitPlaneTileMin, MaxTileZ * UnitPlaneTileMin);
+		ViewTileMax2.xy = max(MinTileZ2 * UnitPlaneTileMax, MaxTileZ * UnitPlaneTileMax);
+		ViewTileMin2.z = MinTileZ2;
+		ViewTileMax2.z = MaxTileZ;
+	}
+	else
+	{
+		float2 TileSize = 2 * ScissorRectMinAndSize.zw / ((float2)Frame.ViewSizeAndInvSize.xy * NumGroups);
+		float2 UnitPlaneMin = -1 + 2 * (ScissorRectMinAndSize.xy - Frame.ViewRectMin.xy) * Frame.ViewSizeAndInvSize.zw;
+
+		float2 UnitPlaneTileMin = (GroupId.xy * TileSize + UnitPlaneMin) * float2(1, -1);
+		float2 UnitPlaneTileMax = ((GroupId.xy + 1) * TileSize + UnitPlaneMin) * float2(1, -1);
+
+		{
+			float MinTileDeviceZ = ConvertToDeviceZ(MinTileZ);
+			float4 MinDepthMinCorner = mul(float4(UnitPlaneTileMin.x, UnitPlaneTileMin.y, MinTileDeviceZ, 1), View.ClipToView);
+			float4 MinDepthMaxCorner = mul(float4(UnitPlaneTileMax.x, UnitPlaneTileMax.y, MinTileDeviceZ, 1), View.ClipToView);
+	
+			float MaxTileDeviceZ = ConvertToDeviceZ(MaxTileZ2);
+			float4 MaxDepthMinCorner = mul(float4(UnitPlaneTileMin.x, UnitPlaneTileMin.y, MaxTileDeviceZ, 1), View.ClipToView);
+			float4 MaxDepthMaxCorner = mul(float4(UnitPlaneTileMax.x, UnitPlaneTileMax.y, MaxTileDeviceZ, 1), View.ClipToView);
+
+			ViewTileMin.xy = min(MinDepthMinCorner.xy / MinDepthMinCorner.w, MaxDepthMinCorner.xy / MaxDepthMinCorner.w);
+			ViewTileMax.xy = max(MinDepthMaxCorner.xy / MinDepthMaxCorner.w, MaxDepthMaxCorner.xy / MaxDepthMaxCorner.w);
+			ViewTileMin.z = MinTileZ;
+			ViewTileMax.z = MaxTileZ2;
+		}
+
+		{
+			float MinTileDeviceZ = ConvertToDeviceZ(MinTileZ2);
+			float4 MinDepthMinCorner = mul(float4(UnitPlaneTileMin.x, UnitPlaneTileMin.y, MinTileDeviceZ, 1), View.ClipToView);
+			float4 MinDepthMaxCorner = mul(float4(UnitPlaneTileMax.x, UnitPlaneTileMax.y, MinTileDeviceZ, 1), View.ClipToView);
+	
+			float MaxTileDeviceZ = ConvertToDeviceZ(MaxTileZ);
+			float4 MaxDepthMinCorner = mul(float4(UnitPlaneTileMin.x, UnitPlaneTileMin.y, MaxTileDeviceZ, 1), View.ClipToView);
+			float4 MaxDepthMaxCorner = mul(float4(UnitPlaneTileMax.x, UnitPlaneTileMax.y, MaxTileDeviceZ, 1), View.ClipToView);
+
+			ViewTileMin2.xy = min(MinDepthMinCorner.xy / MinDepthMinCorner.w, MaxDepthMinCorner.xy / MaxDepthMinCorner.w);
+			ViewTileMax2.xy = max(MinDepthMaxCorner.xy / MinDepthMaxCorner.w, MaxDepthMaxCorner.xy / MaxDepthMaxCorner.w);
+			ViewTileMin2.z = MinTileZ2;
+			ViewTileMax2.z = MaxTileZ;
+		}
+	}
 
 	float3 ViewGroup0Center = (ViewTileMax + ViewTileMin) / 2;
 	TileCullingData0.BoundingSphere.xyz = mul(float4(ViewGroup0Center, 1), View.ViewToTranslatedWorld).xyz - View.PreViewTranslation;
@@ -392,32 +436,32 @@ float ShadowConeTraceAgainstCulledCapsuleShapes(
 		BRANCH
 		if (CapsuleOrientationAndLength.w > 0)
 		{
-		float3 VectorToCapsuleCenter = CapsuleCenterAndRadius.xyz - WorldRayStart;
-
-		// Closest point on line segment to ray
-		float3 L01 = CapsuleOrientationAndLength.xyz * CapsuleOrientationAndLength.w;
-		float3 L0 = VectorToCapsuleCenter - 0.5 * L01;
-		float3 L1 = VectorToCapsuleCenter + 0.5 * L01;
-
-		// The below is computing the shortest distance between capsule line segment and ray
-		float CapsuleOrientationProjectedOntoRay = dot(UnitRayDirection, L01);
-		// Vector that spans L01 perpendicular to the ray
-		float3 PerpendicularSpanningVector = CapsuleOrientationProjectedOntoRay * UnitRayDirection - L01;
-		// Length of PerpendicularSpanningVector using the right triangle formed by L01 and UnitRayDirection * CapsuleOrientationProjectedOntoRay
-		float PerpendicularDistance = Square(CapsuleOrientationAndLength.w) - CapsuleOrientationProjectedOntoRay * CapsuleOrientationProjectedOntoRay;
-		// Project the vector to a capsule endpoint onto the perpendicular spanning vector, normalized
-		float t = saturate(dot(L0, PerpendicularSpanningVector) / PerpendicularDistance);
-		// Compute the vector to the shadow sphere which best approximates the capsule's shadowing 
-		float3 VectorToShadowSphere = L0 + t * L01;
+			float3 VectorToCapsuleCenter = CapsuleCenterAndRadius.xyz - WorldRayStart;
+
+			// Closest point on line segment to ray
+			float3 L01 = CapsuleOrientationAndLength.xyz * CapsuleOrientationAndLength.w;
+			float3 L0 = VectorToCapsuleCenter - 0.5 * L01;
+			float3 L1 = VectorToCapsuleCenter + 0.5 * L01;
+
+			// The below is computing the shortest distance between capsule line segment and ray
+			float CapsuleOrientationProjectedOntoRay = dot(UnitRayDirection, L01);
+			// Vector that spans L01 perpendicular to the ray
+			float3 PerpendicularSpanningVector = CapsuleOrientationProjectedOntoRay * UnitRayDirection - L01;
+			// Length of PerpendicularSpanningVector using the right triangle formed by L01 and UnitRayDirection * CapsuleOrientationProjectedOntoRay
+			float PerpendicularDistance = Square(CapsuleOrientationAndLength.w) - CapsuleOrientationProjectedOntoRay * CapsuleOrientationProjectedOntoRay;
+			// Project the vector to a capsule endpoint onto the perpendicular spanning vector, normalized
+			float t = saturate(dot(L0, PerpendicularSpanningVector) / PerpendicularDistance);
+			// Compute the vector to the shadow sphere which best approximates the capsule's shadowing 
+			float3 VectorToShadowSphere = L0 + t * L01;
 
 			DistanceToShadowSphere = length(VectorToShadowSphere);
 			UnitVectorToShadowSphere = VectorToShadowSphere / DistanceToShadowSphere;
 		
-		// The above 'best shadow sphere' calculation doesn't take into account the project solid angle of the potential shadow spheres
-		// As a result, there's a discontinuity when the capsule and the ray point in nearly the same direction, where the far end of the capsule gets chosen
-		// Here we mitigate the effect by overriding the distance to shadow sphere if one of the capsule end points was closer
-		DistanceToShadowSphere = min(DistanceToShadowSphere, length(L0));
-		DistanceToShadowSphere = min(DistanceToShadowSphere, length(L1));
+			// The above 'best shadow sphere' calculation doesn't take into account the project solid angle of the potential shadow spheres
+			// As a result, there's a discontinuity when the capsule and the ray point in nearly the same direction, where the far end of the capsule gets chosen
+			// Here we mitigate the effect by overriding the distance to shadow sphere if one of the capsule end points was closer
+			DistanceToShadowSphere = min(DistanceToShadowSphere, length(L0));
+			DistanceToShadowSphere = min(DistanceToShadowSphere, length(L1));
 		}
 		else
 		{
@@ -535,6 +579,9 @@ void CapsuleShadowingCS(
 				TileCullingData0,
 				TileCullingData1);
 
+			NumTileIntersectingShapes = TileNumCapsules0 + TileNumCapsules1;
+		#else
+			NumTileIntersectingShapes = NumShadowCapsules;
 		#endif
 			
 		NumPixelIntersectingShapes += NumIntersectingCapsules;
@@ -548,8 +595,6 @@ void CapsuleShadowingCS(
 			NumIntersectingCapsules,
 			USE_CULLING ? true : false);
 
-		NumTileIntersectingShapes = TileNumCapsules0 + TileNumCapsules1;
-
 #if !APPLY_TO_BENT_NORMAL
 		if (all(GroupThreadId.xy == 0) && all(GroupId.xy < TileDimensions))
 		{
@@ -562,9 +607,9 @@ void CapsuleShadowingCS(
 	BRANCH
 	if (ReduceSelfShadowingIntensity > 0)
 	{
-	FGBufferData GBufferData = GetGBufferData(ScreenUV);
-	// Reduce self shadowing intensity
-	Visibility = lerp(Visibility, 1, HasDistanceFieldRepresentation(GBufferData) ? .8f : 0);
+		FGBufferData GBufferData = GetGBufferData(ScreenUV);
+		// Reduce self shadowing intensity
+		Visibility = lerp(Visibility, 1, HasDistanceFieldRepresentation(GBufferData) ? .8f : 0);
 	}
 #endif
 
diff --git a/Engine/Shaders/DeferredLightingCommon.usf b/Engine/Shaders/DeferredLightingCommon.usf
--- a/Engine/Shaders/DeferredLightingCommon.usf
+++ b/Engine/Shaders/DeferredLightingCommon.usf
@@ -492,7 +492,7 @@ float3 GetSimpleDynamicLighting(float3 WorldPosition, float3 CameraVector, FScre
 		// Apply SSAO to the direct lighting since we're not going to have any other shadowing
 		float Attenuation = DistanceAttenuation * ScreenSpaceData.AmbientOcclusion;
 
-		OutLighting += LightColor * (NoL * Attenuation) * SimpleShading( ScreenSpaceData.GBuffer, ScreenSpaceData.GBuffer.Roughness, L, V, N );
+		OutLighting += LightColor * (NoL * Attenuation) * SimpleShading( ScreenSpaceData.GBuffer.DiffuseColor, ScreenSpaceData.GBuffer.SpecularColor, ScreenSpaceData.GBuffer.Roughness, L, V, N );
 	}
 
 	return OutLighting;
diff --git a/Engine/Shaders/DistanceFieldShadowing.usf b/Engine/Shaders/DistanceFieldShadowing.usf
--- a/Engine/Shaders/DistanceFieldShadowing.usf
+++ b/Engine/Shaders/DistanceFieldShadowing.usf
@@ -394,7 +394,9 @@ void CullObjectsToTileWithGather(
 
 	float ExpandRadius = 0;
 
-	float2 TanViewFOV = float2(1 / View.ViewToClip[0][0], 1 / View.ViewToClip[1][1]);
+	// Note: this code is assuming a centered projection, aka no translation present in ViewToClip
+	// Stereo rendering uses an off center projection
+	float2 TanViewFOV = GetTanHalfFieldOfView();
 	// tan(FOV) = HalfUnitPlaneWidth / 1, so TanViewFOV * 2 is the size of the whole unit view plane
 	// We are operating on a subset of that defined by ScissorRectMinAndSize
 	float2 TileSize = TanViewFOV * 2 * ScissorRectMinAndSize.zw / ((float2)Frame.ViewSizeAndInvSize.xy * NumGroups);
diff --git a/Engine/Shaders/DistortAccumulatePS.usf b/Engine/Shaders/DistortAccumulatePS.usf
--- a/Engine/Shaders/DistortAccumulatePS.usf
+++ b/Engine/Shaders/DistortAccumulatePS.usf
@@ -42,10 +42,14 @@ void Main(
 	// X = IOR, Y = Refraction Bias
 	float2 Refraction = GetMaterialRefraction(PixelMaterialInputs);
 	 	
+#if REFRACTION_USE_PIXEL_NORMAL_OFFSET && MATERIAL_TANGENTSPACENORMAL
+	half3 ViewVertexNormal = TransformTangentVectorToView(MaterialParameters, float3(0, 0, 1));
+	float2 Distortion = (ViewVertexNormal.xy - ViewNormal.xy) * (Refraction.x - 1);
+#else
 	// we assume the camera is in air
 	float AirIOR = 1.0f;
-
 	float2 Distortion = ViewNormal.xy * (Refraction.x - AirIOR);
+#endif
 
 	//clip if the distortion distance (squared) is too small to be noticed
 	//this will result in a less expensive apply pass since the clipped pixels won't set stencil to pass
diff --git a/Engine/Shaders/DitheredTransitionStencil.usf b/Engine/Shaders/DitheredTransitionStencil.usf
--- a/Engine/Shaders/DitheredTransitionStencil.usf
+++ b/Engine/Shaders/DitheredTransitionStencil.usf
@@ -9,8 +9,7 @@
 float DitheredTransitionFactor;
 
 void Main(
-	FScreenVertexOutput Interpolators,
-	out float OutColor : SV_Target0)
+	FScreenVertexOutput Interpolators)
 {
 	float4 SvPosition = Interpolators.Position;
 	if (abs(DitheredTransitionFactor) > 0.001f)
@@ -22,6 +21,4 @@ void Main(
 			(DitheredTransitionFactor < RandomVal);
 		clip(RetVal - 0.001f);
 	}
-
-	OutColor = 0.f;
 }
diff --git a/Engine/Shaders/LightmapCommon.usf b/Engine/Shaders/LightmapCommon.usf
--- a/Engine/Shaders/LightmapCommon.usf
+++ b/Engine/Shaders/LightmapCommon.usf
@@ -120,6 +120,7 @@ float GetAOMaterialMask(float2 LightmapUV)
 // Used by deferred renderer only
 half4 GetPrecomputedShadowMasks(FVertexFactoryInterpolantsVSToPS Interpolants)
 {
+	// Note: WRITES_PRECSHADOWFACTOR_TO_GBUFFER and WRITES_PRECSHADOWFACTOR_ZERO have to match the logic here
 	#if STATICLIGHTING_TEXTUREMASK && STATICLIGHTING_SIGNEDDISTANCEFIELD
 
 		// Fetch the 4 channels of distance field data
diff --git a/Engine/Shaders/LocalVertexFactory.usf b/Engine/Shaders/LocalVertexFactory.usf
--- a/Engine/Shaders/LocalVertexFactory.usf
+++ b/Engine/Shaders/LocalVertexFactory.usf
@@ -95,12 +95,13 @@ struct FVertexFactoryInput
 	float4	Position	: ATTRIBUTE0;
 #if METAL_PROFILE
 	//@todo-rco: FIXME!
-	float4	TangentX	: ATTRIBUTE1;
-	float2	TangentZ	: ATTRIBUTE2;
+	float3	TangentX	: ATTRIBUTE1;
+	float4	TangentZ	: ATTRIBUTE2;
 	float4	Color		: ATTRIBUTE3;
 #else
-	half4	TangentX	: ATTRIBUTE1;
-	half2	TangentZ	: ATTRIBUTE2;
+	half3	TangentX	: ATTRIBUTE1;
+	// TangentZ.w contains sign of tangent basis determinant
+	half4	TangentZ	: ATTRIBUTE2;
 	half4	Color		: ATTRIBUTE3;
 #endif	// METAL_PROFILE
 
@@ -474,26 +475,26 @@ half3x3 CalcTangentToLocal(FVertexFactoryInput Input)
 {
 	half3x3 Result;
 	
-	half4 TangentX = TangentBias(Input.TangentX);
-	half3 TangentZ = OctToUnit(TangentBias(Input.TangentZ));
+	half3 TangentX = TangentBias(Input.TangentX);
+	half4 TangentZ = TangentBias(Input.TangentZ);
 
 #if USE_SPLINEDEFORM
 	// Make slice rotation matrix, and use that to transform tangents
 	half3x3 SliceRot = CalcSliceRot(dot(Input.Position.xyz, SplineMeshDir));
 
-	TangentX.xyz = mul(TangentX.xyz, SliceRot);
-	TangentZ = mul(TangentZ, SliceRot);
+	TangentX = mul(TangentX, SliceRot);
+	TangentZ.xyz = mul(TangentZ.xyz, SliceRot);
 #endif	// USE_SPLINEDEFORM
 
 	// derive the binormal by getting the cross product of the normal and tangent
-	half3 TangentY = cross(TangentZ, TangentX.xyz) * TangentX.w;
+	half3 TangentY = cross(TangentZ.xyz, TangentX) * TangentZ.w;
 	
 	// Recalculate TangentX off of the other two vectors
 	// This corrects quantization error since TangentX was passed in as a quantized vertex input
 	// The error shows up most in specular off of a mesh with a smoothed UV seam (normal is smooth, but tangents vary across the seam)
-	Result[0] = cross(TangentY, TangentZ) * TangentX.w;
+	Result[0] = cross(TangentY, TangentZ.xyz) * TangentZ.w;
 	Result[1] = TangentY;
-	Result[2] = TangentZ;
+	Result[2] = TangentZ.xyz;
 
 	return Result;
 }
@@ -516,7 +517,7 @@ FVertexFactoryIntermediates GetVertexFactoryIntermediates(FVertexFactoryInput In
 
 	Intermediates.TangentToLocal = CalcTangentToLocal(Input);
 	Intermediates.TangentToWorld = CalcTangentToWorld(Input,Intermediates.TangentToLocal);
-	Intermediates.TangentToWorldSign = TangentBias(Input.TangentX.w) * Primitive.LocalToWorldDeterminantSign;
+	Intermediates.TangentToWorldSign = TangentBias(Input.TangentZ.w) * Primitive.LocalToWorldDeterminantSign;
 
 	// Swizzle vertex color.
 	Intermediates.Color = Input.Color FCOLOR_COMPONENT_SWIZZLE;
diff --git a/Engine/Shaders/MaterialTemplate.usf b/Engine/Shaders/MaterialTemplate.usf
--- a/Engine/Shaders/MaterialTemplate.usf
+++ b/Engine/Shaders/MaterialTemplate.usf
@@ -83,6 +83,8 @@ struct FMaterialParticleParameters
 	half4 MacroUV;
 	/** Dynamic parameter used by particle systems. */
 	half4 DynamicParameter;
+	/** mesh particle orientation */
+	float4x4 LocalToWorld;
 
 #if USE_PARTICLE_SUBUVS
 	/** SubUV texture coordinates*/
@@ -1271,7 +1273,11 @@ half3 GetMaterialEmissiveRaw(FPixelMaterialInputs PixelMaterialInputs)
 
 half3 GetMaterialEmissive(FPixelMaterialInputs PixelMaterialInputs)
 {
-	return max(GetMaterialEmissiveRaw(PixelMaterialInputs), 0.0f);
+	half3 EmissiveColor = GetMaterialEmissiveRaw(PixelMaterialInputs);
+#if !MATERIAL_ALLOW_NEGATIVE_EMISSIVECOLOR
+	EmissiveColor = max(EmissiveColor, 0.0f);
+#endif
+	return EmissiveColor;
 }
 
 half3 GetMaterialEmissiveForCS(FMaterialPixelParameters Parameters)
@@ -1521,12 +1527,18 @@ void CalcPixelMaterialInputs(in out FMaterialPixelParameters Parameters, in out
 	// The Normal is a special case as it might have its own expressions and also be used to calculate other inputs, so perform the assignment here
 %s
 
-	// Note that here the Parameters.TangentNormal can be in world space or tangent space
-	Parameters.TangentNormal = GetMaterialNormal(Parameters, PixelMaterialInputs);
+	// Note that here MaterialNormal can be in world space or tangent space
+	float3 MaterialNormal = GetMaterialNormal(Parameters, PixelMaterialInputs);
 
 #if MATERIAL_TANGENTSPACENORMAL
+
+#if SIMPLE_FORWARD_SHADING
+	// Drop normal with simple forward shading for scalability
+	MaterialNormal = float3(0, 0, 1);
+#endif
+
 	// flip the normal for backfaces being rendered with a two-sided material
-	Parameters.TangentNormal *= Parameters.TwoSidedSign;
+	Parameters.TangentNormal = MaterialNormal * Parameters.TwoSidedSign;
 
 #if FEATURE_LEVEL >= FEATURE_LEVEL_SM4
 	// ES2 will rely on only the final normalize for performance
@@ -1538,7 +1550,7 @@ void CalcPixelMaterialInputs(in out FMaterialPixelParameters Parameters, in out
 	Parameters.WorldNormal = normalize(float3(TransformTangentVectorToWorld(Parameters.TangentToWorld, Parameters.TangentNormal)));
 #else
 	// Here we don't supoport two sided materials
-	Parameters.TangentNormal = Parameters.WorldNormal = normalize(Parameters.TangentNormal);
+	Parameters.TangentNormal = Parameters.WorldNormal = normalize(MaterialNormal);
 #endif
 
 	Parameters.ReflectionVector = ReflectionAboutCustomWorldNormal(Parameters, Parameters.WorldNormal, false);
@@ -1555,27 +1567,6 @@ void CalcPixelMaterialInputs(in out FMaterialPixelParameters Parameters, in out
 // This allows shader error line numbers after this point to be the same regardless of which material is being compiled
 #line %s
 
-float3 GetMaterialHemisphereLightTransferFull(float3 DiffuseColor,FMaterialPixelParameters Parameters, float3 UpperColor, float3 LowerColor)
-{
-	float3	UpperLighting = 0;
-	float3  LowerLighting = 0;
-
-#if (MATERIALBLENDING_TRANSLUCENT || MATERIALBLENDING_ADDITIVE) && TRANSLUCENCY_LIGHTING_VOLUMETRIC_NONDIRECTIONAL
-	UpperLighting = DiffuseColor;
-	LowerLighting = DiffuseColor;	
-#else
-	float3 SkyVector = float3(0, 0, 1);
-	float  NormalContribution  = dot(SkyVector,Parameters.WorldNormal);
-	float2 ContributionWeightsSqrt = float2(0.5, 0.5f) + float2(0.5f, -0.5f) * NormalContribution;
-	float2 ContributionWeights = ContributionWeightsSqrt * ContributionWeightsSqrt;
-	UpperLighting = DiffuseColor * ContributionWeights[0];
-	LowerLighting = DiffuseColor * ContributionWeights[1];
-#endif
-
-	return UpperLighting * UpperColor +
-	       LowerLighting * LowerColor;
-}
-
 void ClipLODTransition(float2 SvPosition, float DitherFactor)
 {
 	if (abs(DitherFactor) > .001)
@@ -1630,9 +1621,9 @@ void ClipLODTransition(float2 SvPosition)
 void GetMaterialClippingShadowDepth(FMaterialPixelParameters Parameters, FPixelMaterialInputs PixelMaterialInputs)
 {
 	ClipLODTransition(Parameters);
-	#if MATERIALBLENDING_MASKED && MATERIAL_DITHER_OPACITY_MASK
-		clip(GetMaterialMaskInputRaw(PixelMaterialInputs) - 1.0f / 255.0f);
-	#elif MATERIALBLENDING_MASKED
+	//#if MATERIALBLENDING_MASKED && MATERIAL_DITHER_OPACITY_MASK
+	//	clip(GetMaterialMaskInputRaw(PixelMaterialInputs) - 1.0f / 255.0f);
+	#if MATERIALBLENDING_MASKED
 		clip(GetMaterialMask(PixelMaterialInputs)); 
 	#elif MATERIALBLENDING_TRANSLUCENT
 		clip(GetMaterialOpacity(PixelMaterialInputs) - 1.0f / 255.0f);
diff --git a/Engine/Shaders/MaterialTexCoordScalesPixelShader.usf b/Engine/Shaders/MaterialTexCoordScalesPixelShader.usf
--- a/Engine/Shaders/MaterialTexCoordScalesPixelShader.usf
+++ b/Engine/Shaders/MaterialTexCoordScalesPixelShader.usf
@@ -127,6 +127,11 @@ void Main(
 			// The quad material used only has a valid TexCoord for UV 0
 			MaterialParameters.TexCoords[CoordinateIndex] = CoordinateIndex != CurrentCoordIndex ? 0 : DebugInputs.TexCoord01.xy;
 		}
+
+		// The scale must also be computed to match the actual coord inspected. That would be required when rendering the mesh in world space
+		// Although not required here as using the material quad only use texcoord0.
+		// Params.OneOverDDU = GetComponent(Params.OneOverDDU, CurrentCoordIndex).xxxx;
+		// Params.OneOverDDV = GetComponent(Params.OneOverDDV, CurrentCoordIndex).xxxx;
 #endif
 	}	
 
diff --git a/Engine/Shaders/MeshParticleVertexFactory.usf b/Engine/Shaders/MeshParticleVertexFactory.usf
--- a/Engine/Shaders/MeshParticleVertexFactory.usf
+++ b/Engine/Shaders/MeshParticleVertexFactory.usf
@@ -40,8 +40,9 @@ float4x4 PreviousLocalToWorld;
 struct FVertexFactoryInput
 {
 	float4	Position	: ATTRIBUTE0;
-	half4	TangentX	: ATTRIBUTE1;
-	half2	TangentZ	: ATTRIBUTE2;
+	half3	TangentX	: ATTRIBUTE1;
+	// TangentZ.w contains sign of tangent basis determinant
+	half4	TangentZ	: ATTRIBUTE2;
 	half4	VertexColor : ATTRIBUTE3;
 
 #if NUM_MATERIAL_TEXCOORDS_VERTEX
@@ -115,6 +116,8 @@ struct FVertexFactoryInterpolantsVSToPS
 
 	/** Relative alive time of the particle */
 	float ParticleTime : PARTICLE_TIME;
+
+	float4 LocalToWorld[4] : PARTICLE_TRANSFORM;
 #endif
 
 #if INSTANCED_STEREO
@@ -151,6 +154,8 @@ struct FVertexFactoryIntermediates
 	
 	/** Relative time. */
 	float RelativeTime;
+
+	float4 LocalToWorld[4];
 };
 
 float4x4 GetParticleTransform(FVertexFactoryInput Input)
@@ -251,6 +256,10 @@ FMaterialPixelParameters GetMaterialPixelParameters(FVertexFactoryInterpolantsVS
 	Result.Particle.PositionAndSize.w = 1;
 	Result.Particle.Velocity = Interpolants.ParticleVelocity;
 	Result.Particle.RelativeTime = Interpolants.ParticleTime;
+	Result.Particle.LocalToWorld[0] = Interpolants.LocalToWorld[0];
+	Result.Particle.LocalToWorld[1] = Interpolants.LocalToWorld[1];
+	Result.Particle.LocalToWorld[2] = Interpolants.LocalToWorld[2];
+	Result.Particle.LocalToWorld[3] = Interpolants.LocalToWorld[3];
 #endif
 
 	Result.Particle.MotionBlurFade = 1.0f;
@@ -278,6 +287,10 @@ FMaterialVertexParameters GetMaterialVertexParameters(FVertexFactoryInput Input,
 
 	float4x4 Transform = GetParticleTransform(Input);
 	Result.InstanceLocalToWorld = Transform;
+	Result.Particle.LocalToWorld[0] = Transform[0];
+	Result.Particle.LocalToWorld[1] = Transform[1];
+	Result.Particle.LocalToWorld[2] = Transform[2];
+	Result.Particle.LocalToWorld[3] = Transform[3];
 
 	Result.TangentToWorld = mul(TangentToLocal, (float3x3)Transform ); 
 
@@ -333,9 +346,14 @@ FVertexFactoryIntermediates GetVertexFactoryIntermediates(FVertexFactoryInput In
 	// Swizzle vertex color.
 	Intermediates.VertexColor = Input.VertexColor FCOLOR_COMPONENT_SWIZZLE;
 
-	// World position.
+	// World position and localtoworld transform.
 #if PARTICLE_MESH_INSTANCED
 	float3 ParticleWorldPosition = float3( Input.Transform1.w, Input.Transform2.w, Input.Transform3.w);
+	float4x4 LocalToWorldMat = float4x4(Input.Transform1, Input.Transform2, Input.Transform3, float4(0.0f, 0.0f, 0.0f, 1.0f));
+	Intermediates.LocalToWorld[0] = LocalToWorldMat[0];
+	Intermediates.LocalToWorld[1] = LocalToWorldMat[1];
+	Intermediates.LocalToWorld[2] = LocalToWorldMat[2];
+	Intermediates.LocalToWorld[3] = LocalToWorldMat[3];
 #else
 	float3 ParticleWorldPosition = float3( Transform1.w, Transform2.w, Transform3.w);
 #endif
@@ -398,21 +416,21 @@ FVertexFactoryIntermediates GetVertexFactoryIntermediates(FVertexFactoryInput In
 float3x3 VertexFactoryGetTangentToLocal( FVertexFactoryInput Input, FVertexFactoryIntermediates Intermediates )
 {
 	float3x3 Result=0;	
-	float4 TangentX = TangentBias(Input.TangentX);
-	float3 TangentZ = OctToUnit(TangentBias(Input.TangentZ));
+	float3 TangentX = TangentBias(Input.TangentX);
+	float4 TangentZ = TangentBias(Input.TangentZ);
 
 	// pass-thru the tangent
-	Result[0] = TangentX.xyz;
+	Result[0] = TangentX;
 	// pass-thru the normal
-	Result[2] = TangentZ;	
+	Result[2] = TangentZ.xyz;	
 
 	
 	// derive the binormal by getting the cross product of the normal and tangent
-	Result[1] = cross(Result[2], Result[0]) * TangentX.w;	
+	Result[1] = cross(Result[2], Result[0]) * TangentZ.w;	
 	// Recalculate TangentX off of the other two vectors
 	// This corrects quantization error since TangentX was passed in as a quantized vertex input
 	// The error shows up most in specular off of a mesh with a smoothed UV seam (normal is smooth, but tangents vary across the seam)
-	Result[0] = cross(Result[1], Result[2]) * TangentX.w;
+	Result[0] = cross(Result[1], Result[2]) * TangentZ.w;
 
 	return Result;
 }
@@ -442,7 +460,7 @@ void CalcTangentToWorld(FVertexFactoryInput Input, FVertexFactoryIntermediates I
 
 	// Normalize to remove scaling. Incorrect but faster than computing the inverse-transpose.
 	TangentToWorld0 = normalize(TangentToWorld[0]);
-	TangentToWorld2 = float4(normalize(TangentToWorld[2]), TangentBias(Input.TangentX.w) * Primitive.LocalToWorldDeterminantSign);
+	TangentToWorld2 = float4(normalize(TangentToWorld[2]), TangentBias(Input.TangentZ.w) * Primitive.LocalToWorldDeterminantSign);
 }
 
 FVertexFactoryInterpolantsVSToPS VertexFactoryGetInterpolantsVSToPS(FVertexFactoryInput Input, FVertexFactoryIntermediates Intermediates, FMaterialVertexParameters VertexParameters)
@@ -503,6 +521,7 @@ FVertexFactoryInterpolantsVSToPS VertexFactoryGetInterpolantsVSToPS(FVertexFacto
 	Interpolants.ParticlePosition = Intermediates.ParticleWorldPosition;
 	Interpolants.ParticleVelocity = Intermediates.ParticleVelocity;
 	Interpolants.ParticleTime = Intermediates.RelativeTime;
+	Interpolants.LocalToWorld = Intermediates.LocalToWorld;
 #endif
 
 #if INSTANCED_STEREO
diff --git a/Engine/Shaders/ParticleSimulationShader.usf b/Engine/Shaders/ParticleSimulationShader.usf
--- a/Engine/Shaders/ParticleSimulationShader.usf
+++ b/Engine/Shaders/ParticleSimulationShader.usf
@@ -10,6 +10,8 @@
 	#include "GlobalDistanceFieldShared.usf"
 #endif
 
+#define IMPROVED_DEPTH_BUFFER_COLLISION 1
+
 /*------------------------------------------------------------------------------
 	Shared declarations and functions.
 ------------------------------------------------------------------------------*/
@@ -279,7 +281,12 @@ void CollideWithDepthBuffer(
 	NewVelocity = InVelocity.xyz + Acceleration;
 
 	// Figure out where to sample the depth buffer.
+#if IMPROVED_DEPTH_BUFFER_COLLISION
+	float3 CollisionOffset = normalize(DeltaPosition*0.5f) * CollisionRadius;
+#else
 	float3 CollisionOffset = normalize(DeltaPosition) * CollisionRadius;
+#endif
+
 	float3 CollisionPosition = InPosition + CollisionOffset;
 	float4 SamplePosition = float4(CollisionPosition + View.PreViewTranslation,1);
 	float4 ClipPosition = mul(SamplePosition, View.TranslatedWorldToClip);
@@ -304,8 +311,41 @@ void CollideWithDepthBuffer(
 
 			// Compute the portion of velocity normal to the collision plane.
 			float VelocityDot = dot(CollisionPlane.xyz, DeltaPosition.xyz);
-			float InvVelocityDot = rcp(VelocityDot + 0.0001f); // Add a small amount to avoid division by zero.
 
+
+
+#if IMPROVED_DEPTH_BUFFER_COLLISION
+			// distance to the plane from current and predicted position
+			float d_back = ( dot(CollisionPlane.xyz, InPosition.xyz)+CollisionRadius - CollisionPlane.w );
+			float d_front = ( dot(CollisionPlane.xyz, NewPosition.xyz)-CollisionRadius - CollisionPlane.w );
+
+			if (d_back >= 0.0f && d_front <= 0.0f && VelocityDot<0.0f)
+			{
+				// Separate velocity in to the components perpendicular and tangent to the collision plane.
+				float3 PerpVelocity = dot(MidVelocity,CollisionPlane.xyz) * CollisionPlane.xyz;
+				float3 TanVelocity = MidVelocity - PerpVelocity;
+
+				// Compute the new velocity accounting for resilience and friction.
+				NewVelocity = Simulation.OneMinusFriction * TanVelocity - Resilience * PerpVelocity;
+
+				// If the particle lies approximately on the collision plane, don't jump to the point of collision.
+				d_back *= step(VelocityDot,-1);
+
+				// Integrate position taking the collision in to account.
+				float PositionAdjustment = ( dot(CollisionPlane.xyz, InPosition.xyz)-CollisionRadius - CollisionPlane.w );
+				NewPosition = InPosition + PositionAdjustment*CollisionPlane.xyz + NewVelocity * DeltaSeconds*0.1;
+
+				// Update the relative time. Usually this does nothing, but if the
+				// user has elected to kill the particle upon collision this will do
+				// so.
+				RelativeTime += Simulation.CollisionTimeBias;
+			}
+			else if (d_front < 0.0f && d_back < 0.0f)
+			{
+				RelativeTime = 1.1f;
+			}
+#else
+			float InvVelocityDot = rcp(VelocityDot + 0.0001f); // Add a small amount to avoid division by zero.
 			// Distance to the plane from the center of the particle.
 			float DistanceToPlane = dot(CollisionPlane.xyz, InPosition.xyz) - CollisionPlane.w;
 
@@ -313,7 +353,6 @@ void CollideWithDepthBuffer(
 			float t_back = -(DistanceToPlane + CollisionRadius) * InvVelocityDot;
 			float t_front = -(DistanceToPlane - CollisionRadius) * InvVelocityDot;
 
-			//if (t_back >= 0 && t_front <= 1 && DistanceToPlane >= 0)
 			if (step(0, t_back) * step(t_front, 1) * step(0, DistanceToPlane))
 			{
 				// Separate velocity in to the components perpendicular and tangent to the collision plane.
@@ -341,6 +380,7 @@ void CollideWithDepthBuffer(
 				// relative time to a value > 1.0.
 				RelativeTime = 1.1f;
 			}
+#endif
 		}
 	}
 }
diff --git a/Engine/Shaders/PostProcessAmbient.usf b/Engine/Shaders/PostProcessAmbient.usf
--- a/Engine/Shaders/PostProcessAmbient.usf
+++ b/Engine/Shaders/PostProcessAmbient.usf
@@ -171,7 +171,7 @@ float3 ImageBasedLightingMIS( FGBufferData GBuffer, float3 V, float3 N, uint2 Ra
 				}
 				else
 				{
-					Shading = StandardShading( GBuffer, GBuffer.Roughness, 1, L, V, N, 1 );
+					Shading = StandardShading( GBuffer.DiffuseColor, GBuffer.SpecularColor, GBuffer.Roughness, 1, L, V, N, 1 );
 				}
 			
 				Lighting += SampleColor * Shading * ( NoL * Weight );
diff --git a/Engine/Shaders/ReflectionEnvironmentComputeShaders.usf b/Engine/Shaders/ReflectionEnvironmentComputeShaders.usf
--- a/Engine/Shaders/ReflectionEnvironmentComputeShaders.usf
+++ b/Engine/Shaders/ReflectionEnvironmentComputeShaders.usf
@@ -194,6 +194,7 @@ void DoTileCulling(uint3 GroupId, uint ThreadIndex, float MinTileZ, float MaxTil
 	float4 ViewPos2 = mul(float4(ScreenPos1.x, ScreenPos0.y, ScreenPos1.z, 1), View.ClipToView);	ViewPos2.xyz /= ViewPos2.w;
 	float4 ViewPos3 = mul(float4(ScreenPos1.x, ScreenPos1.y, ScreenPos1.z, 1), View.ClipToView);	ViewPos3.xyz /= ViewPos3.w;
 	// front point
+	// Warning: this assumes a point at the near depth, which is not a valid assumption, will cause culling artifacts
 	float4 ViewPos4 = mul(float4(ScreenPos0.xy, ScreenPos0.z, 1), View.ClipToView);	ViewPos4.xyz /= ViewPos4.w;
 
 	float3 TileBoxMin = min(ViewPos4.xyz, min(ViewPos0.xyz, ViewPos3.xyz));
diff --git a/Engine/Shaders/ReflectionEnvironmentShared.usf b/Engine/Shaders/ReflectionEnvironmentShared.usf
--- a/Engine/Shaders/ReflectionEnvironmentShared.usf
+++ b/Engine/Shaders/ReflectionEnvironmentShared.usf
@@ -99,7 +99,7 @@ float3 GetSkySHDiffuse(float3 Normal)
 	float vC = NormalVector.x * NormalVector.x - NormalVector.y * NormalVector.y;
 	Intermediate2 = Frame.SkyIrradianceEnvironmentMap[6].xyz * vC;
 
-	// max to not get negative colors, maybe we can optimize to use saturate() and scale later or clamp this way: saturate(A/100)*100
+	// max to not get negative colors
 	return max(0, Intermediate0 + Intermediate1 + Intermediate2);
 }
 
@@ -117,7 +117,7 @@ float3 GetSkySHDiffuseSimple(float3 Normal)
 	Intermediate0.y = dot(Frame.SkyIrradianceEnvironmentMap[1], NormalVector);
 	Intermediate0.z = dot(Frame.SkyIrradianceEnvironmentMap[2], NormalVector);		
 
-	// max to not get negative colors, maybe we can optimize to use saturate() and scale later or clamp this way: saturate(A/100)*100
+	// max to not get negative colors
 	return max(0, Intermediate0);
 }
 
diff --git a/Engine/Shaders/ScreenSpaceRayCast.usf b/Engine/Shaders/ScreenSpaceRayCast.usf
--- a/Engine/Shaders/ScreenSpaceRayCast.usf
+++ b/Engine/Shaders/ScreenSpaceRayCast.usf
@@ -48,22 +48,25 @@ void RayCast(
 {
 	// TODO provide RayStartUVz
 	const float4 RayStartV = mul( float4( RayOriginTranslatedWorld, 1 ), View.TranslatedWorldToView );
-	const float4 RayDirV = mul(float4(RayDirection * SceneDepth, 0),  View.TranslatedWorldToView);
+	const float4 RayDirV = mul(float4(RayDirection * SceneDepth, 0), View.TranslatedWorldToView);
 	const float4 RayEndV = RayStartV + RayDirV;
 
 	const float RayEndRadiusV = length(RayDirV.xyz) * tan(ConeAngleWorld * 0.4);
 	const float2 RayEndBorderV = RayEndV.xy + RayEndRadiusV / sqrt(2.0);
 
 	const float4 RayStartClip	= mul( RayStartV, View.ViewToClip );
-	const float4 RayEndClip	= mul( RayEndV, View.ViewToClip );
+	const float4 RayEndClip		= mul( RayEndV, View.ViewToClip );
 	const float2 RayEndBorderClip = mul(float4(RayEndBorderV, RayEndV.zw), View.ViewToClip).xy;
 
 	const float3 RayStartScreen = RayStartClip.xyz / RayStartClip.w;
 	const float3 RayEndScreen = RayEndClip.xyz / RayEndClip.w;
 	const float2 RayEndBorderScreen = RayEndBorderClip.xy / RayEndClip.w;
-	
-	float3 RayStepScreen = (RayEndScreen - RayStartScreen);
-	float2 RayStepRadiusScreen = (RayEndBorderScreen - RayEndScreen.xy);
+
+	float4 RayDepthClip	= RayStartClip + mul( float4( 0, 0, SceneDepth, 0 ), View.ViewToClip );
+	float3 RayDepthScreen = RayDepthClip.xyz / RayDepthClip.w;
+
+	float3 RayStepScreen = RayEndScreen - RayStartScreen;
+	float2 RayStepRadiusScreen = RayEndBorderScreen - RayEndScreen.xy;
 
 	{
 		// Computes the scale down factor for RayStepScreen required to fit on the X and Y axis in order to clip it in the viewport
@@ -92,8 +95,7 @@ void RayCast(
 	
 	const float Step = 1.0 / NumSteps;
 
-	// *2 to get less morie pattern in extreme cases, larger values make object appear not grounded in reflections
-	const float CompareTolerance = abs( RayStepUVz.z ) * Step * 2;
+	const float CompareTolerance = abs( RayDepthScreen.z - RayStartScreen.z ) * Step * 4;
 
 	// avoid bugs with early returns inside of loops on certain platform compilers.
 	float4 Result = float4( 0, 0, 0, 1 );
diff --git a/Engine/Shaders/ShadingModels.usf b/Engine/Shaders/ShadingModels.usf
--- a/Engine/Shaders/ShadingModels.usf
+++ b/Engine/Shaders/ShadingModels.usf
@@ -44,7 +44,7 @@ void StandardShadingPerLight( Shared, float3 L, float3 V, half3 N )
 #endif
 
 // @param DiffSpecMask .r: diffuse, .g:specular e.g. float2(1,1) for both, float2(1,0) for diffuse only
-float3 StandardShading( FGBufferData GBuffer, float3 LobeRoughness, float3 LobeEnergy, float3 L, float3 V, half3 N, float2 DiffSpecMask )
+float3 StandardShading( float3 DiffuseColor, float3 SpecularColor, float3 LobeRoughness, float3 LobeEnergy, float3 L, float3 V, half3 N, float2 DiffSpecMask )
 {
 	float3 H = normalize(V + L);
 	float NoL = saturate( dot(N, L) );
@@ -55,16 +55,16 @@ float3 StandardShading( FGBufferData GBuffer, float3 LobeRoughness, float3 LobeE
 	// Generalized microfacet specular
 	float D = D_GGX( LobeRoughness[1], NoH ) * LobeEnergy[1];
 	float Vis = Vis_SmithJointApprox( LobeRoughness[1], NoV, NoL );
-	float3 F = F_Schlick( GBuffer.SpecularColor, VoH );
+	float3 F = F_Schlick( SpecularColor, VoH );
 
-	float3 Diffuse = Diffuse_Lambert( GBuffer.DiffuseColor );
-	//float3 Diffuse = Diffuse_Burley( GBuffer.DiffuseColor, LobeRoughness[1], NoV, NoL, VoH );
-	//float3 Diffuse = Diffuse_OrenNayar( GBuffer.DiffuseColor, LobeRoughness[1], NoV, NoL, VoH );
+	float3 Diffuse = Diffuse_Lambert( DiffuseColor );
+	//float3 Diffuse = Diffuse_Burley( DiffuseColor, LobeRoughness[1], NoV, NoL, VoH );
+	//float3 Diffuse = Diffuse_OrenNayar( DiffuseColor, LobeRoughness[1], NoV, NoL, VoH );
 	
 	return Diffuse * (LobeEnergy[2] * DiffSpecMask.r) + (D * Vis * DiffSpecMask.g) * F;
 }
 
-float3 SimpleShading( FGBufferData GBuffer, float Roughness, float3 L, float3 V, half3 N )
+float3 SimpleShading( float3 DiffuseColor, float3 SpecularColor, float Roughness, float3 L, float3 V, half3 N )
 {
 	float3 H = normalize(V + L);
 	float NoH = saturate( dot(N, H) );
@@ -72,9 +72,9 @@ float3 SimpleShading( FGBufferData GBuffer, float Roughness, float3 L, float3 V,
 	// Generalized microfacet specular
 	float D = D_GGX( Roughness, NoH );
 	float Vis = Vis_Implicit();
-	float3 F = F_None( GBuffer.SpecularColor );
+	float3 F = F_None( SpecularColor );
 
-	return Diffuse_Lambert( GBuffer.DiffuseColor ) + (D * Vis) * F;
+	return Diffuse_Lambert( DiffuseColor ) + (D * Vis) * F;
 }
 
 float3 ClearCoatShading( FGBufferData GBuffer, float3 LobeRoughness, float3 LobeEnergy, float3 L, float3 V, half3 N )
@@ -500,11 +500,12 @@ float3 HairShadingRef( FGBufferData GBuffer, float3 L, float3 V, half3 N, uint2
 
 	HairTemp.n_prime = sqrt( n*n - 1 + Pow2( HairTemp.CosThetaD ) ) / HairTemp.CosThetaD;
 
+	float Shift = 0.035;
 	float Alpha[] =
 	{
-		0.05,
-		-0.05 * 0.5,
-		-0.05 * 1.5,
+		-Shift * 2,
+		Shift,
+		Shift * 4,
 	};
 	float B[] =
 	{
@@ -593,6 +594,7 @@ float3 HairShading( FGBufferData GBuffer, float3 L, float3 V, half3 N, float Sha
 	float3 S = 0;
 
 	// R
+	if(1)
 	{
 		const float sa = sin( Alpha[0] );
 		const float ca = cos( Alpha[0] );
@@ -605,6 +607,7 @@ float3 HairShading( FGBufferData GBuffer, float3 L, float3 V, half3 N, float Sha
 	}
 
 	// TT
+	if(1)
 	{
 		float Mp = Hair_g( B[1], SinThetaL + SinThetaV - Alpha[1] );
 
@@ -631,6 +634,7 @@ float3 HairShading( FGBufferData GBuffer, float3 L, float3 V, half3 N, float Sha
 	}
 
 	// TRT
+	if(1)
 	{
 		float Mp = Hair_g( B[2], SinThetaL + SinThetaV - Alpha[2] );
 		
@@ -648,17 +652,22 @@ float3 HairShading( FGBufferData GBuffer, float3 L, float3 V, half3 N, float Sha
 	}
 #endif
 
-	float3 FakeNormal = normalize( V - N * dot(V,N) );
-	//N = normalize( DiffuseN + FakeNormal * 2 );
-	N = FakeNormal;
+	if(1)
+	{
+		float3 FakeNormal = normalize( V - N * dot(V,N) );
+		//N = normalize( DiffuseN + FakeNormal * 2 );
+		N = FakeNormal;
+
+		// Hack approximation for multiple scattering.
+		float Wrap = 1;
+		float NoL = saturate( ( dot(N, L) + Wrap ) / Square( 1 + Wrap ) );
+		float DiffuseScatter = (1 / PI) * NoL * GBuffer.Metallic;
+		float Luma = Luminance( GBuffer.BaseColor );
+		float3 ScatterTint = pow( GBuffer.BaseColor / Luma, 1 - Shadow );
+		S += sqrt( GBuffer.BaseColor ) * DiffuseScatter * ScatterTint;
+	}
 
-	// Hack approximation for multiple scattering.
-	float Wrap = 1;
-	float NoL = saturate( ( dot(N, L) + Wrap ) / Square( 1 + Wrap ) );
-	float DiffuseScatter = (1 / PI) * NoL * GBuffer.Metallic;
-	float Luma = Luminance( GBuffer.BaseColor );
-	float3 ScatterTint = pow( GBuffer.BaseColor / Luma, 1 - Shadow );
-	S += sqrt( GBuffer.BaseColor ) * DiffuseScatter * ScatterTint;
+	S = -min(-S, 0.0);
 
 	return S;
 }
@@ -771,7 +780,7 @@ float3 SurfaceShading( FGBufferData GBuffer, float3 LobeRoughness, float3 LobeEn
 		case SHADINGMODELID_PREINTEGRATED_SKIN:
 		case SHADINGMODELID_SUBSURFACE_PROFILE:
 		case SHADINGMODELID_TWOSIDED_FOLIAGE:
-			return StandardShading( GBuffer, LobeRoughness, LobeEnergy, L, V, N, DiffSpecMask);
+			return StandardShading( GBuffer.DiffuseColor, GBuffer.SpecularColor, LobeRoughness, LobeEnergy, L, V, N, DiffSpecMask);
 		case SHADINGMODELID_CLEAR_COAT:
 			return ClearCoatShading( GBuffer, LobeRoughness, LobeEnergy, L, V, N );
 		case SHADINGMODELID_CLOTH:
diff --git a/Engine/Source/Developer/MeshUtilities/Private/MeshUtilities.cpp b/Engine/Source/Developer/MeshUtilities/Private/MeshUtilities.cpp
--- a/Engine/Source/Developer/MeshUtilities/Private/MeshUtilities.cpp
+++ b/Engine/Source/Developer/MeshUtilities/Private/MeshUtilities.cpp
@@ -391,6 +391,7 @@ protected:
 		SrcModel->BuildSettings.bRecomputeNormals = false;
 		SrcModel->BuildSettings.bRecomputeTangents = false;
 		SrcModel->BuildSettings.bRemoveDegenerates = false;
+		SrcModel->BuildSettings.bUseHighPrecisionTangentBasis = false;
 		SrcModel->BuildSettings.bUseFullPrecisionUVs = false;
 		SrcModel->RawMeshBulkData->SaveRawMesh(Data->RawMesh);
 
@@ -3226,6 +3227,7 @@ public:
 
 			// Initialize the vertex buffer.
 			int32 NumTexCoords = ComputeNumTexCoords(RawMesh, MAX_STATIC_TEXCOORDS);
+			LODModel.VertexBuffer.SetUseHighPrecisionTangentBasis(LODBuildSettings[LODIndex].bUseHighPrecisionTangentBasis);
 			LODModel.VertexBuffer.SetUseFullPrecisionUVs(LODBuildSettings[LODIndex].bUseFullPrecisionUVs);
 			LODModel.VertexBuffer.Init(Vertices, NumTexCoords);
 			LODModel.PositionVertexBuffer.Init(Vertices);
@@ -7320,6 +7322,7 @@ void FMeshUtilities::MergeStaticMeshComponents(const TArray<UStaticMeshComponent
 				SrcModel->BuildSettings.bRecomputeNormals = false;
 				SrcModel->BuildSettings.bRecomputeTangents = false;
 				SrcModel->BuildSettings.bRemoveDegenerates = false;
+				SrcModel->BuildSettings.bUseHighPrecisionTangentBasis = false;
 				SrcModel->BuildSettings.bUseFullPrecisionUVs = false;
 				SrcModel->BuildSettings.bGenerateLightmapUVs = InSettings.bGenerateLightMapUV;
 				SrcModel->BuildSettings.MinLightmapResolution = InSettings.TargetLightMapResolution;
diff --git a/Engine/Source/Developer/VulkanShaderFormat/Private/VulkanGlslangBridge.cpp b/Engine/Source/Developer/VulkanShaderFormat/Private/VulkanGlslangBridge.cpp
--- a/Engine/Source/Developer/VulkanShaderFormat/Private/VulkanGlslangBridge.cpp
+++ b/Engine/Source/Developer/VulkanShaderFormat/Private/VulkanGlslangBridge.cpp
@@ -4,6 +4,13 @@
 
 #include "VulkanShaderFormat.h"
 
+#if _MSC_VER == 1800
+#pragma warning(push)
+// warning C4510: 'glslang::TShader::Includer::IncludeResult' : default constructor could not be generated
+//error C4610: struct 'glslang::TShader::Includer::IncludeResult' can never be instantiated - user defined constructor required
+#pragma warning(disable:4510)
+#pragma warning(disable:4610)
+#endif // _MSC_VER == 1800
 #include "glslang/Include/ShHandle.h"
 #include "glslang/Include/revision.h"
 #include "glslang/Public/ShaderLang.h"
@@ -11,6 +18,9 @@
 #include "SPIRV/GLSL.std.450.h"
 #include "SPIRV/doc.h"
 #include "SPIRV/disassemble.h"
+#if _MSC_VER == 1800
+#pragma warning(pop)
+#endif // _MSC_VER == 1800
 
 #include "hlslcc.h"
 
diff --git a/Engine/Source/Editor/Cascade/Private/Cascade.cpp b/Engine/Source/Editor/Cascade/Private/Cascade.cpp
--- a/Engine/Source/Editor/Cascade/Private/Cascade.cpp
+++ b/Engine/Source/Editor/Cascade/Private/Cascade.cpp
@@ -2162,6 +2162,12 @@ void FCascade::BindCommands()
 		FCanExecuteAction(),
 		FIsActionChecked::CreateSP(this, &FCascade::IsViewSystemCompletedChecked));
 
+	ToolkitCommands->MapAction(
+		Commands.View_EmitterTickTimes,
+		FExecuteAction::CreateSP(this, &FCascade::OnViewEmitterTickTimes),
+		FCanExecuteAction(),
+		FIsActionChecked::CreateSP(this, &FCascade::IsViewEmitterTickTimesChecked));
+
 	ToolkitCommands->MapAction(
 		Commands.ToggleGeometry,
 		FExecuteAction::CreateSP(this, &FCascade::OnViewGeometry),
@@ -3586,6 +3592,17 @@ bool FCascade::IsDrawOptionEnabled(int32 Element) const
 	}
 }
 
+void FCascade::OnViewEmitterTickTimes()
+{
+	ToggleDrawOption(FCascadeEdPreviewViewportClient::EmitterTickTimes);
+}
+
+bool FCascade::IsViewEmitterTickTimesChecked() const
+{
+	return IsDrawOptionEnabled(FCascadeEdPreviewViewportClient::EmitterTickTimes);
+}
+
+
 void FCascade::OnViewOriginAxis()
 {
 	ToggleDrawOption(FCascadeEdPreviewViewportClient::OriginAxis);
diff --git a/Engine/Source/Editor/Cascade/Private/Cascade.h b/Engine/Source/Editor/Cascade/Private/Cascade.h
--- a/Engine/Source/Editor/Cascade/Private/Cascade.h
+++ b/Engine/Source/Editor/Cascade/Private/Cascade.h
@@ -308,6 +308,8 @@ private:
 	bool IsViewParticleMemoryChecked() const;
 	void OnViewSystemCompleted();
 	bool IsViewSystemCompletedChecked() const;
+	void OnViewEmitterTickTimes();
+	bool IsViewEmitterTickTimesChecked() const;
 	void OnViewGeometry();
 	bool IsViewGeometryChecked() const;
 	void OnViewGeometryProperties();
diff --git a/Engine/Source/Editor/Cascade/Private/CascadeActions.cpp b/Engine/Source/Editor/Cascade/Private/CascadeActions.cpp
--- a/Engine/Source/Editor/Cascade/Private/CascadeActions.cpp
+++ b/Engine/Source/Editor/Cascade/Private/CascadeActions.cpp
@@ -35,6 +35,7 @@ void FCascadeCommands::RegisterCommands()
 	UI_COMMAND(View_ParticleTimes, "Particle Times", "Display Particle Times", EUserInterfaceActionType::ToggleButton, FInputChord());
 	UI_COMMAND(View_ParticleMemory, "Particle Memory", "Display Particle Memory", EUserInterfaceActionType::ToggleButton, FInputChord());
 	UI_COMMAND(View_SystemCompleted, "System Completed", "Display 'Completed' when the particle system has finished playing and has not been reset (only for non-looping systems)", EUserInterfaceActionType::ToggleButton, FInputChord());
+	UI_COMMAND(View_EmitterTickTimes, "Emitter Tick Times", "Show the tick duration for each emitter", EUserInterfaceActionType::ToggleButton, FInputChord());
 	UI_COMMAND(ViewMode_Wireframe, "Wireframe", "Select Wireframe Render Mode", EUserInterfaceActionType::RadioButton, FInputChord());
 	UI_COMMAND(ViewMode_Unlit, "Unlit", "Select Unlit Render Mode", EUserInterfaceActionType::RadioButton, FInputChord());
 	UI_COMMAND(ViewMode_Lit, "Lit", "Select Lit Render Mode", EUserInterfaceActionType::RadioButton, FInputChord());
diff --git a/Engine/Source/Editor/Cascade/Private/CascadeActions.h b/Engine/Source/Editor/Cascade/Private/CascadeActions.h
--- a/Engine/Source/Editor/Cascade/Private/CascadeActions.h
+++ b/Engine/Source/Editor/Cascade/Private/CascadeActions.h
@@ -24,6 +24,7 @@ public:
 	TSharedPtr<FUICommandInfo> View_ParticleTimes;
 	TSharedPtr<FUICommandInfo> View_ParticleMemory;
 	TSharedPtr<FUICommandInfo> View_SystemCompleted;
+	TSharedPtr<FUICommandInfo> View_EmitterTickTimes;
 	TSharedPtr<FUICommandInfo> ToggleGeometry;
 	TSharedPtr<FUICommandInfo> ToggleGeometry_Properties;
 	TSharedPtr<FUICommandInfo> ToggleLocalVectorFields;
diff --git a/Engine/Source/Editor/Cascade/Private/CascadePreviewViewportClient.cpp b/Engine/Source/Editor/Cascade/Private/CascadePreviewViewportClient.cpp
--- a/Engine/Source/Editor/Cascade/Private/CascadePreviewViewportClient.cpp
+++ b/Engine/Source/Editor/Cascade/Private/CascadePreviewViewportClient.cpp
@@ -295,6 +295,12 @@ void FCascadeEdPreviewViewportClient::Draw(FViewport* InViewport, FCanvas* Canva
 					continue;
 				}
 
+				if (GetDrawElement(EmitterTickTimes))
+				{
+					FString StatLine = Instance->SpriteTemplate->EmitterName.ToString() + " tick: " + FString::Printf(TEXT("%.3f ms"), Instance->LastTickDurationMs);
+					Canvas->DrawShadowedString(0, i*16+25, *StatLine, GEngine->GetTinyFont(), FLinearColor::Gray);
+				}
+
 				strOutput = TEXT("");
 				if (Instance->SpriteTemplate->EmitterRenderMode != ERM_None)
 				{
diff --git a/Engine/Source/Editor/Cascade/Private/CascadePreviewViewportClient.h b/Engine/Source/Editor/Cascade/Private/CascadePreviewViewportClient.h
--- a/Engine/Source/Editor/Cascade/Private/CascadePreviewViewportClient.h
+++ b/Engine/Source/Editor/Cascade/Private/CascadePreviewViewportClient.h
@@ -50,7 +50,8 @@ public:
 		WireSphere = 0x040,
 		OriginAxis = 0x080,
 		Orbit = 0x100,
-		ParticleSystemCompleted = 0x200
+		ParticleSystemCompleted = 0x200,
+		EmitterTickTimes = 0x400
 	};
 
 	/** Accessors */
diff --git a/Engine/Source/Editor/Cascade/Private/SCascadePreviewToolbar.cpp b/Engine/Source/Editor/Cascade/Private/SCascadePreviewToolbar.cpp
--- a/Engine/Source/Editor/Cascade/Private/SCascadePreviewToolbar.cpp
+++ b/Engine/Source/Editor/Cascade/Private/SCascadePreviewToolbar.cpp
@@ -70,6 +70,7 @@ TSharedRef<SWidget> SCascadePreviewViewportToolBar::GenerateViewMenu() const
 				Menu.AddMenuEntry(FCascadeCommands::Get().View_ParticleTimes);
 				Menu.AddMenuEntry(FCascadeCommands::Get().View_ParticleMemory);
 				Menu.AddMenuEntry(FCascadeCommands::Get().View_SystemCompleted);
+				Menu.AddMenuEntry(FCascadeCommands::Get().View_EmitterTickTimes);
 			}
 			Menu.EndSection();
 		}
diff --git a/Engine/Source/Editor/DestructibleMeshEditor/Private/SDestructibleMeshEditorViewport.cpp b/Engine/Source/Editor/DestructibleMeshEditor/Private/SDestructibleMeshEditorViewport.cpp
--- a/Engine/Source/Editor/DestructibleMeshEditor/Private/SDestructibleMeshEditorViewport.cpp
+++ b/Engine/Source/Editor/DestructibleMeshEditor/Private/SDestructibleMeshEditorViewport.cpp
@@ -109,7 +109,6 @@ void FDestructibleMeshEditorViewportClient::UpdateLighting()
 	const UDestructableMeshEditorSettings* Options = GetDefault<UDestructableMeshEditorSettings>();
 
 	PreviewScene->SetLightDirection(Options->AnimPreviewLightingDirection);
-	PreviewScene->GetScene()->UpdateDynamicSkyLight(Options->AnimPreviewSkyBrightness * FLinearColor(Options->AnimPreviewSkyColor), Options->AnimPreviewSkyBrightness * FLinearColor(Options->AnimPreviewFloorColor));
 	PreviewScene->SetLightColor(Options->AnimPreviewDirectionalColor);
 	PreviewScene->SetLightBrightness(Options->AnimPreviewLightBrightness);
 }
diff --git a/Engine/Source/Editor/DetailCustomizations/Private/AnimMontageSegmentDetails.cpp b/Engine/Source/Editor/DetailCustomizations/Private/AnimMontageSegmentDetails.cpp
--- a/Engine/Source/Editor/DetailCustomizations/Private/AnimMontageSegmentDetails.cpp
+++ b/Engine/Source/Editor/DetailCustomizations/Private/AnimMontageSegmentDetails.cpp
@@ -36,7 +36,6 @@ void FAnimationSegmentViewportClient::UpdateLighting()
 	const UDestructableMeshEditorSettings* Options = GetDefault<UDestructableMeshEditorSettings>();
 
 	PreviewScene->SetLightDirection(Options->AnimPreviewLightingDirection);
-	PreviewScene->GetScene()->UpdateDynamicSkyLight(Options->AnimPreviewSkyBrightness * FLinearColor(Options->AnimPreviewSkyColor), Options->AnimPreviewSkyBrightness * FLinearColor(Options->AnimPreviewFloorColor));
 	PreviewScene->SetLightColor(Options->AnimPreviewDirectionalColor);
 	PreviewScene->SetLightBrightness(Options->AnimPreviewLightBrightness);
 }
diff --git a/Engine/Source/Editor/LevelEditor/Private/LevelEditor.cpp b/Engine/Source/Editor/LevelEditor/Private/LevelEditor.cpp
--- a/Engine/Source/Editor/LevelEditor/Private/LevelEditor.cpp
+++ b/Engine/Source/Editor/LevelEditor/Private/LevelEditor.cpp
@@ -571,6 +571,7 @@ void FLevelEditorModule::BindGlobalLevelEditorCommands()
 
 	ActionList.MapAction( Commands.BrowseDocumentation, FExecuteAction::CreateStatic( &FLevelEditorActionCallbacks::BrowseDocumentation ) );
 	ActionList.MapAction( Commands.BrowseAPIReference, FExecuteAction::CreateStatic( &FLevelEditorActionCallbacks::BrowseAPIReference ) );
+	ActionList.MapAction( Commands.BrowseCVars, FExecuteAction::CreateStatic( &FLevelEditorActionCallbacks::BrowseCVars ) );
 	ActionList.MapAction( Commands.BrowseViewportControls, FExecuteAction::CreateStatic( &FLevelEditorActionCallbacks::BrowseViewportControls ) );
 	ActionList.MapAction( Commands.NewLevel, FExecuteAction::CreateStatic( &FLevelEditorActionCallbacks::NewLevel ), FCanExecuteAction::CreateStatic( &FLevelEditorActionCallbacks::NewLevel_CanExecute ) );
 	ActionList.MapAction(Commands.OpenLevel, FExecuteAction::CreateStatic(&FLevelEditorActionCallbacks::OpenLevel), FCanExecuteAction::CreateStatic(&FLevelEditorActionCallbacks::OpenLevel_CanExecute));
diff --git a/Engine/Source/Editor/LevelEditor/Private/LevelEditorActions.cpp b/Engine/Source/Editor/LevelEditor/Private/LevelEditorActions.cpp
--- a/Engine/Source/Editor/LevelEditor/Private/LevelEditorActions.cpp
+++ b/Engine/Source/Editor/LevelEditor/Private/LevelEditorActions.cpp
@@ -131,6 +131,11 @@ void FLevelEditorActionCallbacks::BrowseAPIReference()
 	IDocumentation::Get()->OpenAPIHome();
 }
 
+void FLevelEditorActionCallbacks::BrowseCVars()
+{
+	GEditor->Exec(GetWorld(), TEXT("help"));
+}
+
 void FLevelEditorActionCallbacks::BrowseViewportControls()
 {
 	FString URL;
@@ -2881,6 +2886,7 @@ void FLevelEditorCommands::RegisterCommands()
 {
 	UI_COMMAND( BrowseDocumentation, "Documentation...", "Opens the main documentation page, and allows you to search across all UE4 support sites.", EUserInterfaceActionType::Button, FInputChord( EKeys::F1 ) );
 	UI_COMMAND( BrowseAPIReference, "API Reference...", "Opens the API reference documentation", EUserInterfaceActionType::Button, FInputChord() );
+	UI_COMMAND( BrowseCVars, "Console Variables", "Creates an HTML file to browse the console variables and commands (console command 'help')", EUserInterfaceActionType::Button, FInputChord() );
 	UI_COMMAND( BrowseViewportControls, "Viewport Controls...", "Opens the viewport controls cheat sheet", EUserInterfaceActionType::Button, FInputChord() );
 	UI_COMMAND( NewLevel, "New Level...", "Create a new level, or choose a level template to start from.", EUserInterfaceActionType::Button, FInputChord( EModifierKey::Control, EKeys::N ) );
 	UI_COMMAND( OpenLevel, "Open Level...", "Loads an existing level", EUserInterfaceActionType::Button, FInputChord( EModifierKey::Control, EKeys::O ) );
diff --git a/Engine/Source/Editor/LevelEditor/Private/LevelEditorMenu.cpp b/Engine/Source/Editor/LevelEditor/Private/LevelEditorMenu.cpp
--- a/Engine/Source/Editor/LevelEditor/Private/LevelEditorMenu.cpp
+++ b/Engine/Source/Editor/LevelEditor/Private/LevelEditorMenu.cpp
@@ -330,6 +330,7 @@ TSharedRef< SWidget > FLevelEditorMenu::MakeLevelEditorMenu( const TSharedPtr<FU
 				{
 					MenuBuilder.AddMenuEntry( FLevelEditorCommands::Get().BrowseAPIReference );
 				}
+				MenuBuilder.AddMenuEntry( FLevelEditorCommands::Get().BrowseCVars );
 
 				MenuBuilder.AddMenuSeparator();
 
diff --git a/Engine/Source/Editor/LevelEditor/Public/LevelEditorActions.h b/Engine/Source/Editor/LevelEditor/Public/LevelEditorActions.h
--- a/Engine/Source/Editor/LevelEditor/Public/LevelEditorActions.h
+++ b/Engine/Source/Editor/LevelEditor/Public/LevelEditorActions.h
@@ -35,6 +35,7 @@ public:
 	
 	TSharedPtr< FUICommandInfo > BrowseDocumentation;
 	TSharedPtr< FUICommandInfo > BrowseAPIReference;
+	TSharedPtr< FUICommandInfo > BrowseCVars;
 	TSharedPtr< FUICommandInfo > BrowseViewportControls;
 
 	/** Level file commands */
@@ -618,6 +619,9 @@ public:
 	/** Opens the API documentation CHM */
 	static void BrowseAPIReference();
 
+	/** Creates an HTML file to browse the console variables and commands */
+	static void BrowseCVars();
+
 	/** Opens the viewport controls page*/
 	static void BrowseViewportControls();
 
diff --git a/Engine/Source/Editor/Persona/Private/SAnimationSequenceBrowser.cpp b/Engine/Source/Editor/Persona/Private/SAnimationSequenceBrowser.cpp
--- a/Engine/Source/Editor/Persona/Private/SAnimationSequenceBrowser.cpp
+++ b/Engine/Source/Editor/Persona/Private/SAnimationSequenceBrowser.cpp
@@ -916,7 +916,6 @@ void SAnimationSequenceBrowser::CreateAssetTooltipResources()
 	const UDestructableMeshEditorSettings* Options = GetDefault<UDestructableMeshEditorSettings>();
 
 	PreviewScene.SetLightDirection(Options->AnimPreviewLightingDirection);
-	PreviewScene.GetScene()->UpdateDynamicSkyLight(Options->AnimPreviewSkyBrightness * FLinearColor(Options->AnimPreviewSkyColor), Options->AnimPreviewSkyBrightness * FLinearColor(Options->AnimPreviewFloorColor));
 	PreviewScene.SetLightColor(Options->AnimPreviewDirectionalColor);
 	PreviewScene.SetLightBrightness(Options->AnimPreviewLightBrightness);
 }
diff --git a/Engine/Source/Editor/StaticMeshEditor/Private/StaticMeshEditorTools.cpp b/Engine/Source/Editor/StaticMeshEditor/Private/StaticMeshEditorTools.cpp
--- a/Engine/Source/Editor/StaticMeshEditor/Private/StaticMeshEditorTools.cpp
+++ b/Engine/Source/Editor/StaticMeshEditor/Private/StaticMeshEditorTools.cpp
@@ -385,6 +385,22 @@ void FMeshBuildSettingsLayout::GenerateChildContent( IDetailChildrenBuilder& Chi
 		];
 	}
 
+	{
+		ChildrenBuilder.AddChildContent(LOCTEXT("UseHighPrecisionTangentBasis", "Use High Precision Tangent Basis"))
+		.NameContent()
+		[
+			SNew(STextBlock)
+			.Font(IDetailLayoutBuilder::GetDetailFont())
+			.Text(LOCTEXT("UseHighPrecisionTangentBasis", "Use High Precision Tangent Basis"))
+		]
+		.ValueContent()
+		[
+			SNew(SCheckBox)
+			.IsChecked(this, &FMeshBuildSettingsLayout::ShouldUseHighPrecisionTangentBasis)
+			.OnCheckStateChanged(this, &FMeshBuildSettingsLayout::OnUseHighPrecisionTangentBasisChanged)
+		];
+	}
+
 	{
 		ChildrenBuilder.AddChildContent( LOCTEXT("UseFullPrecisionUVs", "Use Full Precision UVs") )
 		.NameContent()
@@ -617,6 +633,11 @@ ECheckBoxState FMeshBuildSettingsLayout::ShouldBuildReversedIndexBuffer() const
 	return BuildSettings.bBuildReversedIndexBuffer ? ECheckBoxState::Checked : ECheckBoxState::Unchecked;
 }
 
+ECheckBoxState FMeshBuildSettingsLayout::ShouldUseHighPrecisionTangentBasis() const
+{
+	return BuildSettings.bUseHighPrecisionTangentBasis ? ECheckBoxState::Checked : ECheckBoxState::Unchecked;
+}
+
 ECheckBoxState FMeshBuildSettingsLayout::ShouldUseFullPrecisionUVs() const
 {
 	return BuildSettings.bUseFullPrecisionUVs ? ECheckBoxState::Checked : ECheckBoxState::Unchecked;
@@ -741,6 +762,19 @@ void FMeshBuildSettingsLayout::OnBuildReversedIndexBufferChanged(ECheckBoxState
 	}
 }
 
+void FMeshBuildSettingsLayout::OnUseHighPrecisionTangentBasisChanged(ECheckBoxState NewState)
+{
+	const bool bUseHighPrecisionTangents = (NewState == ECheckBoxState::Checked) ? true : false;
+	if (BuildSettings.bUseHighPrecisionTangentBasis != bUseHighPrecisionTangents)
+	{
+		if (FEngineAnalytics::IsAvailable())
+		{
+			FEngineAnalytics::GetProvider().RecordEvent(TEXT("Editor.Usage.StaticMesh.BuildSettings"), TEXT("bUseHighPrecisionTangentBasis"), bUseHighPrecisionTangents ? TEXT("True") : TEXT("False"));
+		}
+		BuildSettings.bUseHighPrecisionTangentBasis = bUseHighPrecisionTangents;
+	}
+}
+
 void FMeshBuildSettingsLayout::OnUseFullPrecisionUVsChanged(ECheckBoxState NewState)
 {
 	const bool bUseFullPrecisionUVs = (NewState == ECheckBoxState::Checked) ? true : false;
diff --git a/Engine/Source/Editor/StaticMeshEditor/Private/StaticMeshEditorTools.h b/Engine/Source/Editor/StaticMeshEditor/Private/StaticMeshEditorTools.h
--- a/Engine/Source/Editor/StaticMeshEditor/Private/StaticMeshEditorTools.h
+++ b/Engine/Source/Editor/StaticMeshEditor/Private/StaticMeshEditorTools.h
@@ -135,6 +135,7 @@ private:
 	ECheckBoxState ShouldRemoveDegenerates() const;
 	ECheckBoxState ShouldBuildAdjacencyBuffer() const;
 	ECheckBoxState ShouldBuildReversedIndexBuffer() const;
+	ECheckBoxState ShouldUseHighPrecisionTangentBasis() const;
 	ECheckBoxState ShouldUseFullPrecisionUVs() const;
 	ECheckBoxState ShouldGenerateLightmapUVs() const;
 	ECheckBoxState ShouldGenerateDistanceFieldAsIfTwoSided() const;
@@ -152,6 +153,7 @@ private:
 	void OnRemoveDegeneratesChanged(ECheckBoxState NewState);
 	void OnBuildAdjacencyBufferChanged(ECheckBoxState NewState);
 	void OnBuildReversedIndexBufferChanged(ECheckBoxState NewState);
+	void OnUseHighPrecisionTangentBasisChanged(ECheckBoxState NewState);
 	void OnUseFullPrecisionUVsChanged(ECheckBoxState NewState);
 	void OnGenerateLightmapUVsChanged(ECheckBoxState NewState);
 	void OnGenerateDistanceFieldAsIfTwoSidedChanged(ECheckBoxState NewState);
diff --git a/Engine/Source/Editor/UnrealEd/Private/EditorViewportClient.cpp b/Engine/Source/Editor/UnrealEd/Private/EditorViewportClient.cpp
--- a/Engine/Source/Editor/UnrealEd/Private/EditorViewportClient.cpp
+++ b/Engine/Source/Editor/UnrealEd/Private/EditorViewportClient.cpp
@@ -3192,7 +3192,8 @@ void FEditorViewportClient::Draw(FViewport* InViewport, FCanvas* Canvas)
 	float DeltaTimeSeconds;
 
 	UWorld* World = GetWorld();
-	if (( GetScene() != World->Scene) || (IsRealtime() == true))
+	// During Simulation blueprints are directly using the World time, causing a mismatch with Material's Frame time
+	if (( GetScene() != World->Scene) || (IsRealtime() && !IsSimulateInEditorViewport()))
 	{
 		// Use time relative to start time to avoid issues with float vs double
 		TimeSeconds = FApp::GetCurrentTime() - GStartTime;
diff --git a/Engine/Source/Editor/UnrealEd/Private/Fbx/FbxMainExport.cpp b/Engine/Source/Editor/UnrealEd/Private/Fbx/FbxMainExport.cpp
--- a/Engine/Source/Editor/UnrealEd/Private/Fbx/FbxMainExport.cpp
+++ b/Engine/Source/Editor/UnrealEd/Private/Fbx/FbxMainExport.cpp
@@ -557,7 +557,7 @@ void FFbxExporter::ExportModel(UModel* Model, FbxNode* Node, const char* Name)
 	for (uint32 VertexIdx = 0; VertexIdx < VertCount; ++VertexIdx)
 	{
 		FModelVertex& Vertex = Model->VertexBuffer.Vertices[VertexIdx];
-		FVector Normal = (FVector) Vertex.TangentZ;
+		FVector Normal = Vertex.TangentZ;
 
 		// If the vertex is outside of the world extent, snap it to the origin.  The faces associated with
 		// these vertices will be removed before exporting.  We leave the snapped vertex in the buffer so
diff --git a/Engine/Source/Editor/UnrealEd/Private/Fbx/FbxSkeletalMeshExport.cpp b/Engine/Source/Editor/UnrealEd/Private/Fbx/FbxSkeletalMeshExport.cpp
--- a/Engine/Source/Editor/UnrealEd/Private/Fbx/FbxSkeletalMeshExport.cpp
+++ b/Engine/Source/Editor/UnrealEd/Private/Fbx/FbxSkeletalMeshExport.cpp
@@ -139,7 +139,7 @@ FbxNode* FFbxExporter::CreateMesh(const USkeletalMesh* SkelMesh, const TCHAR* Me
 
 	for (int32 VertIndex = 0; VertIndex < VertexCount; ++VertIndex)
 	{
-		FVector Normal			= (FVector)(Vertices[VertIndex].TangentZ);
+		FVector Normal			= Vertices[VertIndex].TangentZ;
 		FbxVector4 FbxNormal	= Converter.ConvertToFbxPos(Normal);
 
 		LayerElementNormal->GetDirectArray().Add(FbxNormal);
diff --git a/Engine/Source/Editor/UnrealEd/Private/LevelEditorViewport.cpp b/Engine/Source/Editor/UnrealEd/Private/LevelEditorViewport.cpp
--- a/Engine/Source/Editor/UnrealEd/Private/LevelEditorViewport.cpp
+++ b/Engine/Source/Editor/UnrealEd/Private/LevelEditorViewport.cpp
@@ -4087,6 +4087,7 @@ void FLevelEditorViewportClient::DrawCanvas( FViewport& InViewport, FSceneView&
  */
 void FLevelEditorViewportClient::DrawTextureStreamingBounds(const FSceneView* View,FPrimitiveDrawInterface* PDI)
 {
+#if 0 //@todo
 	// Iterate each level
 	for (TObjectIterator<ULevel> It; It; ++It)
 	{
@@ -4118,6 +4119,7 @@ void FLevelEditorViewportClient::DrawTextureStreamingBounds(const FSceneView* Vi
 			}
 		}
 	}
+#endif
 }
 
 
diff --git a/Engine/Source/Editor/UnrealEd/Private/Lightmass/LightmassRender.cpp b/Engine/Source/Editor/UnrealEd/Private/Lightmass/LightmassRender.cpp
--- a/Engine/Source/Editor/UnrealEd/Private/Lightmass/LightmassRender.cpp
+++ b/Engine/Source/Editor/UnrealEd/Private/Lightmass/LightmassRender.cpp
@@ -250,6 +250,10 @@ public:
 			{
 				return true;
 			}
+			else if(FCString::Stristr(ShaderType->GetName(), TEXT("Simple")))
+			{
+				return true;
+			}
 			else if(FCString::Stristr(ShaderType->GetName(), TEXT("BasePassPSFNoLightMapPolicy")))
 			{
 				return true;
diff --git a/Engine/Source/Editor/UnrealEd/Private/PreviewMaterial.cpp b/Engine/Source/Editor/UnrealEd/Private/PreviewMaterial.cpp
--- a/Engine/Source/Editor/UnrealEd/Private/PreviewMaterial.cpp
+++ b/Engine/Source/Editor/UnrealEd/Private/PreviewMaterial.cpp
@@ -98,6 +98,10 @@ public:
 			{
 				bShaderTypeMatches = true;
 			}
+			else if (FCString::Stristr(ShaderType->GetName(), TEXT("Simple")))
+			{
+				bShaderTypeMatches = true;
+			}
 			else if (FCString::Stristr(ShaderType->GetName(), TEXT("BasePassPSFNoLightMapPolicy")))
 			{
 				bShaderTypeMatches = true;
diff --git a/Engine/Source/Editor/UnrealEd/Public/ThumbnailHelpers.h b/Engine/Source/Editor/UnrealEd/Public/ThumbnailHelpers.h
--- a/Engine/Source/Editor/UnrealEd/Public/ThumbnailHelpers.h
+++ b/Engine/Source/Editor/UnrealEd/Public/ThumbnailHelpers.h
@@ -216,6 +216,10 @@ class UActorComponent;
 class UNREALED_API FClassActorThumbnailScene : public FThumbnailPreviewScene
 {
 public:
+	FClassActorThumbnailScene()
+		:PreviewActor(0)
+	{
+	}
 
 	/** Returns true if this component can be visualized */
 	bool IsValidComponentForVisualization(UActorComponent* Component) const;
diff --git a/Engine/Source/Editor/WorldBrowser/Private/Tiles/WorldTileCollectionModel.cpp b/Engine/Source/Editor/WorldBrowser/Private/Tiles/WorldTileCollectionModel.cpp
--- a/Engine/Source/Editor/WorldBrowser/Private/Tiles/WorldTileCollectionModel.cpp
+++ b/Engine/Source/Editor/WorldBrowser/Private/Tiles/WorldTileCollectionModel.cpp
@@ -2082,6 +2082,7 @@ bool FWorldTileCollectionModel::GenerateLODLevels(FLevelModelList InLevelList, i
 				SrcModel->BuildSettings.bRecomputeNormals = false;
 				SrcModel->BuildSettings.bRecomputeTangents = false;
 				SrcModel->BuildSettings.bRemoveDegenerates = false;
+				SrcModel->BuildSettings.bUseHighPrecisionTangentBasis = false;
 				SrcModel->BuildSettings.bUseFullPrecisionUVs = false;
 				SrcModel->RawMeshBulkData->SaveRawMesh(LandscapeRawMesh);
 
diff --git a/Engine/Source/Runtime/Core/Private/HAL/ConsoleManager.cpp b/Engine/Source/Runtime/Core/Private/HAL/ConsoleManager.cpp
--- a/Engine/Source/Runtime/Core/Private/HAL/ConsoleManager.cpp
+++ b/Engine/Source/Runtime/Core/Private/HAL/ConsoleManager.cpp
@@ -1753,14 +1753,6 @@ static TAutoConsoleVariable<int32> CVarClearWithExcludeRects(
 	TEXT(" 2: Auto (default is 2, pick what is considered best on this hardware)"),
 	ECVF_RenderThreadSafe);
 
-static TAutoConsoleVariable<int32> CVarSimpleDynamicLighting(
-	TEXT("r.SimpleDynamicLighting"),
-	0,
-	TEXT("Whether to use simple dynamic lighting, which just renders an unshadowed dynamic directional light and a skylight.\n")
-	TEXT("All other lighting features are disabled when true.  This is useful for supporting very low end hardware.\n")
-	TEXT("0:off, 1:on"),
-	ECVF_RenderThreadSafe);
-
 static TAutoConsoleVariable<int32> CVarTranslucentSortPolicy(
 	TEXT("r.TranslucentSortPolicy"),
 	0,
@@ -2279,13 +2271,6 @@ static TAutoConsoleVariable<int32> CVarFeatureLevelPreview(
 	TEXT("If 1 the quick settings menu will contain an option to enable feature level preview modes"),
 	ECVF_RenderThreadSafe);
 
-
-static TAutoConsoleVariable<int32> CVarGBuffer(
-	TEXT("r.GBuffer"),
-	1,
-	TEXT("0=Do not use GBuffer (fast minimal rendering path), 1=Use GBuffer [default]."),
-	ECVF_RenderThreadSafe);
-
 static TAutoConsoleVariable<int32> CVarVerifyPeer(
 	TEXT("n.VerifyPeer"),
 	1,
diff --git a/Engine/Source/Runtime/Core/Public/Math/UnrealMathDirectX.h b/Engine/Source/Runtime/Core/Public/Math/UnrealMathDirectX.h
--- a/Engine/Source/Runtime/Core/Public/Math/UnrealMathDirectX.h
+++ b/Engine/Source/Runtime/Core/Public/Math/UnrealMathDirectX.h
@@ -610,7 +610,7 @@ FORCEINLINE VectorRegister VectorLoadByte4Reverse( const uint8* Ptr )
 * @param Ptr			Unaligned memory pointer to the RGB10A2(4 bytes).
 * @return				VectorRegister with 4 FLOATs loaded from Ptr.
 */
-#define VectorLoadRGB10A2( Ptr )	DirectX::PackedVector::XMLoadUDecN4( (const DirectX::PackedVector::XMUDECN4*)(Ptr) )
+#define VectorLoadURGB10A2N( Ptr )	DirectX::PackedVector::XMLoadUDecN4( (const DirectX::PackedVector::XMUDECN4*)(Ptr) )
 
 /**
 * Converts the 4 FLOATs in the vector RGB10A2, clamped to [0, 1023] and [0, 3], and stores to unaligned memory.
@@ -619,7 +619,25 @@ FORCEINLINE VectorRegister VectorLoadByte4Reverse( const uint8* Ptr )
 * @param Vec			Vector containing 4 FLOATs
 * @param Ptr			Unaligned memory pointer to store the packed RGB10A2(4 bytes).
 */
-#define VectorStoreRGB10A2( Vec, Ptr )	DirectX::PackedVector::XMStoreUDecN4( (const DirectX::PackedVector::XMUDECN4*)(Ptr), Vec )
+#define VectorStoreURGB10A2N( Vec, Ptr )	DirectX::PackedVector::XMStoreUDecN4( (const DirectX::PackedVector::XMUDECN4*)(Ptr), Vec )
+
+/**
+* Loads packed RGBA16(4 bytes) from unaligned memory and converts them into 4 FLOATs.
+* IMPORTANT: You need to call VectorResetFloatRegisters() before using scalar FLOATs after you've used this intrinsic!
+*
+* @param Ptr			Unaligned memory pointer to the RGBA16(8 bytes).
+* @return				VectorRegister with 4 FLOATs loaded from Ptr.
+*/
+#define VectorLoadURGBA16N( Ptr )	DirectX::PackedVector::XMLoadUShortN4( (const DirectX::PackedVector::XMUSHORTN4*)(Ptr) )
+
+/**
+* Converts the 4 FLOATs in the vector RGBA16, clamped to [0, 65535], and stores to unaligned memory.
+* IMPORTANT: You need to call VectorResetFloatRegisters() before using scalar FLOATs after you've used this intrinsic!
+*
+* @param Vec			Vector containing 4 FLOATs
+* @param Ptr			Unaligned memory pointer to store the packed RGBA16(8 bytes).
+*/
+#define VectorStoreURGBA16N( Vec, Ptr )	DirectX::PackedVector::XMStoreUShortN4( (const DirectX::PackedVector::XMUSHORTN4*)(Ptr), Vec )
 
 /**
  * Returns non-zero if any element in Vec1 is greater than the corresponding element in Vec2, otherwise 0.
diff --git a/Engine/Source/Runtime/Core/Public/Math/UnrealMathFPU.h b/Engine/Source/Runtime/Core/Public/Math/UnrealMathFPU.h
--- a/Engine/Source/Runtime/Core/Public/Math/UnrealMathFPU.h
+++ b/Engine/Source/Runtime/Core/Public/Math/UnrealMathFPU.h
@@ -17,13 +17,21 @@ struct VectorRegister
 };
 
 /**
-*	int32[4] vector register type, where the first in32 (X) is stored in the lowest 32 bits, and so on.
+*	int32[4] vector register type, where the first int32 (X) is stored in the lowest 32 bits, and so on.
 */
 struct VectorRegisterInt
 {
 	int32	V[4];
 };
 
+/**
+*	double[2] vector register type, where the first double (X) is stored in the lowest 64 bits, and so on.
+*/
+struct VectorRegisterDouble
+{
+	double	V[2];
+};
+
 // For a struct of 4 floats, we need the double braces
 #define DECLARE_VECTOR_REGISTER(X, Y, Z, W) { { X, Y, Z, W } }
 
@@ -951,7 +959,7 @@ FORCEINLINE void VectorStoreByte4( const VectorRegister& Vec, void* Ptr )
 * @param Ptr			Unaligned memory pointer to the RGB10A2(4 bytes).
 * @return				VectorRegister with 4 FLOATs loaded from Ptr.
 */
-FORCEINLINE VectorRegister VectorLoadRGB10A2(void* Ptr)
+FORCEINLINE VectorRegister VectorLoadURGB10A2N(void* Ptr)
 {
 	float V[4];
 	uint32 E = *(uint32*)Ptr;
@@ -971,7 +979,7 @@ FORCEINLINE VectorRegister VectorLoadRGB10A2(void* Ptr)
 * @param Vec			Vector containing 4 FLOATs
 * @param Ptr			Unaligned memory pointer to store the packed RGB10A2(4 bytes).
 */
-FORCEINLINE void VectorStoreRGB10A2(const VectorRegister& Vec, void* Ptr)
+FORCEINLINE void VectorStoreURGB10A2N(const VectorRegister& Vec, void* Ptr)
 {
 	VectorRegister Tmp;
 	Tmp = VectorMax(Vec, MakeVectorRegister(0.0f, 0.0f, 0.0f, 0.0f));
@@ -1189,6 +1197,47 @@ FORCEINLINE VectorRegister VectorStep(const VectorRegister& X)
 		(float)(VectorGetComponent(X, 3) >= 0.0f ? 1.0f : -1.0f));
 }
 
+/**
+* Loads packed RGBA16(4 bytes) from unaligned memory and converts them into 4 FLOATs.
+* IMPORTANT: You need to call VectorResetFloatRegisters() before using scalar FLOATs after you've used this intrinsic!
+*
+* @param Ptr			Unaligned memory pointer to the RGBA16(8 bytes).
+* @return				VectorRegister with 4 FLOATs loaded from Ptr.
+*/
+FORCEINLINE VectorRegister VectorLoadURGBA16N(void* Ptr)
+{
+	float V[4];
+	uint16* E = (uint16*)Ptr;
+
+	V[0] = float(E[0]) / 65535.0f;
+	V[1] = float(E[1]) / 65535.0f;
+	V[2] = float(E[2]) / 65535.0f;
+	V[3] = float(E[3]) / 65535.0f;
+
+	return MakeVectorRegister(V[0], V[1], V[2], V[3]);
+}
+
+/**
+* Converts the 4 FLOATs in the vector RGBA16, clamped to [0, 65535], and stores to unaligned memory.
+* IMPORTANT: You need to call VectorResetFloatRegisters() before using scalar FLOATs after you've used this intrinsic!
+*
+* @param Vec			Vector containing 4 FLOATs
+* @param Ptr			Unaligned memory pointer to store the packed RGBA16(8 bytes).
+*/
+FORCEINLINE void VectorStoreURGBA16N(const VectorRegister& Vec, void* Ptr)
+{
+	VectorRegister Tmp;
+	Tmp = VectorMax(Vec, MakeVectorRegister(0.0f, 0.0f, 0.0f, 0.0f));
+	Tmp = VectorMin(Tmp, MakeVectorRegister(1.0f, 1.0f, 1.0f, 1.0f));
+	Tmp = VectorMultiplyAdd(Tmp, MakeVectorRegister(65535.0f, 65535.0f, 65535.0f, 65535.0f), MakeVectorRegister(0.5f, 0.5f, 0.5f, 0.5f));
+	Tmp = VectorTruncate(Tmp);
+
+	uint16* Out = (uint16*)Ptr;
+	Out[0] = (uint16)Tmp.V[0];
+	Out[1] = (uint16)Tmp.V[0];
+	Out[2] = (uint16)Tmp.V[0];
+	Out[3] = (uint16)Tmp.V[0];
+}
 
 // To be continued...
 
diff --git a/Engine/Source/Runtime/Core/Public/Math/UnrealMathNeon.h b/Engine/Source/Runtime/Core/Public/Math/UnrealMathNeon.h
--- a/Engine/Source/Runtime/Core/Public/Math/UnrealMathNeon.h
+++ b/Engine/Source/Runtime/Core/Public/Math/UnrealMathNeon.h
@@ -919,7 +919,7 @@ FORCEINLINE void VectorStoreHalf4(VectorRegister Vec, void* RESTRICT Ptr)
 * @param Ptr			Unaligned memory pointer to the RGB10A2(4 bytes).
 * @return				VectorRegister with 4 FLOATs loaded from Ptr.
 */
-FORCEINLINE VectorRegister VectorLoadRGB10A2(void* Ptr)
+FORCEINLINE VectorRegister VectorLoadURGB10A2N(void* Ptr)
 {
 	float V[4];
 	uint32 E = *(uint32*)Ptr;
@@ -939,7 +939,7 @@ FORCEINLINE VectorRegister VectorLoadRGB10A2(void* Ptr)
 * @param Vec			Vector containing 4 FLOATs
 * @param Ptr			Unaligned memory pointer to store the packed RGB10A2(4 bytes).
 */
-FORCEINLINE void VectorStoreRGB10A2(const VectorRegister& Vec, void* Ptr)
+FORCEINLINE void VectorStoreURGB10A2N(const VectorRegister& Vec, void* Ptr)
 {
 	union { VectorRegister V; float F[4]; } Tmp;
 	Tmp.V = VectorMax(Vec, MakeVectorRegister(0.0f, 0.0f, 0.0f, 0.0f));
@@ -1184,7 +1184,46 @@ FORCEINLINE VectorRegister VectorStep(const VectorRegister& X)
 		(float)(VectorGetComponent(X, 3) >= 0.0f ? 1.0f : -1.0f));
 }
 
+/**
+* Loads packed RGBA16(4 bytes) from unaligned memory and converts them into 4 FLOATs.
+* IMPORTANT: You need to call VectorResetFloatRegisters() before using scalar FLOATs after you've used this intrinsic!
+*
+* @param Ptr			Unaligned memory pointer to the RGBA16(8 bytes).
+* @return				VectorRegister with 4 FLOATs loaded from Ptr.
+*/
+FORCEINLINE VectorRegister VectorLoadURGBA16N(void* Ptr)
+{
+	float V[4];
+	uint16* E = (uint16*)Ptr;
 
-// To be continued...
+	V[0] = float(E[0]) / 65535.0f;
+	V[1] = float(E[1]) / 65535.0f;
+	V[2] = float(E[2]) / 65535.0f;
+	V[3] = float(E[3]) / 65535.0f;
+
+	return MakeVectorRegister(V[0], V[1], V[2], V[3]);
+}
+
+/**
+* Converts the 4 FLOATs in the vector RGBA16, clamped to [0, 65535], and stores to unaligned memory.
+* IMPORTANT: You need to call VectorResetFloatRegisters() before using scalar FLOATs after you've used this intrinsic!
+*
+* @param Vec			Vector containing 4 FLOATs
+* @param Ptr			Unaligned memory pointer to store the packed RGBA16(8 bytes).
+*/
+FORCEINLINE void VectorStoreURGBA16N(const VectorRegister& Vec, void* Ptr)
+{
+	VectorRegister Tmp;
+	Tmp = VectorMax(Vec, MakeVectorRegister(0.0f, 0.0f, 0.0f, 0.0f));
+	Tmp = VectorMin(Tmp, MakeVectorRegister(1.0f, 1.0f, 1.0f, 1.0f));
+	Tmp = VectorMultiplyAdd(Tmp, MakeVectorRegister(65535.0f, 65535.0f, 65535.0f, 65535.0f), MakeVectorRegister(0.5f, 0.5f, 0.5f, 0.5f));
+	Tmp = VectorTruncate(Tmp);
 
+	uint16* Out = (uint16*)Ptr;
+	Out[0] = (uint16)VectorGetComponent(Tmp, 0);
+	Out[1] = (uint16)VectorGetComponent(Tmp, 1);
+	Out[2] = (uint16)VectorGetComponent(Tmp, 2);
+	Out[3] = (uint16)VectorGetComponent(Tmp, 3);
+}
 
+// To be continued...
diff --git a/Engine/Source/Runtime/Core/Public/Math/UnrealMathSSE.h b/Engine/Source/Runtime/Core/Public/Math/UnrealMathSSE.h
--- a/Engine/Source/Runtime/Core/Public/Math/UnrealMathSSE.h
+++ b/Engine/Source/Runtime/Core/Public/Math/UnrealMathSSE.h
@@ -20,6 +20,7 @@
  */
 typedef __m128	VectorRegister;
 typedef __m128i VectorRegisterInt;
+typedef __m128d VectorRegisterDouble;
 
 // for an __m128, we need a single set of braces (for clang)
 #define DECLARE_VECTOR_REGISTER(X, Y, Z, W) { X, Y, Z, W }
@@ -911,7 +912,7 @@ FORCEINLINE void VectorStoreByte4( const VectorRegister& Vec, void* Ptr )
 * @param Ptr			Unaligned memory pointer to the RGB10A2(4 bytes).
 * @return				VectorRegister with 4 FLOATs loaded from Ptr.
 */
-FORCEINLINE VectorRegister VectorLoadRGB10A2(void* Ptr)
+FORCEINLINE VectorRegister VectorLoadURGB10A2N(void* Ptr)
 {
 	VectorRegister Tmp;
 
@@ -929,9 +930,9 @@ FORCEINLINE VectorRegister VectorLoadRGB10A2(void* Ptr)
 * IMPORTANT: You need to call VectorResetFloatRegisters() before using scalar FLOATs after you've used this intrinsic!
 *
 * @param Vec			Vector containing 4 FLOATs
-* @param Ptr			Unaligned memory pointer to store the packed RGB10A2(4 bytes).
+* @param Ptr			Unaligned memory pointer to store the packed RGBA16(8 bytes).
 */
-FORCEINLINE void VectorStoreRGB10A2(const VectorRegister& Vec, void* Ptr)
+FORCEINLINE void VectorStoreURGB10A2N(const VectorRegister& Vec, void* Ptr)
 {
 	VectorRegister Tmp;
 	Tmp = _mm_max_ps(Vec, MakeVectorRegister(0.0f, 0.0f, 0.0f, 0.0f));
@@ -953,6 +954,56 @@ FORCEINLINE void VectorStoreRGB10A2(const VectorRegister& Vec, void* Ptr)
 	_mm_store_ss((float *)Ptr, *(const VectorRegister*)&TmpI);
 }
 
+/**
+* Loads packed RGBA16(4 bytes) from unaligned memory and converts them into 4 FLOATs.
+* IMPORTANT: You need to call VectorResetFloatRegisters() before using scalar FLOATs after you've used this intrinsic!
+*
+* @param Ptr			Unaligned memory pointer to the RGBA16(8 bytes).
+* @return				VectorRegister with 4 FLOATs loaded from Ptr.
+*/
+FORCEINLINE VectorRegister VectorLoadURGBA16N(void* Ptr)
+{
+	VectorRegisterDouble TmpD = _mm_load1_pd(reinterpret_cast<const double *>(Ptr));
+
+	VectorRegisterInt Mask = MakeVectorRegisterInt(0x0000FFFF, 0x0000FFFF, 0xFFFF0000, 0xFFFF0000);
+	VectorRegisterInt FlipSign = MakeVectorRegisterInt(0, 0, 0x80000000, 0x80000000);
+
+	VectorRegister Tmp = _mm_and_ps(
+		reinterpret_cast<const VectorRegister*>(&TmpD)[0], 
+		reinterpret_cast<const VectorRegister*>(&Mask)[0]
+		);
+
+	Tmp = _mm_xor_ps(Tmp, reinterpret_cast<const VectorRegister*>(&FlipSign)[0]);
+	Tmp = _mm_cvtepi32_ps(reinterpret_cast<const VectorRegisterInt*>(&Tmp)[0]);
+	Tmp = _mm_add_ps(Tmp, MakeVectorRegister(0, 0, 32768.0f*65536.0f, 32768.0f*65536.0f));
+	Tmp = _mm_mul_ps(Tmp, MakeVectorRegister(1.0f / 65535.0f, 1.0f / 65535.0f, 1.0f / (65535.0f*65536.0f), 1.0f / (65535.0f*65536.0f)));
+	return _mm_shuffle_ps(Tmp, Tmp, _MM_SHUFFLE(3, 1, 2, 0));
+}
+
+/**
+* Converts the 4 FLOATs in the vector RGBA16, clamped to [0, 65535], and stores to unaligned memory.
+* IMPORTANT: You need to call VectorResetFloatRegisters() before using scalar FLOATs after you've used this intrinsic!
+*
+* @param Vec			Vector containing 4 FLOATs
+* @param Ptr			Unaligned memory pointer to store the packed RGB10A2(4 bytes).
+*/
+FORCEINLINE void VectorStoreURGBA16N(const VectorRegister& Vec, void* Ptr)
+{
+
+	VectorRegister Tmp;
+	Tmp = _mm_max_ps(Vec, MakeVectorRegister(0.0f, 0.0f, 0.0f, 0.0f));
+	Tmp = _mm_min_ps(Tmp, MakeVectorRegister(1.0f, 1.0f, 1.0f, 1.0f));
+	Tmp = _mm_mul_ps(Tmp, MakeVectorRegister(65535.0f, 65535.0f, 65535.0f, 65535.0f));
+
+	VectorRegisterInt TmpI = _mm_cvtps_epi32(Tmp);
+
+	uint16* Out = (uint16*)Ptr;
+	Out[0] = static_cast<int16>(_mm_extract_epi16(TmpI, 0));
+	Out[1] = static_cast<int16>(_mm_extract_epi16(TmpI, 2));
+	Out[2] = static_cast<int16>(_mm_extract_epi16(TmpI, 4));
+	Out[3] = static_cast<int16>(_mm_extract_epi16(TmpI, 6));
+}
+
 /**
  * Returns non-zero if any element in Vec1 is greater than the corresponding element in Vec2, otherwise 0.
  *
diff --git a/Engine/Source/Runtime/Core/Public/UObject/RenderingObjectVersion.h b/Engine/Source/Runtime/Core/Public/UObject/RenderingObjectVersion.h
--- a/Engine/Source/Runtime/Core/Public/UObject/RenderingObjectVersion.h
+++ b/Engine/Source/Runtime/Core/Public/UObject/RenderingObjectVersion.h
@@ -15,6 +15,9 @@ struct CORE_API FRenderingObjectVersion
 		// Allows specifying resolution for reflecture capture probes
 		CustomReflectionCaptureResolutionSupport,
 
+		// Allows specifying resolution for reflecture capture probes
+		RemovedTextureStreamingLevelData,
+
 		// -----<new versions can be added above this line>-------------------------------------------------
 		VersionPlusOne,
 		LatestVersion = VersionPlusOne - 1
diff --git a/Engine/Source/Runtime/CoreUObject/Classes/Object.h b/Engine/Source/Runtime/CoreUObject/Classes/Object.h
--- a/Engine/Source/Runtime/CoreUObject/Classes/Object.h
+++ b/Engine/Source/Runtime/CoreUObject/Classes/Object.h
@@ -444,12 +444,25 @@ struct FPackedNormal
 * The full C++ class is located here: Engine\Source\Runtime\RenderCore\Public\PackedNormal.h
 */
 USTRUCT(immutable, noexport)
-struct FPackedRGB10A2
+struct FPackedRGB10A2N
 {
 	UPROPERTY(EditAnywhere, BlueprintReadWrite, Category = PackedBasis, SaveGame)
 	int32 Packed;
 };
 
+/**
+* A packed vector.
+* The full C++ class is located here: Engine\Source\Runtime\RenderCore\Public\PackedNormal.h
+*/
+USTRUCT(immutable, noexport)
+struct FPackedRGBA16N
+{
+	UPROPERTY(EditAnywhere, BlueprintReadWrite, Category = PackedNormal, SaveGame)
+	int32 XY;
+
+	UPROPERTY(EditAnywhere, BlueprintReadWrite, Category = PackedNormal, SaveGame)
+	int32 ZW;
+};
 
 /**
  * Screen coordinates.
diff --git a/Engine/Source/Runtime/Engine/Classes/Animation/VertexAnim/VertexAnimBase.h b/Engine/Source/Runtime/Engine/Classes/Animation/VertexAnim/VertexAnimBase.h
--- a/Engine/Source/Runtime/Engine/Classes/Animation/VertexAnim/VertexAnimBase.h
+++ b/Engine/Source/Runtime/Engine/Classes/Animation/VertexAnim/VertexAnimBase.h
@@ -32,7 +32,7 @@ struct FVertexAnimDelta
 
 			if (Ar.IsLoading())
 			{
-				V.TangentZDelta = FVector(TangentZDelta_DEPRECATED);
+				V.TangentZDelta = TangentZDelta_DEPRECATED;
 			}
 		}
 		else
diff --git a/Engine/Source/Runtime/Engine/Classes/Engine/Canvas.h b/Engine/Source/Runtime/Engine/Classes/Engine/Canvas.h
--- a/Engine/Source/Runtime/Engine/Classes/Engine/Canvas.h
+++ b/Engine/Source/Runtime/Engine/Classes/Engine/Canvas.h
@@ -145,7 +145,7 @@ public:
 /**
  * A drawing canvas.
  */
-UCLASS(transient)
+UCLASS(transient, BlueprintType)
 class ENGINE_API UCanvas
 	: public UObject
 {
diff --git a/Engine/Source/Runtime/Engine/Classes/Engine/EngineTypes.h b/Engine/Source/Runtime/Engine/Classes/Engine/EngineTypes.h
--- a/Engine/Source/Runtime/Engine/Classes/Engine/EngineTypes.h
+++ b/Engine/Source/Runtime/Engine/Classes/Engine/EngineTypes.h
@@ -269,6 +269,25 @@ enum ETranslucencyLightingMode
 	TLM_MAX,
 };
 
+/** Determines how the refraction offset should be computed for the material. */
+UENUM()
+enum ERefractionMode
+{
+	/** 
+	 * Refraction is computed based on the camera vector entering a medium whose index of refraction is defined by the Refraction material input.  
+	 * The new medium's surface is defined by the material's normal.  With this mode, a flat plane seen from the side will have a constant refraction offset.
+	 * This is a physical model of refraction but causes reading outside the scene color texture so is a poor fit for large refractive surfaces like water.
+	 */
+	RM_IndexOfRefraction UMETA(DisplayName="Index Of Refraction"),
+
+	/** 
+	 * The refraction offset into Scene Color is computed based on the difference between the per-pixel normal and the per-vertex normal.  
+	 * With this mode, a material whose normal is the default (0, 0, 1) will never cause any refraction.  This mode is only valid with tangent space normals.
+	 * The refraction material input scales the offset, although a value of 1.0 maps to no refraction, and a value of 2 maps to a scale of 1.0 on the offset.
+	 * This is a non-physical model of refraction but is useful on large refractive surfaces like water, since offsets have to stay small to avoid reading outside scene color.
+	 */
+	RM_PixelNormalOffset UMETA(DisplayName="Pixel Normal Offset")
+};
 
 /**
  * Enumerates available options for the translucency sort policy.
@@ -2336,6 +2355,10 @@ struct FMeshBuildSettings
 	UPROPERTY(EditAnywhere, Category=BuildSettings)
 	bool bBuildReversedIndexBuffer;
 
+	/** If true, Tangents will be stored at 16 bit vs 8 bit precision. */
+	UPROPERTY(EditAnywhere, Category = BuildSettings)
+	bool bUseHighPrecisionTangentBasis;
+
 	/** If true, UVs will be stored at full floating point precision. */
 	UPROPERTY(EditAnywhere, Category=BuildSettings)
 	bool bUseFullPrecisionUVs;
@@ -2384,6 +2407,7 @@ struct FMeshBuildSettings
 		, bRemoveDegenerates(true)
 		, bBuildAdjacencyBuffer(true)
 		, bBuildReversedIndexBuffer(true)
+		, bUseHighPrecisionTangentBasis(false)
 		, bUseFullPrecisionUVs(false)
 		, bGenerateLightmapUVs(true)
 		, MinLightmapResolution(64)
@@ -2405,6 +2429,7 @@ struct FMeshBuildSettings
 			&& bRemoveDegenerates == Other.bRemoveDegenerates
 			&& bBuildAdjacencyBuffer == Other.bBuildAdjacencyBuffer
 			&& bBuildReversedIndexBuffer == Other.bBuildReversedIndexBuffer
+			&& bUseHighPrecisionTangentBasis == Other.bUseHighPrecisionTangentBasis
 			&& bUseFullPrecisionUVs == Other.bUseFullPrecisionUVs
 			&& bGenerateLightmapUVs == Other.bGenerateLightmapUVs
 			&& MinLightmapResolution == Other.MinLightmapResolution
diff --git a/Engine/Source/Runtime/Engine/Classes/Engine/Level.h b/Engine/Source/Runtime/Engine/Classes/Engine/Level.h
--- a/Engine/Source/Runtime/Engine/Classes/Engine/Level.h
+++ b/Engine/Source/Runtime/Engine/Classes/Engine/Level.h
@@ -398,20 +398,11 @@ public:
 	UPROPERTY()
 	TArray<FGuid> StreamingTextureGuids;
 
-	/** Static information used by texture streaming code, generated during PreSave									*/
-	TMap<UTexture2D*,TArray<FStreamableTextureInstance> >	TextureToInstancesMap;
-
-	/** Information about textures on dynamic primitives. Used by texture streaming code, generated during PreSave.		*/
-	TMap<TWeakObjectPtr<UPrimitiveComponent>,TArray<FDynamicTextureInstance> >	DynamicTextureInstances;
-
-	/** Set of textures used by PrimitiveComponents that have bForceMipStreaming checked. */
-	TMap<UTexture2D*,bool>									ForceStreamTextures;
-
 	/** Index into Actors array pointing to first net relevant actor. Used as an optimization for FActorIterator	*/
-	int32											iFirstNetRelevantActor;
+	int32										iFirstNetRelevantActor;
 
 	/** Data structures for holding the tick functions **/
-	class FTickTaskLevel*		TickTaskLevel;
+	class FTickTaskLevel*						TickTaskLevel;
 
 	/** 
 	* The precomputed light information for this level.  
@@ -434,9 +425,6 @@ public:
 	/** Whether the geometry needs to be rebuilt for correct lighting */
 	uint32										bGeometryDirtyForLighting:1;
 
-	/** Has texture streaming been built */
-	uint32										bTextureStreamingBuilt:1;
-
 	/** Whether a level transform rotation was applied since the texture streaming builds. Invalidates the precomputed streaming bounds. */
 	UPROPERTY()
 	uint32 										bTextureStreamingRotationChanged : 1;
@@ -651,8 +639,6 @@ public:
 	 */
 	void RouteActorInitialize();
 
-
-
 	/**
 	 * Rebuilds static streaming data for all levels in the specified UWorld.
 	 *
@@ -662,22 +648,6 @@ public:
 	 */
 	ENGINE_API static void BuildStreamingData(UWorld* World, ULevel* TargetLevel=NULL, UTexture2D* TargetTexture=NULL);
 
-	/**
-	 * Rebuilds static streaming data for this level.
-	 *
-	 * @param TargetTexture			[opt] Specifies a single texture to process. If NULL, all textures will be processed.
-	 */
-	void BuildStreamingData(UTexture2D* TargetTexture=NULL);
-
-	/**
-	 * Clamp lightmap and shadowmap texelfactors to 20-80% range.
-	 * This is to prevent very low-res or high-res charts to dominate otherwise decent streaming.
-	 */
-	void NormalizeLightmapTexelFactor();
-
-	/** Retrieves the array of streamable texture isntances. */
-	ENGINE_API TArray<FStreamableTextureInstance>* GetStreamableTextureInstances(UTexture2D*& TargetTexture);
-
 	/**
 	* Deprecated. Returns the default brush for this level.
 	*
diff --git a/Engine/Source/Runtime/Engine/Classes/Engine/Texture2D.h b/Engine/Source/Runtime/Engine/Classes/Engine/Texture2D.h
--- a/Engine/Source/Runtime/Engine/Classes/Engine/Texture2D.h
+++ b/Engine/Source/Runtime/Engine/Classes/Engine/Texture2D.h
@@ -218,6 +218,9 @@ public:
 		return PlatformData->Mips;
 	}
 
+
+	FORCEINLINE int32 GetStreamingIndex() const { return StreamingIndex; }
+
 private:
 	/** The minimum number of mips that must be resident in memory (cannot be streamed). */
 	static ENGINE_API int32 GMinTextureResidentMipCount;
diff --git a/Engine/Source/Runtime/Engine/Classes/Engine/TextureStreamingTypes.h b/Engine/Source/Runtime/Engine/Classes/Engine/TextureStreamingTypes.h
--- a/Engine/Source/Runtime/Engine/Classes/Engine/TextureStreamingTypes.h
+++ b/Engine/Source/Runtime/Engine/Classes/Engine/TextureStreamingTypes.h
@@ -18,7 +18,7 @@ struct FStreamingTexturePrimitiveInfo
 	GENERATED_USTRUCT_BODY()
 
 	UPROPERTY()
-	UTexture* Texture;
+	UTexture2D* Texture;
 
 	UPROPERTY()
 	FBoxSphereBounds Bounds;
@@ -38,11 +38,12 @@ struct FStreamingTexturePrimitiveInfo
 	 *	Set this struct to match the packed params.
 	 *
 	 *	@param	InTexture			The texture as refered by the packed info.
+	 * 	@param	ExtraScale			Extra scale to be applied to the texcoord world size.
 	 *	@param	RefBounds			The reference bounds used to unpack the relative box.
 	 *	@param	Info				The packed params.
 	 *	@param	bUseRelativeBox		true if the relative box is relevant. Could be irrelevant if a level transform was applied after the texture streaming build.
 	 */
-	ENGINE_API void UnPackFrom(UTexture2D* InTexture, const FBoxSphereBounds& RefBounds, const FStreamingTextureBuildInfo& Info, bool bUseRelativeBox);
+	ENGINE_API void UnPackFrom(UTexture2D* InTexture, float ExtraScale, const FBoxSphereBounds& RefBounds, const FStreamingTextureBuildInfo& Info, bool bUseRelativeBox);
 };
 
 /** 
@@ -184,10 +185,10 @@ class FStreamingTextureLevelContext
 	const TArray<FStreamingTextureBuildInfo>* ComponentBuildData;
 	/** The last bound component bounds. */
 	FBoxSphereBounds ComponentBounds;
-	/** The last bound component mesh scale. */
-	float ComponentMeshScale;
-	/** The last bound component streaming scale. */
-	float ComponentStreamingScale;
+	/** The last bound component precomputed data scale. */
+	float ComponentPrecomputedDataScale;
+	/** The last bound component streaming fallback scale. */
+	float ComponentFallbackScale;
 
 	struct FTextureBoundState
 	{
@@ -210,7 +211,7 @@ public:
 	FStreamingTextureLevelContext(const ULevel* InLevel = nullptr);
 	~FStreamingTextureLevelContext();
 
-	void BindComponent(const TArray<FStreamingTextureBuildInfo>* BuildData, const FBoxSphereBounds& Bounds, float MeshScale, float StreamingScale);
+	void BindComponent(const TArray<FStreamingTextureBuildInfo>* BuildData, const FBoxSphereBounds& Bounds, float PrecomputedDataScale, float FallbackScale);
 	void Process(const TArray<UTexture*>& InTextures, TArray<FStreamingTexturePrimitiveInfo>& OutInfos);
 };
 
