remove mesh particle local to world transform from interpolants when it's not being used in the material

diff --git a/Engine/Shaders/MeshParticleVertexFactory.usf b/Engine/Shaders/MeshParticleVertexFactory.usf
--- a/Engine/Shaders/MeshParticleVertexFactory.usf
+++ b/Engine/Shaders/MeshParticleVertexFactory.usf
@@ -116,9 +116,10 @@ struct FVertexFactoryInterpolantsVSToPS
 
 	/** Relative alive time of the particle */
 	float ParticleTime : PARTICLE_TIME;
-
+#if NEEDS_PARTICLE_TRANSFORM
 	float4 LocalToWorld[4] : PARTICLE_TRANSFORM;
 #endif
+#endif
 
 #if INSTANCED_STEREO
 	nointerpolation uint PackedEyeIndex : PACKED_EYE_INDEX;
@@ -155,7 +156,9 @@ struct FVertexFactoryIntermediates
 	/** Relative time. */
 	float RelativeTime;
 
+#if NEEDS_PARTICLE_TRANSFORM
 	float4 LocalToWorld[4];
+#endif
 };
 
 float4x4 GetParticleTransform(FVertexFactoryInput Input)
@@ -256,10 +259,12 @@ FMaterialPixelParameters GetMaterialPixelParameters(FVertexFactoryInterpolantsVS
 	Result.Particle.PositionAndSize.w = 1;
 	Result.Particle.Velocity = Interpolants.ParticleVelocity;
 	Result.Particle.RelativeTime = Interpolants.ParticleTime;
+#if NEEDS_PARTICLE_TRANSFORM
 	Result.Particle.LocalToWorld[0] = Interpolants.LocalToWorld[0];
 	Result.Particle.LocalToWorld[1] = Interpolants.LocalToWorld[1];
 	Result.Particle.LocalToWorld[2] = Interpolants.LocalToWorld[2];
 	Result.Particle.LocalToWorld[3] = Interpolants.LocalToWorld[3];
+#endif
 #endif
 
 	Result.Particle.MotionBlurFade = 1.0f;
@@ -287,10 +292,12 @@ FMaterialVertexParameters GetMaterialVertexParameters(FVertexFactoryInput Input,
 
 	float4x4 Transform = GetParticleTransform(Input);
 	Result.InstanceLocalToWorld = Transform;
+#if NEEDS_PARTICLE_TRANSFORM
 	Result.Particle.LocalToWorld[0] = Transform[0];
 	Result.Particle.LocalToWorld[1] = Transform[1];
 	Result.Particle.LocalToWorld[2] = Transform[2];
 	Result.Particle.LocalToWorld[3] = Transform[3];
+#endif
 
 	Result.TangentToWorld = mul(TangentToLocal, (float3x3)Transform ); 
 
@@ -350,10 +357,12 @@ FVertexFactoryIntermediates GetVertexFactoryIntermediates(FVertexFactoryInput In
 #if PARTICLE_MESH_INSTANCED
 	float3 ParticleWorldPosition = float3( Input.Transform1.w, Input.Transform2.w, Input.Transform3.w);
 	float4x4 LocalToWorldMat = float4x4(Input.Transform1, Input.Transform2, Input.Transform3, float4(0.0f, 0.0f, 0.0f, 1.0f));
+#if NEEDS_PARTICLE_TRANSFORM
 	Intermediates.LocalToWorld[0] = LocalToWorldMat[0];
 	Intermediates.LocalToWorld[1] = LocalToWorldMat[1];
 	Intermediates.LocalToWorld[2] = LocalToWorldMat[2];
 	Intermediates.LocalToWorld[3] = LocalToWorldMat[3];
+#endif
 #else
 	float3 ParticleWorldPosition = float3( Transform1.w, Transform2.w, Transform3.w);
 #endif
@@ -521,8 +530,10 @@ FVertexFactoryInterpolantsVSToPS VertexFactoryGetInterpolantsVSToPS(FVertexFacto
 	Interpolants.ParticlePosition = Intermediates.ParticleWorldPosition;
 	Interpolants.ParticleVelocity = Intermediates.ParticleVelocity;
 	Interpolants.ParticleTime = Intermediates.RelativeTime;
+#if NEEDS_PARTICLE_TRANSFORM
 	Interpolants.LocalToWorld = Intermediates.LocalToWorld;
 #endif
+#endif
 
 #if INSTANCED_STEREO
 	Interpolants.PackedEyeIndex = 0;
diff --git a/Engine/Source/Runtime/Engine/Private/Materials/HLSLMaterialTranslator.h b/Engine/Source/Runtime/Engine/Private/Materials/HLSLMaterialTranslator.h
--- a/Engine/Source/Runtime/Engine/Private/Materials/HLSLMaterialTranslator.h
+++ b/Engine/Source/Runtime/Engine/Private/Materials/HLSLMaterialTranslator.h
@@ -203,6 +203,9 @@ protected:
 	uint32 bUsesVertexColor : 1;
 	/** true if the material reads particle color in the pixel shader. */
 	uint32 bUsesParticleColor : 1;
+	/** true if the material reads mesh particle transform in the pixel shader. */
+	uint32 bUsesParticleTransform : 1;
+
 	uint32 bUsesTransformVector : 1;
 	// True if the current property requires last frame's information
 	uint32 bCompilingPreviousFrame : 1;
@@ -253,6 +256,7 @@ public:
 	,	bUsesVertexColor(false)
 	,	bUsesParticleColor(false)
 	,	bUsesTransformVector(false)
+	,	bUsesParticleTransform(false)
 	,	bCompilingPreviousFrame(false)
 	,	bOutputsBasePassVelocities(true)
 	,	bUsesPixelDepthOffset(false)
@@ -734,7 +738,8 @@ public:
 		OutEnvironment.SetDefine(TEXT("MATERIAL_ATMOSPHERIC_FOG"), bUsesAtmosphericFog);
 		OutEnvironment.SetDefine(TEXT("INTERPOLATE_VERTEX_COLOR"), bUsesVertexColor); 
 		OutEnvironment.SetDefine(TEXT("NEEDS_PARTICLE_COLOR"), bUsesParticleColor); 
-		OutEnvironment.SetDefine(TEXT("USES_TRANSFORM_VECTOR"), bUsesTransformVector); 
+		OutEnvironment.SetDefine(TEXT("NEEDS_PARTICLE_TRANSFORM"), bUsesParticleTransform);
+		OutEnvironment.SetDefine(TEXT("USES_TRANSFORM_VECTOR"), bUsesTransformVector);
 		OutEnvironment.SetDefine(TEXT("WANT_PIXEL_DEPTH_OFFSET"), bUsesPixelDepthOffset); 
 		// Distortion uses tangent space transform 
 		OutEnvironment.SetDefine(TEXT("USES_DISTORTION"), Material->IsDistorted()); 
@@ -3756,6 +3761,7 @@ protected:
 				else if (DestCoordBasis == MCB_MeshParticle)
 				{
 					CodeStr = TEXT("mul(<A>, <MATRIX>(Parameters.Particle.LocalToWorld))");
+					bUsesParticleTransform = true;
 				}
 
 				// else use MCB_TranslatedWorld as intermediary basis
@@ -3787,6 +3793,7 @@ protected:
 				if (DestCoordBasis == MCB_World)
 				{
 					CodeStr = TEXT("mul(<MATRIX>(Parameters.Particle.LocalToWorld), <A>)");
+					bUsesParticleTransform = true;
 				}
 				else
 				{
