Copying //UE4/Dev-Rendering to //UE4/Dev-Main (Source: //UE4/Dev-Rendering @ 2998063)

==========================
MAJOR FEATURES + CHANGES
==========================

Change 2981877 on 2016/05/18 by [REDACTED_PERSON]

	DR - Fix some PVS warnings
	- Removed 'uniform' as a keyword on hlslcc as it's ignored/causes issues/doesn't really optimize

Change 2981966 on 2016/05/18 by [REDACTED_PERSON]

	DR - Fix OpenGL crash quitting editor

Change 2982072 on 2016/05/18 by [REDACTED_PERSON]

	Fixed a "Build Texture Streaming" crash related to custom outputs.
	Fixed issue with debug view mode and translucent primitives
	Fix bug with visibility that made some texture low res.
	Enabled per instance visibility in the texture streaming order (in game only).
	Improved InvestigateTexture Logs.
	Tweaked the mip computations from screen size.

Change 2982077 on 2016/05/18 by [REDACTED_PERSON]

	Removed debug options!

Change 2982108 on 2016/05/18 by [REDACTED_PERSON]

	moving AMD hacks to console variables

Change 2982422 on 2016/05/18 by [REDACTED_PERSON]

	UE4 - Potential crash fix on foliage occlusion queries and reflection captures.

Change 2982547 on 2016/05/18 by [REDACTED_PERSON]

	UE-26409 Crash when Light Propagation Volume Plugin is disabled on a Project

Change 2982548 on 2016/05/18 by [REDACTED_PERSON]

	Refactored MRT and outer Pixel Shader output to a struct
	to allow it passed in/out of functions allowing for a more readable code (less #ifdefs, reducing the boolean hell)
	Needed for upcoming MeshDecals

Change 2982601 on 2016/05/18 by [REDACTED_PERSON]

	Movable skylight now matches stationary for subsurface shading models
	* Two sided was broken in 4.11, Subsurface had never been handled

Change 2982603 on 2016/05/18 by [REDACTED_PERSON]

	DrawMaterialToRenderTarget / BeginDrawCanvasToRenderTarget now work correctly with material parameter collections and Time
	* FCanvas stores an optional scene to render to, UWorld caches the UCanvas needed for implementing these functions

Change 2982618 on 2016/05/18 by [REDACTED_PERSON]

	Better categories for some Rendering project settings

Change 2982619 on 2016/05/18 by [REDACTED_PERSON]

	Scene capture 2d improvements
	* Orthographic projection supported
	* Opacity is now captured in alpha, allows partial rendering in a scene capture and compositing into another scene later
	* Various GBuffer attributes are now available to be captured, including depth
	* Changed Blueprint capture function to CaptureScene, which happens immediately (was previously deferred), allowing multiple captures with different parameters

Change 2982664 on 2016/05/18 by [REDACTED_PERSON]

	Fading out planar reflections based on roughness since they don't have support for variable roughness (fade starts at .2 roughness, ends at .3)

Change 2982684 on 2016/05/18 by [REDACTED_PERSON]

	polish ImageValidator

Change 2982685 on 2016/05/18 by [REDACTED_PERSON]

	show testimage on sm4 as well

Change 2982736 on 2016/05/18 by [REDACTED_PERSON]

	Improved overbudget retention logic.
	Updated ListStreamingTextures stats.

Change 2982854 on 2016/05/18 by [REDACTED_PERSON]

	ImageValidator can now save/load IVxml file

Change 2982863 on 2016/05/18 by [REDACTED_PERSON]

	Fixed shader compile error

Change 2982864 on 2016/05/18 by [REDACTED_PERSON]

	Removed deprecation message

Change 2982927 on 2016/05/18 by [REDACTED_PERSON]

	ImageValidator is now sorting by time

Change 2983743 on 2016/05/19 by [REDACTED_PERSON]

	Saturated tri-planar UV function outputs to prevent negative blending

Change 2983747 on 2016/05/19 by [REDACTED_PERSON]

	Fix for morph curves not getting applied to meshes in cooked builds (smart names were not being corrected).

Change 2984008 on 2016/05/19 by [REDACTED_PERSON]

	New contact shadows feature.

	Hair uses ray cast for nonshadow lights

Change 2984009 on 2016/05/19 by [REDACTED_PERSON]

	changed to ShadowedBits

Change 2984054 on 2016/05/19 by [REDACTED_PERSON]

	Deleted old motion blur

Change 2984420 on 2016/05/19 by [REDACTED_PERSON]

	Shorter display name for WorldPosition material node

Change 2984423 on 2016/05/19 by [REDACTED_PERSON]

	Fixed WorldPosition and ScreenPosition for downsampled separate translucency by scaling SvPosition before those computations.  The View uniform buffer still contains incorrect buffer sizes for this pass.

Change 2984432 on 2016/05/19 by [REDACTED_PERSON]

	DR - Disable PPCombineLUT for Vulkan to work around glslang issue (will need to be reenabled when adding SM4/5 path)

Change 2985415 on 2016/05/20 by [REDACTED_PERSON]

	Added a Texture2D exporter for .hdr
	Texture2D exporters now implement SupportsObject properly, so you only see extensions that are valid based on the format

Change 2985439 on 2016/05/20 by [REDACTED_PERSON]

	Scene color alpha clear value validation

Change 2987173 on 2016/05/23 by [REDACTED_PERSON]

	ImageValidator Report with Thumbnail (cannot be copied and pasted into email), non Thumbnail version could be.

Change 2987248 on 2016/05/23 by [REDACTED_PERSON]

	ImageValidator: added Summary, removed timer hack

Change 2987369 on 2016/05/23 by [REDACTED_PERSON]

	ImageValidator polish

Change 2987390 on 2016/05/23 by [REDACTED_PERSON]

	Improvement to temporal aa sharpness and speed.

Change 2988038 on 2016/05/24 by [REDACTED_PERSON]

	Merging //UE4/Dev-Main@2987977 to Dev-Rendering (//UE4/Dev-Rendering)

Change 2988304 on 2016/05/24 by [REDACTED_PERSON]

	added const to prevent coding errors

Change 2988332 on 2016/05/24 by [REDACTED_PERSON]

	Fixed motion blur crash on SM4

Change 2988446 on 2016/05/24 by [REDACTED_PERSON]

	nicer UI

Change 2988990 on 2016/05/24 by [REDACTED_PERSON]

	fixed UE-31227 Building lighting produces bad results

Change 2989729 on 2016/05/25 by [REDACTED_PERSON]

	Fixed lightmaps and shadowmaps having low resolutions after building the lighting.

Change 2989752 on 2016/05/25 by [REDACTED_PERSON]

	CVar to disable/freeze GPU particle simulation

Change 2989811 on 2016/05/25 by [REDACTED_PERSON]

	Making use of MATERIALBLENDING_ANY_TRANSLUCENT

Change 2989812 on 2016/05/25 by [REDACTED_PERSON]

	Hide DFGI show flags from UI

Change 2989901 on 2016/05/25 by [REDACTED_PERSON]

	Height fog now works properly in planar reflections
	* The ray used for computing fog is first clipped by the reflection plane

Change 2989904 on 2016/05/25 by [REDACTED_PERSON]

	Always use PF_FloatRGBA for LightAccumulation to guarantee alpha channel and negative range

Change 2989991 on 2016/05/25 by [REDACTED_PERSON]

	Improved usability for DBuffer Decals
	* 'Show Decals' works correctly, previously would fetch from uninitialized textures
	* DBuffer being enabled forces a full prepass, previously decals would render incorrectly unless correct settings of r.EarlyZPass were used
	* Improved the PrePass draw event to indicate whether it's full or partial
	* Materials using DBuffer blend modes will fail to compile when the DBuffer project setting is disabled, instead of just being invisible
	* r.EarlyZPass can now be changed at runtime, which is useful for profiling

Change 2990008 on 2016/05/25 by [REDACTED_PERSON]

	Fixed capsule shadows on skeletal meshes with scaling

Change 2990274 on 2016/05/25 by [REDACTED_PERSON]

	Fixed DFAO (from cl 2961310)

Change 2990304 on 2016/05/25 by [REDACTED_PERSON]

	OR-22233 GPU Sprites invisible unless solo'd

Change 2990309 on 2016/05/25 by [REDACTED_PERSON]

	Added SubDSurface actor (using CPU code of OpenSubDiv), component, asset as starting point for more work in that direction,
	(Early work in progress)

Change 2990363 on 2016/05/25 by [REDACTED_PERSON]

	Spreading precomputed visibility to neighbors now uses a 2d grid to find neighbors, speeds up the process for 800k cells from 40 mins to 20s

Change 2990392 on 2016/05/25 by [REDACTED_PERSON]

	Added r.AOSpecularOcclusionMode, which determines how specular should be occluded by DFAO
	0: Apply non-directional AO to specular.
	1: (default) Intersect the reflection cone with the unoccluded cone produced by DFAO.  This gives more accurate occlusion than 0, but can bring out DFAO sampling artifacts.
	2: (experimental) Cone trace through distance fields along the reflection vector.  Costs about the same as DFAO again because more cone tracing is done, but produces more accurate occlusion.

Change 2990454 on 2016/05/25 by [REDACTED_PERSON]

	polish readme

Change 2990610 on 2016/05/25 by [REDACTED_PERSON]

	fixed building with VS2015, the right OpenSubDiv .lib files are missing, temporarily disabled the relevant code

Change 2990754 on 2016/05/25 by [REDACTED_PERSON]

	Fix compiler warning: C4456: declaration of 'NewStaticMesh' hides previous local declaration.

Change 2990801 on 2016/05/25 by [REDACTED_PERSON]

	Only allocate reflection capture cubemaps if we are actually doing a reflection capture. The old approach always allocated based on CVar and worked with a warning on D3D since if SRC was larger than the DST, the extra SRC area was dropped. New approach only allocates when necessary and is correctly sized everytime. Also hardened access to CVar and what users are allowed to set.

Change 2991169 on 2016/05/26 by [REDACTED_PERSON]

	fixed compiler warning
	WARNING: Non-editor build cannot depend on non-redistributable modules. Details:

Change 2991238 on 2016/05/26 by [REDACTED_PERSON]

	fixed build
	fatal error C1083: Cannot open include file: 'RawMesh.h': No such file or directory  (when RawMesh is not part of PrivateDependencyModuleNames in Engine.Build.cs)

Change 2991726 on 2016/05/26 by [REDACTED_PERSON]

	Subsurface materials are now handled with simple forward shading

Change 2991727 on 2016/05/26 by [REDACTED_PERSON]

	Emissive decals are now supported with simple forward shading

Change 2994849 on 2016/05/31 by [REDACTED_PERSON]

	Disabled fix for WorldPosition and ScreenPosition in downsampled separate translucency, since it breaks GetScreenAlignedUV (used in DepthFade), since the uniform buffer still contains full res buffer sizes

Change 2997243 on 2016/06/01 by [REDACTED_PERSON]

	Merging //UE4/Dev-Main@2996565 to Dev-Rendering (//UE4/Dev-Rendering)

diff --git a/Engine/Config/BaseEngine.ini b/Engine/Config/BaseEngine.ini
--- a/Engine/Config/BaseEngine.ini
+++ b/Engine/Config/BaseEngine.ini
@@ -288,6 +288,8 @@ bUseStreamingPause=false
 +ActiveClassRedirects=(OldClassName="SkeletalPhysicsActor",NewClassName="SkeletalMeshActor")
 +ActiveClassRedirects=(OldClassName="SingleAnimSkeletalActor",NewClassName="SkeletalMeshActor")
 +ActiveClassRedirects=(OldClassName="SingleAnimSkeletalComponent",NewClassName="SkeletalMeshComponent")
++K2FieldRedirects=(OldFieldName="SceneCaptureComponent2D.UpdateContent",NewFieldName="SceneCaptureComponent2D.CaptureScene")
++K2FieldRedirects=(OldFieldName="SceneCaptureComponentCube.UpdateContent",NewFieldName="SceneCaptureComponentCube.CaptureScene")
 +K2FieldRedirects=(OldFieldName="SingleAnimSkeletalComponent.SetAnim",NewFieldName="SkeletalMeshComponent.SetAnimation")
 +K2FieldRedirects=(OldFieldName="SingleAnimSkeletalComponent.Play",NewFieldName="SkeletalMeshComponent.Play")
 +K2FieldRedirects=(OldFieldName="SingleAnimSkeletalComponent.Stop",NewFieldName="SkeletalMeshComponent.Stop")
diff --git a/Engine/Plugins/Experimental/StereoPanorama/Source/StereoPanorama/Private/SceneCapturer.cpp b/Engine/Plugins/Experimental/StereoPanorama/Source/StereoPanorama/Private/SceneCapturer.cpp
--- a/Engine/Plugins/Experimental/StereoPanorama/Source/StereoPanorama/Private/SceneCapturer.cpp
+++ b/Engine/Plugins/Experimental/StereoPanorama/Source/StereoPanorama/Private/SceneCapturer.cpp
@@ -292,9 +292,9 @@ void USceneCapturer::SetPositionAndRotation( int32 CurrentHorizontalStep, int32
     }
 
 	LeftEyeCaptureComponents[CaptureIndex]->SetWorldLocationAndRotation( StartLocation - Offset, Rotation );
-    LeftEyeCaptureComponents[CaptureIndex]->UpdateContent();
+    LeftEyeCaptureComponents[CaptureIndex]->CaptureSceneDeferred();
 	RightEyeCaptureComponents[CaptureIndex]->SetWorldLocationAndRotation( StartLocation + Offset, Rotation );
-    RightEyeCaptureComponents[CaptureIndex]->UpdateContent();
+    RightEyeCaptureComponents[CaptureIndex]->CaptureSceneDeferred();
 }
 
 void USceneCapturer::ValidateParameters()
diff --git a/Engine/Shaders/BasePassCommon.usf b/Engine/Shaders/BasePassCommon.usf
--- a/Engine/Shaders/BasePassCommon.usf
+++ b/Engine/Shaders/BasePassCommon.usf
@@ -10,7 +10,7 @@
 // Translucent materials need to compute fogging in the forward shading pass
 // Materials that read from scene color skip getting fogged, because the contents of the scene color lookup have already been fogged
 // This is not foolproof, as any additional color the material adds will then not be fogged correctly
-#define NEEDS_BASEPASS_FOGGING			(ENABLE_TRANSLUCENCY_VERTEX_FOG && ((MATERIALBLENDING_TRANSLUCENT || MATERIALBLENDING_ADDITIVE || MATERIALBLENDING_MODULATE) && !MATERIAL_USES_SCENE_COLOR_COPY))
+#define NEEDS_BASEPASS_FOGGING			(ENABLE_TRANSLUCENCY_VERTEX_FOG && (MATERIALBLENDING_ANY_TRANSLUCENT && !MATERIAL_USES_SCENE_COLOR_COPY))
 #define NEEDS_LIGHTMAP					(NEEDS_LIGHTMAP_COORDINATE)
 
 #define USES_GBUFFER						(FEATURE_LEVEL >= FEATURE_LEVEL_SM4 && (MATERIALBLENDING_SOLID || MATERIALBLENDING_MASKED) && !SIMPLE_FORWARD_SHADING)
@@ -20,7 +20,7 @@
 
 // Based on GetPrecomputedShadowMasks()
 // Note: WRITES_PRECSHADOWFACTOR_TO_GBUFFER is currently disabled because we use the precomputed shadow factor GBuffer outside of STATICLIGHTING_TEXTUREMASK to store UseSingleSampleShadowFromStationaryLights
-#define WRITES_PRECSHADOWFACTOR_TO_GBUFFER	(USES_GBUFFER && ALLOW_STATIC_LIGHTING && (!SELECTIVE_BASEPASS_OUTPUTS || (1)))
+#define WRITES_PRECSHADOWFACTOR_TO_GBUFFER	(USES_GBUFFER && ALLOW_STATIC_LIGHTING)
 #define WRITES_PRECSHADOWFACTOR_ZERO		(USES_GBUFFER && ALLOW_STATIC_LIGHTING && (!SELECTIVE_BASEPASS_OUTPUTS || HQ_TEXTURE_LIGHTMAP || LQ_TEXTURE_LIGHTMAP))
 
 // If a primitive has static lighting, we assume it is not moving. If it is, it will be rerendered in an extra renderpass.
diff --git a/Engine/Shaders/BasePassForForwardShadingPixelShader.usf b/Engine/Shaders/BasePassForForwardShadingPixelShader.usf
--- a/Engine/Shaders/BasePassForForwardShadingPixelShader.usf
+++ b/Engine/Shaders/BasePassForForwardShadingPixelShader.usf
@@ -193,9 +193,9 @@ half3 BlendReflectionCaptures(FMaterialPixelParameters MaterialParameters, half
 }
 #endif // HQ_REFLECTIONS
 
-float4 GetPlanarReflection(float3 WorldPosition, float3 WorldNormal)
+float4 GetPlanarReflection(float3 WorldPosition, float3 WorldNormal, float Roughness)
 {
-	return ComputePlanarReflections(WorldPosition, WorldNormal, PlanarReflectionSampler);
+	return ComputePlanarReflections(WorldPosition, WorldNormal, Roughness, PlanarReflectionSampler);
 }
 
 half3 GetImageBasedReflectionLighting(FMaterialPixelParameters MaterialParameters, half Roughness)
@@ -467,7 +467,7 @@ void Main(
 				BRANCH
 				if (abs(dot(ReflectionPlane.xyz, 1)) > .0001f)
 				{
-					half4 PlanarReflection = GetPlanarReflection(MaterialParameters.AbsoluteWorldPosition, MaterialParameters.WorldNormal);
+					half4 PlanarReflection = GetPlanarReflection(MaterialParameters.AbsoluteWorldPosition, MaterialParameters.WorldNormal, Roughness);
 					// Planar reflections win over reflection environment
 					SpecularIBL = lerp(SpecularIBL, PlanarReflection.rgb, PlanarReflection.a);
 				}
diff --git a/Engine/Shaders/BasePassPixelShader.usf b/Engine/Shaders/BasePassPixelShader.usf
--- a/Engine/Shaders/BasePassPixelShader.usf
+++ b/Engine/Shaders/BasePassPixelShader.usf
@@ -6,8 +6,8 @@
 
 #include "Common.usf"
 #include "SHCommon.usf"
-#include "BasePassCommon.usf"
 #include "Material.usf"
+#include "BasePassCommon.usf"
 #include "VertexFactory.usf"
 #include "LightmapCommon.usf"  
 #include "ReflectionEnvironmentShared.usf"
@@ -32,6 +32,11 @@
 	 
 #endif
 
+#if MATERIALBLENDING_ANY_TRANSLUCENT
+// Downsampled translucency scale (eg 2.0 for r.SeparateTranslucencyScreenPercentage=50)
+float DownsampleFactorFromSceneBufferSize;
+#endif
+
 #if MATERIAL_SHADINGMODEL_HAIR || SIMPLE_FORWARD_DIRECTIONAL_LIGHT
 #include "ShadingModels.usf"
 #endif
@@ -198,7 +203,7 @@ half3 GetImageBasedReflectionLighting(FMaterialPixelParameters MaterialParameter
 	if (abs(dot(ReflectionPlane.xyz, 1)) > .0001f)
 	{
 		// Reuse ReflectionCubemapSampler to avoid reducing the sampler count available to artists
-		float4 PlanarReflection = ComputePlanarReflections(MaterialParameters.AbsoluteWorldPosition, MaterialParameters.WorldNormal, ReflectionCubemapSampler);
+		float4 PlanarReflection = ComputePlanarReflections(MaterialParameters.AbsoluteWorldPosition, MaterialParameters.WorldNormal, Roughness, ReflectionCubemapSampler);
 		// Planar reflections win over SSR and reflection environment
 		SpecularLighting = PlanarReflection.rgb + (1 - PlanarReflection.a) * SpecularLighting;
 	}
@@ -576,7 +581,7 @@ void GetPrecomputedIndirectLightingAndSkyLight(
 
 #if SIMPLE_FORWARD_DIRECTIONAL_LIGHT
 
-float3 GetSimpleForwardLightingDirectionalLight(float3 DiffuseColor, float3 SpecularColor, float Roughness, float3 WorldNormal, float3 CameraVector)
+float3 GetSimpleForwardLightingDirectionalLight(FGBufferData GBuffer, float3 DiffuseColor, float3 SpecularColor, float Roughness, float3 WorldNormal, float3 CameraVector)
 {
 	float3 V = CameraVector;
 	float3 N = WorldNormal;
@@ -584,9 +589,11 @@ float3 GetSimpleForwardLightingDirectionalLight(float3 DiffuseColor, float3 Spec
 	float NoL = saturate( dot( N, L ) );
 
 	float3 LightColor = Frame.DirectionalLightColor.rgb * PI;
-	
 	// Not computing specular, material was forced fully rough
-	return LightColor * NoL * StandardShading(DiffuseColor, SpecularColor, Roughness.xxx, 1.0f, L, V, N, float2(1, 0));
+	float3 DiffuseLighting = StandardShading(DiffuseColor, SpecularColor, Roughness.xxx, 1.0f, L, V, N, float2(1, 0));
+	float3 SubsurfaceLighting = SubsurfaceShading(GBuffer, L, V, N, 1.0f, uint2(0, 0));
+
+	return LightColor * NoL * (DiffuseLighting + SubsurfaceLighting);
 }
 
 #endif
@@ -692,48 +699,11 @@ uint GetSelectiveOutputMask()
 }
 #endif // USES_GBUFFER
 
-#if OUTPUT_GBUFFER_VELOCITY
-#define GBUFFER_D_TARGET SV_Target5
-#define GBUFFER_E_TARGET SV_Target6 // Only used with ALLOW_STATIC_LIGHTING
-#else
-#define GBUFFER_D_TARGET SV_Target4
-#define GBUFFER_E_TARGET SV_Target5
-#endif
-
-void Main(
+// is called in MainPS() from PixelShaderOutputCommon.usf
+void FPixelShaderInOut_MainPS(
 	FVertexFactoryInterpolantsVSToPS Interpolants,
 	FBasePassInterpolantsVSToPS BasePassInterpolants,
-	in INPUT_POSITION_QUALIFIERS float4 SvPosition : SV_Position
-	OPTIONAL_IsFrontFace
-	,out float4 OutColor		: SV_Target0
-
-// Two options control the render target bindings : OUTPUT_GBUFFER_VELOCITY and ALLOW_STATIC_LIGHTING
-#if USES_GBUFFER
-	,out float4 OutGBufferA : SV_Target1 
-	,out float4 OutGBufferB : SV_Target2
-	,out float4 OutGBufferC : SV_Target3
-
-// Big issue here on PC with AMD hardware : putting the velocity buffer after an optionnal RT fails!!!
-#if WRITES_VELOCITY_TO_GBUFFER
-	,out float4 OutGBufferVelocity : SV_Target4
-#endif
-
-#if WRITES_CUSTOMDATA_TO_GBUFFER
-	,out float4 OutGBufferD : GBUFFER_D_TARGET
-#endif
-
-#if WRITES_PRECSHADOWFACTOR_TO_GBUFFER
-	,out float4 OutGBufferE : GBUFFER_E_TARGET
-#endif
-#endif // USES_GBUFFER
-
-	OPTIONAL_OutDepthConservative
-
-#if EDITOR_ALPHA2COVERAGE != 0
-   ,in uint InCoverage : SV_Coverage
-   ,out uint OutCoverage : SV_Coverage
-#endif
-   )
+	inout FPixelShaderInOut InOut)
 {
 #if INSTANCED_STEREO
 	ResolvedView = ResolveView(GetEyeIndex(Interpolants.PackedEyeIndex));
@@ -741,21 +711,14 @@ void Main(
 	ResolvedView = ResolveView();
 #endif
 
-#if !WRITES_CUSTOMDATA_TO_GBUFFER
+	// Velocity
+	float4 OutVelocity = 0;
+	// CustomData
 	float4 OutGBufferD = 0;
-#endif
-#if !WRITES_PRECSHADOWFACTOR_TO_GBUFFER
+	// PreShadowFactor
 	float4 OutGBufferE = 0;
-#endif
-#if !WRITES_VELOCITY_TO_GBUFFER
-	float4 OutGBufferVelocity = 0;
-#endif
-
-#if EDITOR_ALPHA2COVERAGE != 0
-	OutCoverage = InCoverage;
-#endif
 
-	FMaterialPixelParameters MaterialParameters = GetMaterialPixelParameters(Interpolants, SvPosition);
+	FMaterialPixelParameters MaterialParameters = GetMaterialPixelParameters(Interpolants, InOut.SvPosition);
 	FPixelMaterialInputs PixelMaterialInputs;
 	
 #if HQ_TEXTURE_LIGHTMAP && USES_AO_MATERIAL_MASK && !MATERIAL_SHADINGMODEL_UNLIT
@@ -765,13 +728,25 @@ void Main(
 	MaterialParameters.AOMaterialMask = GetAOMaterialMask(LightmapUV0 * float2(1, 2));
 #endif
 
+	float4 SVPositionForComputingPositions = InOut.SvPosition;
+
+#if MATERIALBLENDING_ANY_TRANSLUCENT
+	// Modify SVPosition for passes that are rendering to a downsampled target so that ScreenPosition and WorldPosition will be computed correctly
+	//SVPositionForComputingPositions.xy *= DownsampleFactorFromSceneBufferSize;
+#endif
+
 	#if USE_WORLD_POSITION_EXCLUDING_SHADER_OFFSETS
 		{
-			float4 ScreenPosition = SvPositionToScreenPosition(SvPosition);
-			CalcMaterialParametersEx(MaterialParameters, PixelMaterialInputs, SvPosition, ScreenPosition, bIsFrontFace, SvPositionToResolvedTranslatedWorld(SvPosition), BasePassInterpolants.PixelPositionExcludingWPO);
+			float4 ScreenPosition = SvPositionToScreenPosition(SVPositionForComputingPositions);
+			float3 TranslatedWorldPosition = SvPositionToResolvedTranslatedWorld(SVPositionForComputingPositions);
+			CalcMaterialParametersEx(MaterialParameters, PixelMaterialInputs, InOut.SvPosition, ScreenPosition, InOut.bIsFrontFace, TranslatedWorldPosition, BasePassInterpolants.PixelPositionExcludingWPO);
 		}
 	#else
-		CalcMaterialParameters(MaterialParameters, PixelMaterialInputs, SvPosition, bIsFrontFace);
+		{
+			float4 ScreenPosition = SvPositionToScreenPosition(SVPositionForComputingPositions);
+			float3 TranslatedWorldPosition = SvPositionToResolvedTranslatedWorld(SVPositionForComputingPositions);
+			CalcMaterialParametersEx(MaterialParameters, PixelMaterialInputs, InOut.SvPosition, ScreenPosition, InOut.bIsFrontFace, TranslatedWorldPosition, TranslatedWorldPosition);
+		}
 	#endif
 
 #if USE_EDITOR_COMPOSITING && (FEATURE_LEVEL >= FEATURE_LEVEL_SM4 || MOBILE_EMULATION)
@@ -781,7 +756,7 @@ void Main(
 #endif
 		
 #if OUTPUT_PIXEL_DEPTH_OFFSET
-	ApplyPixelDepthOffsetForBasePass(MaterialParameters, PixelMaterialInputs, BasePassInterpolants, OutDepth);
+	ApplyPixelDepthOffsetForBasePass(MaterialParameters, PixelMaterialInputs, BasePassInterpolants, InOut.OutDepth);
 #endif
 
 	//Clip if the blend mode requires it.
@@ -816,9 +791,9 @@ void Main(
 	}
 #endif
 
-#if USE_DBUFFER && MATERIALDECALRESPONSEMASK && !(MATERIALBLENDING_TRANSLUCENT || MATERIALBLENDING_ADDITIVE || MATERIALBLENDING_MODULATE)
+#if USE_DBUFFER && MATERIALDECALRESPONSEMASK && !MATERIALBLENDING_ANY_TRANSLUCENT
 	// apply decals from the DBuffer
-	BRANCH if(Primitive.DecalReceiverMask > 0)
+	BRANCH if(Primitive.DecalReceiverMask > 0 && Frame.ShowDecalsMask > 0)
 	{
 		float2 NDC = MaterialParameters.ScreenPosition.xy / MaterialParameters.ScreenPosition.w;
 		
@@ -840,7 +815,7 @@ void Main(
 	float3 LocalBaseColor = BaseColor;
 	float LocalSpecular = Specular;
 
-#if MATERIAL_SHADINGMODEL_SUBSURFACE_PROFILE
+#if MATERIAL_SHADINGMODEL_SUBSURFACE_PROFILE && !SIMPLE_FORWARD_SHADING
 	// SubsurfaceProfile applies the BaseColor in a later pass. Any lighting output in the base pass needs
 	// to separate specualar and diffuse lighting in a checkerboard pattern
 	{
@@ -994,16 +969,16 @@ void Main(
 	#if !MATERIAL_SHADINGMODEL_UNLIT
 
 		float3 DiffuseDir = MaterialParameters.WorldNormal;
-		float3 EffectiveDiffuseColor = DiffuseColor;
+		float3 DiffuseColorForIndirect = DiffuseColor;
 
 		#if MATERIAL_SHADINGMODEL_SUBSURFACE || MATERIAL_SHADINGMODEL_PREINTEGRATED_SKIN
 			// Add subsurface energy to diffuse
 			//@todo - better subsurface handling for these shading models with skylight and precomputed GI
-			EffectiveDiffuseColor += SubsurfaceColor;
+			DiffuseColorForIndirect += SubsurfaceColor;
 		#endif
 
 		#if MATERIAL_SHADINGMODEL_CLOTH
-			EffectiveDiffuseColor += SubsurfaceColor * saturate( GetMaterialCustomData0(MaterialParameters) );
+			DiffuseColorForIndirect += SubsurfaceColor * saturate( GetMaterialCustomData0(MaterialParameters) );
 		#endif
 
 		#if MATERIAL_SHADINGMODEL_HAIR
@@ -1012,7 +987,7 @@ void Main(
 			float3 V = MaterialParameters.CameraVector;
 			float3 L = normalize( V - N * dot(V,N) );
 			DiffuseDir = L;
-			EffectiveDiffuseColor = PI * HairShading( GBuffer, L, V, N, 1, 0, 0.2, uint2(0,0) );
+			DiffuseColorForIndirect = PI * HairShading( GBuffer, L, V, N, 1, 0, 0.2, uint2(0,0) );
 		}
 		#endif
 
@@ -1020,13 +995,14 @@ void Main(
 		float3 SubsurfaceIndirectLighting;
 		GetPrecomputedIndirectLightingAndSkyLight(MaterialParameters, Interpolants, BasePassInterpolants, DiffuseDir, DiffuseIndirectLighting, SubsurfaceIndirectLighting, IndirectIrradiance);
 		
-		Color += (DiffuseIndirectLighting * EffectiveDiffuseColor + SubsurfaceIndirectLighting * SubsurfaceColor) * MaterialAO;
+		Color += (DiffuseIndirectLighting * DiffuseColorForIndirect + SubsurfaceIndirectLighting * SubsurfaceColor) * MaterialAO;
 
 
 		#if SIMPLE_FORWARD_DIRECTIONAL_LIGHT
 
 			float3 DirectionalLighting = GetSimpleForwardLightingDirectionalLight(
-				EffectiveDiffuseColor,
+				GBuffer,
+				DiffuseColor,
 				SpecularColor,
 				Roughness,
 				MaterialParameters.WorldNormal, 
@@ -1080,23 +1056,26 @@ void Main(
 	Color += Emissive;
 
 	#if MATERIALBLENDING_TRANSLUCENT
-		OutColor = half4(Color * VertexFog.a + VertexFog.rgb, Opacity);
-		OutColor = RETURN_COLOR(OutColor);
+		InOut.MRT[0] = half4(Color * VertexFog.a + VertexFog.rgb, Opacity);
+		InOut.MRT[0] = RETURN_COLOR(InOut.MRT[0]);
 	#elif MATERIALBLENDING_ADDITIVE
-		OutColor = half4(Color * VertexFog.a * Opacity, 0.0f);
-		OutColor = RETURN_COLOR(OutColor);
+		InOut.MRT[0] = half4(Color * VertexFog.a * Opacity, 0.0f);
+		InOut.MRT[0] = RETURN_COLOR(InOut.MRT[0]);
 	#elif MATERIALBLENDING_MODULATE
 		// RETURN_COLOR not needed with modulative blending
 		half3 FoggedColor = lerp(float3(1, 1, 1), Color, VertexFog.aaa * VertexFog.aaa);
-		OutColor = half4(FoggedColor, Opacity);
+		InOut.MRT[0] = half4(FoggedColor, Opacity);
 	#else
 		{
 			FLightAccumulator LightAccumulator = (FLightAccumulator)0;
 
 			LightAccumulator_Add(LightAccumulator, Color, 0, 1.0f);
-			// Note: Scene color alpha used in scene captures and planar reflections to track which parts have been rendered
-			// A value of 0 means that pixel has been rendered, 1 is the clear default
-			OutColor = RETURN_COLOR(LightAccumulator_GetResult(LightAccumulator));
+			InOut.MRT[0] = RETURN_COLOR(LightAccumulator_GetResult(LightAccumulator));
+
+			#if SUBSURFACE_CHANNEL_MODE == 0
+				// Scene captures and planar reflections need opacity in destination alpha.
+				InOut.MRT[0].a = 0;
+			#endif
 		}
 	#endif
 
@@ -1108,34 +1087,76 @@ void Main(
 
 		// -0.5 .. 0.5, could be optimzed as lower quality noise would be sufficient
 		float QuantizationBias = PseudoRandom( MaterialParameters.SvPosition.xy ) - 0.5f;
-		EncodeGBuffer(GBuffer, OutGBufferA, OutGBufferB, OutGBufferC, OutGBufferD, OutGBufferE, OutGBufferVelocity, QuantizationBias);
+		EncodeGBuffer(GBuffer, InOut.MRT[1], InOut.MRT[2], InOut.MRT[3], OutGBufferD, OutGBufferE, OutVelocity, QuantizationBias);
 	#endif 
 
 	if(bEditorWeightedZBuffering)
 	{
-		OutColor.a = 1;
+		InOut.MRT[0].a = 1;
 
 #if MATERIALBLENDING_MASKED
 		// some material might have a opacity value
-		OutColor.a = GetMaterialMaskInputRaw(PixelMaterialInputs);
+		InOut.MRT[0].a = GetMaterialMaskInputRaw(PixelMaterialInputs);
 #endif
 		// we output premultiplied alpha to we have to darken all 4 channels
-		OutColor *= ClipForEditorPrimitives(MaterialParameters);
+		InOut.MRT[0] *= ClipForEditorPrimitives(MaterialParameters);
 
 		#if EDITOR_ALPHA2COVERAGE != 0
 			// per MSAA sample
 			if(MSAASampleCount > 1)
 			{
-				OutCoverage = InCoverage & CustomAlpha2Coverage(OutColor);
+				InOut.OutCoverage = InOut.InCoverage & CustomAlpha2Coverage(InOut.MRT[0]);
 			}
 			else
 			{
 				// no MSAA is handle like per pixel
-				clip(OutColor.a - GetMaterialOpacityMaskClipValue());
+				clip(InOut.MRT[0].a - GetMaterialOpacityMaskClipValue());
 			}
 		#else
 			// per pixel
-			clip(OutColor.a - GetMaterialOpacityMaskClipValue());
+			clip(InOut.MRT[0].a - GetMaterialOpacityMaskClipValue());
 		#endif
 	}
+
+#if USES_GBUFFER
+	// extra starts at MRT4 (see PIXELSHADEROUTPUT_MRT3)
+	int MRT_Write_Index = 4;
+
+	#if WRITES_VELOCITY_TO_GBUFFER
+		InOut.MRT[MRT_Write_Index++] = OutVelocity;
+	#endif
+
+	// we don't test for WRITES_CUSTOMDATA_TO_GBUFFER because GetGBufferRenderTargets() is not masking the MRT based on it, we mask by GetSelectiveOutputMask instead
+	InOut.MRT[MRT_Write_Index++] = OutGBufferD;
+
+	#if WRITES_PRECSHADOWFACTOR_TO_GBUFFER
+		InOut.MRT[MRT_Write_Index++] = OutGBufferE;
+	#endif
+
+#endif
 }
+
+
+// Two options control the render target bindings : OUTPUT_GBUFFER_VELOCITY and ALLOW_STATIC_LIGHTING
+// Big issue here on PC with AMD hardware : putting the velocity buffer after an optional RT fails!!!
+#define VELOCITY_OUTPUT ((USES_GBUFFER) && (WRITES_VELOCITY_TO_GBUFFER))
+// we don't test for WRITES_CUSTOMDATA_TO_GBUFFER because GetGBufferRenderTargets() is not masking the MRT based on it, we mask by GetSelectiveOutputMask instead
+#define CUSTOMDATA_OUTPUT (USES_GBUFFER)
+#define PRESHADOWFACTOR_OUTPUT ((USES_GBUFFER) && (WRITES_PRECSHADOWFACTOR_TO_GBUFFER))
+
+#define EXTRA_MRT_COUNT (VELOCITY_OUTPUT + CUSTOMDATA_OUTPUT + PRESHADOWFACTOR_OUTPUT)
+
+// the following needs to match to the code in FSceneRenderTargets::GetGBufferRenderTargets()
+
+#define PIXELSHADEROUTPUT_BASEPASS 1
+#define PIXELSHADEROUTPUT_MRT0 (1)
+#define PIXELSHADEROUTPUT_MRT1 (USES_GBUFFER)
+#define PIXELSHADEROUTPUT_MRT2 (USES_GBUFFER)
+#define PIXELSHADEROUTPUT_MRT3 (USES_GBUFFER)
+#define PIXELSHADEROUTPUT_MRT4 (EXTRA_MRT_COUNT > 0)
+#define PIXELSHADEROUTPUT_MRT5 (EXTRA_MRT_COUNT > 1)
+#define PIXELSHADEROUTPUT_MRT6 (EXTRA_MRT_COUNT > 2)
+#define PIXELSHADEROUTPUT_A2C ((EDITOR_ALPHA2COVERAGE) != 0)
+// all PIXELSHADEROUTPUT_ and "void FPixelShaderInOut_MainPS()" need to be setup before this include
+// this include generates the wrapper code to call MainPS(inout FPixelShaderOutput PixelShaderOutput)
+#include "PixelShaderOutputCommon.usf"
diff --git a/Engine/Shaders/Common.usf b/Engine/Shaders/Common.usf
--- a/Engine/Shaders/Common.usf
+++ b/Engine/Shaders/Common.usf
@@ -216,7 +216,10 @@
 
 // Whether the platform supports a global clip plane through SV_ClipDistance
 // Ideally this would just be enabled for all SM4 and above feature level platforms, but not tested everywhere yet
-#define SUPPORT_GLOBAL_CLIP_PLANE (PS4_PROFILE || SM5_PROFILE || SM4_PROFILE)
+#define PLATFORM_SUPPORTS_GLOBAL_CLIP_PLANE (PS4_PROFILE || SM5_PROFILE || SM4_PROFILE)
+
+// Using SV_ClipDistance has overhead (15% slower base pass in triangle bound test scene on PS4) so projects have to opt-in
+#define USE_GLOBAL_CLIP_PLANE (PLATFORM_SUPPORTS_GLOBAL_CLIP_PLANE && PROJECT_ALLOW_GLOBAL_CLIP_PLANE)
 
 // This would need to be a #define in GLSL to ignore the SamplerState, however, it is currently a function call in HLSL
 // for type checking of the parameters - ironically the type checking is really only needed in GLSL!
@@ -965,7 +968,7 @@ float2 CalcScreenUVFromOffsetFraction(float4 ScreenPosition, float2 OffsetFracti
 float3 DecodeSceneColorForMaterialNode(float2 ScreenUV)
 {
 #if HIT_PROXY_SHADER || SCENE_TEXTURES_DISABLED
-	// Hit proxies rendering pass doesn't have access to valide render buffers
+	// Hit proxies rendering pass doesn't have access to valid render buffers
 	return float3(0.0f, 0.0f, 0.0f);
 #else
 	float4 EncodedSceneColor = Texture2DSample(LightAttenuationTexture, LightAttenuationTextureSampler, ScreenUV);
@@ -980,7 +983,8 @@ float3 DecodeSceneColorForMaterialNode(float2 ScreenUV)
 // @return SceneDepth (linear in world units, W)
 float ConvertFromDeviceZ(float DeviceZ)
 {
-	return 1.f / (DeviceZ * View.InvDeviceZToWorldZTransform[2] - View.InvDeviceZToWorldZTransform[3]);	
+	// Supports ortho and perspective, see CreateInvDeviceZToWorldZTransform()
+	return DeviceZ * View.InvDeviceZToWorldZTransform[0] + View.InvDeviceZToWorldZTransform[1] + 1.0f / (DeviceZ * View.InvDeviceZToWorldZTransform[2] - View.InvDeviceZToWorldZTransform[3]);
 }
 
 // inverse operation of ConvertFromDeviceZ()
@@ -988,12 +992,21 @@ float ConvertFromDeviceZ(float DeviceZ)
 // @return DeviceZ (Z/W)
 float ConvertToDeviceZ(float SceneDepth)
 {
-	// todo: can change the ADD+MUL into a MAD for precision and performance
-	return 1.0f / ((SceneDepth + View.InvDeviceZToWorldZTransform[3]) * View.InvDeviceZToWorldZTransform[2]);
+	FLATTEN
+	if (View.ViewToClip[3][3] < 1.0f)
+	{
+		// Perspective
+		return 1.0f / ((SceneDepth + View.InvDeviceZToWorldZTransform[3]) * View.InvDeviceZToWorldZTransform[2]);
+	}
+	else
+	{
+		// Ortho
+		return SceneDepth * View.ViewToClip[2][2] + View.ViewToClip[3][2];
+	}
 }
 
 /** Returns clip space W, which is world space distance along the View Z axis. Note if you need DeviceZ LookupDeviceZ() is the faster option */
-float CalcSceneDepth( float2 ScreenUV )
+float CalcSceneDepth(float2 ScreenUV)
 {
 #if SCENE_TEXTURES_DISABLED
 	return 0.0f;
@@ -1021,7 +1034,7 @@ float CalcSceneDepth( float2 ScreenUV )
 #endif
 }
 
-/** Returns DeviceZ which is the z value storesd in the depth buffer. Perfer this over CalcSceneDepth() on most hardware. */
+/** Returns DeviceZ which is the z value stored in the depth buffer. */
 float LookupDeviceZ( float2 ScreenUV )
 {
 #if FEATURE_LEVEL > FEATURE_LEVEL_ES3_1
@@ -1630,3 +1643,21 @@ SamplerState GetTrilinearClampedSampler()
 	return BuiltinSamplers.TrilinearClamped;
 #endif
 }
+
+// see PixelShaderOutputCommon
+struct FPixelShaderInOut
+{
+	// read only
+	float4 SvPosition;
+	// [0..7], only usable if PIXELSHADEROUTPUT_MRT0, PIXELSHADEROUTPUT_MRT1, ... is 1
+	float4 MRT[8];
+
+	// only usable if PIXELSHADEROUTPUT_A2C is 1
+	uint InCoverage;
+	// only usable if PIXELSHADEROUTPUT_A2C is 1
+	uint OutCoverage;
+
+	bool bIsFrontFace;
+
+	float OutDepth;
+};
diff --git a/Engine/Shaders/DeferredDecal.usf b/Engine/Shaders/DeferredDecal.usf
--- a/Engine/Shaders/DeferredDecal.usf
+++ b/Engine/Shaders/DeferredDecal.usf
@@ -167,47 +167,16 @@ float4 SvPositionToScreenPosition2(float4 SvPosition)
 
 // EARLYDEPTHSTENCIL for volumetric decals could work here, but only if we are ouside of the decal
 
-void MainPS(
-	in INPUT_POSITION_QUALIFIERS float4 SvPosition : SV_Position		// after all interpolators
-
-#if DECAL_RENDERTARGET_COUNT > 0
-	, out float4 OutTarget0 : SV_Target0
-#endif
-
-#if DECAL_RENDERTARGET_COUNT > 1 && (BIND_RENDERTARGET1 || COMPILER_METAL) // Used to avoid outputting normals!
-	, out float4 OutTarget1 : SV_Target1
-#endif
-
-#if DECAL_RENDERTARGET_COUNT > 2
-	, out float4 OutTarget2 : SV_Target2
-#endif
-
-#if DECAL_RENDERTARGET_COUNT > 3
-	, out float4 OutTarget3 : SV_Target3
-#endif
-
-#if DECAL_RENDERTARGET_COUNT > 4
-	, out float4 OutTarget4 : SV_Target4
-#endif
-
-	// todo: if we are not inside a volumetric decal we could use OPTIONAL_OutDepthConservative for better performance, this would require a shader permutation
-	OPTIONAL_OutDepth
-
-	)
+// is called in MainPS() from PixelShaderOutputCommon.usf
+void FPixelShaderInOut_MainPS(inout FPixelShaderInOut InOut)
 {
-#if DECAL_RENDERTARGET_COUNT > 1 && !BIND_RENDERTARGET1 && !COMPILER_METAL
-	float4 OutTarget1;
-#endif
-
-	float2 ScreenUV = SvPositionToBufferUV(SvPosition);
+	float2 ScreenUV = SvPositionToBufferUV(InOut.SvPosition);
 
 	// make SvPosition appear to be rasterized with the depth from the depth buffer
-	SvPosition.z = LookupDeviceZ(ScreenUV);
+	InOut.SvPosition.z = LookupDeviceZ(ScreenUV);
 
-#if DECAL_RENDERTARGET_COUNT > 4
 	// default to no shadows, good for volumetric decals
-	OutTarget4 = 1;
-#endif
+	InOut.MRT[4] = 1;
 
 	bool bVolumetric = DECAL_BLEND_MODE == DECALBLENDMODEID_VOLUMETRIC;
 
@@ -220,7 +189,7 @@ void MainPS(
 		// LineBoxIntersect want's to clip a long ray so we provide one, this could be done differently as it limites the draw distance and reduces precision 
 		float SceneW = 100000000;
 
-		float4 SvPositionInDistance = float4(SvPosition.xy, ConvertToDeviceZ(SceneW), 1);
+		float4 SvPositionInDistance = float4(InOut.SvPosition.xy, ConvertToDeviceZ(SceneW), 1);
 
 		float3 WSStartRay = View.WorldCameraOrigin;
 		float3 WSEndRay = SvPositionToWorld(SvPositionInDistance);
@@ -237,7 +206,7 @@ void MainPS(
 		float3 OSBoxStartRay = OSStartRay + OSBoxMinMax.x * (OSEndRay - OSStartRay);
 		float3 OSBoxEndRay = OSStartRay + OSBoxMinMax.y * (OSEndRay - OSStartRay);
 
-		HitT = Raycast(SvPosition, OSBoxStartRay, OSBoxEndRay);
+		HitT = Raycast(InOut.SvPosition, OSBoxStartRay, OSBoxEndRay);
 
 		clip(HitT);
 
@@ -251,7 +220,7 @@ void MainPS(
 	}
 	else
 	{
-		float4 DecalVectorHom = mul(float4(SvPosition.xyz,1), SvPositionToDecal);
+		float4 DecalVectorHom = mul(float4(InOut.SvPosition.xyz,1), SvPositionToDecal);
 		OSPosition = DecalVectorHom.xyz / DecalVectorHom.w;
 
 		// clip content outside the decal
@@ -260,7 +229,7 @@ void MainPS(
 		clip(1.0f - OSPosition.xyz);
 
 		// todo: TranslatedWorld would be better for quality
-		WSPosition = SvPositionToWorld(SvPosition);
+		WSPosition = SvPositionToWorld(InOut.SvPosition);
 	}
 
 	float3 CameraVector = normalize(View.WorldCameraOrigin - WSPosition);
@@ -283,8 +252,8 @@ void MainPS(
 #endif
 	MaterialParameters.VertexColor = 1;
 	MaterialParameters.CameraVector = CameraVector;
-	MaterialParameters.SvPosition = SvPosition;
-	MaterialParameters.ScreenPosition = SvPositionToScreenPosition(SvPosition);
+	MaterialParameters.SvPosition = InOut.SvPosition;
+	MaterialParameters.ScreenPosition = SvPositionToScreenPosition(InOut.SvPosition);
 	MaterialParameters.LightVector = SwizzlePos;
 	MaterialParameters.AbsoluteWorldPosition = 
 		MaterialParameters.WorldPosition_CamRelative = 
@@ -296,7 +265,7 @@ void MainPS(
 		if(bVolumetric)
 		{
 			// not normalized
-			float3 OSNormal = GradientNotNormalized(SvPosition, OSPosition, HitT);
+			float3 OSNormal = GradientNotNormalized(InOut.SvPosition, OSPosition, HitT);
 			
 			// To transform the normals use tranpose(Inverse(DecalToWorld)) = transpose(WorldToDecal)
 			// But we want to only rotate the normals (we don't want to non-uniformaly scale them).
@@ -380,17 +349,17 @@ void MainPS(
 	Opacity *= DecalParams.x;
 
 	// RETURN_COLOR not needed unless writing to SceneColor
-	OutTarget0 = float4(Color, Opacity);
+	InOut.MRT[0] = float4(Color, Opacity);
 
 #if (ES2_PROFILE || ES3_1_PROFILE)
 #if (COMPILER_GLSL_ES2) || (MOBILE_EMULATION)
 	// Do decal blending if encoding requires it.
 	if (GetHDR32bppEncodeMode() == HDR_ENCODE_RGBA)
 	{
-		OutTarget0.rgb = FrameBufferDecalBlendOp(OutTarget0);
+		InOut.MRT[0].rgb = FrameBufferDecalBlendOp(InOut.MRT[0]);
 	}
  	// do 32bpp hdr encoding for mobile if required.
-	OutTarget0 = Encode32BPPHDR(OutTarget0, SvPosition.xy);
+	InOut.MRT[0] = Encode32BPPHDR(InOut.MRT[0], InOut.SvPosition.xy);
 #endif
 #endif
 
@@ -419,7 +388,7 @@ void MainPS(
 		// @param MultiOpacity .x: Color, .y:Normal, .z:Roughness
 		float3 MultiOpacity = Opacity;
 
-		EncodeDBufferData(Data, MultiOpacity, OutTarget0, OutTarget1, OutTarget2);
+		EncodeDBufferData(Data, MultiOpacity, InOut.MRT[0], InOut.MRT[1], InOut.MRT[2]);
 	}
 	#elif DECAL_RENDERSTAGE == 1
 	{
@@ -430,7 +399,7 @@ void MainPS(
 		float4 OutTarget5 = 0;
 		float4 OutTarget6 = 0;
 
-		EncodeGBuffer(Data, OutTarget1, OutTarget2, OutTarget3, OutTarget4, OutTarget5, OutTarget6);
+		EncodeGBuffer(Data, InOut.MRT[1], InOut.MRT[2], InOut.MRT[3], OutTarget4, OutTarget5, OutTarget6);
 	}
 	#elif DECAL_RENDERSTAGE == 2
 	{
@@ -441,23 +410,34 @@ void MainPS(
 		float4 OutTarget5 = 0;
 		float4 OutTarget6 = 0;
 
-		EncodeGBuffer(Data, OutTarget1, OutTarget2, OutTarget3, OutTarget4, OutTarget5, OutTarget6);
+		EncodeGBuffer(Data, InOut.MRT[1], InOut.MRT[2], InOut.MRT[3], OutTarget4, OutTarget5, OutTarget6);
 
 		#if DECAL_BLEND_MODE == DECALBLENDMODEID_STAIN
-			OutTarget3.rgb *= Opacity;
+		InOut.MRT[3].rgb *= Opacity;
 		#endif
 
-		OutTarget0.a = Opacity;	// Emissive
-		OutTarget1.a = Opacity;	// Normal
-		OutTarget2.a = Opacity;	// Metallic, Specular, Roughness
-		OutTarget3.a = Opacity;	// BaseColor
+		InOut.MRT[0].a = Opacity;	// Emissive
+		InOut.MRT[1].a = Opacity;	// Normal
+		InOut.MRT[2].a = Opacity;	// Metallic, Specular, Roughness
+		InOut.MRT[3].a = Opacity;	// BaseColor
 	}
 	#endif // DECAL_RENDERSTAGE
 #endif //DECAL_RENDERTARGET_COUNT > 1
 
 #if DECAL_RENDERTARGET_COUNT > 1 && !BIND_RENDERTARGET1 && COMPILER_METAL
-	OutTarget1 = 0;
+	InOut.MRT[1] = 0;
 #endif
 }
 
+
+#define PIXELSHADEROUTPUT_MRT0 (DECAL_RENDERTARGET_COUNT > 0)
+#define PIXELSHADEROUTPUT_MRT1 (DECAL_RENDERTARGET_COUNT > 1 && (BIND_RENDERTARGET1 || COMPILER_METAL))
+#define PIXELSHADEROUTPUT_MRT2 (DECAL_RENDERTARGET_COUNT > 2)
+#define PIXELSHADEROUTPUT_MRT3 (DECAL_RENDERTARGET_COUNT > 3)
+#define PIXELSHADEROUTPUT_MRT4 (DECAL_RENDERTARGET_COUNT > 4)
+// all PIXELSHADEROUTPUT_ and "void FPixelShaderInOut_MainPS()" need to be setup before this include
+// this include generates the wrapper code to call MainPS()
+#include "PixelShaderOutputCommon.usf"
+
+
 #endif // DECAL_PROJECTION
diff --git a/Engine/Shaders/DeferredLightPixelShaders.usf b/Engine/Shaders/DeferredLightPixelShaders.usf
--- a/Engine/Shaders/DeferredLightPixelShaders.usf
+++ b/Engine/Shaders/DeferredLightPixelShaders.usf
@@ -23,15 +23,15 @@ FDeferredLightData SetupLightDataForStandardDeferred()
 	LightData.LightDirection = DeferredLightUniforms.NormalizedLightDirection;
 	LightData.SpotAnglesAndSourceRadius = float4(DeferredLightUniforms.SpotAngles, DeferredLightUniforms.SourceRadius, DeferredLightUniforms.SourceLength);
 	LightData.MinRoughness = DeferredLightUniforms.MinRoughness;
+	LightData.ContactShadowLength = DeferredLightUniforms.ContactShadowLength;
 	LightData.DistanceFadeMAD = DeferredLightUniforms.DistanceFadeMAD;
 	LightData.ShadowMapChannelMask = DeferredLightUniforms.ShadowMapChannelMask;
+	LightData.ShadowedBits = DeferredLightUniforms.ShadowedBits;
 
 	LightData.bInverseSquared = INVERSE_SQUARED_FALLOFF;
 	LightData.bRadialLight = RADIAL_ATTENUATION;
 	//@todo - permutation opportunity
 	LightData.bSpotLight = RADIAL_ATTENUATION;
-	//@todo - permutation opportunity
-	LightData.bShadowed = DeferredLightUniforms.bShadowed;
 
 	return LightData;
 }
diff --git a/Engine/Shaders/DeferredLightingCommon.usf b/Engine/Shaders/DeferredLightingCommon.usf
--- a/Engine/Shaders/DeferredLightingCommon.usf
+++ b/Engine/Shaders/DeferredLightingCommon.usf
@@ -26,6 +26,7 @@ struct FDeferredLightData
 	float3 LightDirection;
 	float4 SpotAnglesAndSourceRadius;
 	float MinRoughness;
+	float ContactShadowLength;
 	float2 DistanceFadeMAD;
 	float4 ShadowMapChannelMask;
 	/** Whether to use inverse squared falloff. */
@@ -35,7 +36,7 @@ struct FDeferredLightData
 	/** Whether this light needs spotlight attenuation. */
 	bool bSpotLight;
 	/** Whether the light should apply shadowing. */
-	bool bShadowed;
+	uint ShadowedBits;
 };
 
 /** Data about a single light to be shaded with the simple shading model, designed for speed and limited feature set. */
@@ -276,6 +277,58 @@ void GetShadowTerms(FScreenSpaceData ScreenSpaceData, FDeferredLightData LightDa
 	}
 }
 
+float ShadowRayCast(
+	float3 RayOriginTranslatedWorld, float3 RayDirection, float RayLength,
+	int NumSteps, float StepOffset
+)
+{
+	float4 RayStartClip	= mul( float4( RayOriginTranslatedWorld, 1 ), View.TranslatedWorldToClip );
+	float4 RayDirClip	= mul( float4( RayDirection * RayLength, 0 ), View.TranslatedWorldToClip );
+	float4 RayEndClip	= RayStartClip + RayDirClip;
+
+	float3 RayStartScreen = RayStartClip.xyz / RayStartClip.w;
+	float3 RayEndScreen = RayEndClip.xyz / RayEndClip.w;
+	
+	float3 RayStepScreen = RayEndScreen - RayStartScreen;
+
+	float3 RayStartUVz = float3( RayStartScreen.xy * View.ScreenPositionScaleBias.xy + View.ScreenPositionScaleBias.wz, RayStartScreen.z );
+	float3 RayStepUVz = float3( RayStepScreen.xy * View.ScreenPositionScaleBias.xy, RayStepScreen.z );
+
+	float4 RayDepthClip	= RayStartClip + mul( float4( 0, 0, RayLength, 0 ), View.ViewToClip );
+	float3 RayDepthScreen = RayDepthClip.xyz / RayDepthClip.w;
+
+	const float Step = 1.0 / NumSteps;
+
+	// *2 to get less morie pattern in extreme cases, larger values make object appear not grounded in reflections
+	const float CompareTolerance = abs( RayDepthScreen.z - RayStartScreen.z ) * Step * 2;
+
+	float SampleTime = StepOffset * Step + Step;
+	
+	float Coverage = 0;
+
+	UNROLL
+	for( int i = 0; i < NumSteps; i++ )
+	{
+		float3 SampleUVz = RayStartUVz + RayStepUVz * SampleTime;
+		float SampleDepth = SceneDepthTexture.SampleLevel( SceneDepthTextureSampler, SampleUVz.xy, 0 ).r;
+
+		float DepthDiff = SampleUVz.z - SampleDepth;
+		bool Hit = abs( DepthDiff + CompareTolerance ) < CompareTolerance;
+
+		Coverage += Hit ? 1.0 : 0.0;
+
+		SampleTime += Step;
+	}
+
+	float Shadow = Coverage > 0 ? 1 : 0;
+
+	// Off screen masking
+	float2 Vignette = saturate( abs( RayStartScreen.xy ) * 6 - 5 );
+	//Shadow *= saturate( 1.0 - dot( Vignette, Vignette ) );
+
+	return 1 - Shadow;
+}
+
 /** Calculates lighting for a given position, normal, etc with a fully featured lighting model designed for quality. */
 float4 GetDynamicLighting(float3 WorldPosition, float3 CameraVector, float2 InUV, FScreenSpaceData ScreenSpaceData, uint ShadingModelID, FDeferredLightData LightData, float4 LightAttenuation, uint2 Random)
 {
@@ -392,19 +445,48 @@ float4 GetDynamicLighting(float3 WorldPosition, float3 CameraVector, float2 InUV
 		float SubsurfaceShadow = 1;
 
 		BRANCH
-		if (LightData.bShadowed)
+		if (LightData.ShadowedBits)
 		{
 			GetShadowTerms(ScreenSpaceData, LightData, WorldPosition, LightAttenuation, SurfaceShadow, SubsurfaceShadow);
 
 			// greatly reduces shadow mapping artifacts
 			// Commented out because it reduces character shading quality.
 			//SurfaceShadow *= saturate(dot(N, L) * 6 - 0.2);
+
+#ifndef THREADGROUP_SIZEX
+			BRANCH
+			if( LightData.ShadowedBits > 1 && LightData.ContactShadowLength > 0 )
+			{
+				uint FrameRandom = (Frame.FrameNumber & 15) * 1551;
+
+				float StepOffset = float( Random.x & 0xffff ) / (1<<16);
+				StepOffset -= 0.5;
+
+				float Shadow = ShadowRayCast( WorldPosition + View.PreViewTranslation, L, LightData.ContactShadowLength * ScreenSpaceData.GBuffer.Depth, 8, StepOffset );
+				
+				SurfaceShadow *= Shadow;
+				//SubsurfaceShadow *= Shadow;
+			}
+#endif
 		}
 		else
 		{
 			SurfaceShadow = ScreenSpaceData.AmbientOcclusion;
 		}
 
+#ifndef THREADGROUP_SIZEX
+		BRANCH
+		if( LightData.ShadowedBits < 2 && ScreenSpaceData.GBuffer.ShadingModelID == SHADINGMODELID_HAIR )
+		{
+			uint FrameRandom = (Frame.FrameNumber & 15) * 1551;
+
+			float StepOffset = float( Random.x & 0xffff ) / (1<<16);
+			StepOffset -= 0.5;
+
+			SubsurfaceShadow = ShadowRayCast( WorldPosition + View.PreViewTranslation, L, 0.1 * ScreenSpaceData.GBuffer.Depth, 8, StepOffset );
+		}
+#endif
+
 		float SurfaceAttenuation	= (DistanceAttenuation * LightRadiusMask * SpotFalloff) * SurfaceShadow;
 		float SubsurfaceAttenuation	= (DistanceAttenuation * LightRadiusMask * SpotFalloff) * SubsurfaceShadow;
 
diff --git a/Engine/Shaders/DeferredShadingCommon.usf b/Engine/Shaders/DeferredShadingCommon.usf
--- a/Engine/Shaders/DeferredShadingCommon.usf
+++ b/Engine/Shaders/DeferredShadingCommon.usf
@@ -509,6 +509,7 @@ FGBufferData DecodeGBufferData(
 	GBuffer.Metallic	= InGBufferB.r;
 	GBuffer.Specular	= InGBufferB.g;
 	GBuffer.Roughness	= InGBufferB.b;
+	// Note: must match GetShadingModelId standalone function logic
 	GBuffer.ShadingModelID = DecodeShadingModelId(InGBufferB.a);
 	GBuffer.SelectiveOutputMask = DecodeSelectiveOutputMask(InGBufferB.a);
 
@@ -780,8 +781,7 @@ FGBufferData GetGBufferData(float2 UV, bool bGetNormalizedNormal = true)
 // Minimal path for just the lighting model, used to branch around unlit pixels (skybox)
 uint GetShadingModelId(float2 UV)
 {
-	bool Unused;
-	return DecodeShadingModelId(Decode71(Texture2DSampleLevel(GBuffers.GBufferBTexture, GBuffers.GBufferBTextureSampler, UV, 0).a, Unused));
+	return DecodeShadingModelId(Texture2DSampleLevel(GBuffers.GBufferBTexture, GBuffers.GBufferBTextureSampler, UV, 0).a);
 }
 
 // @param UV - UV space in the DBuffer textures (BufferSize resolution)
diff --git a/Engine/Shaders/Definitions.usf b/Engine/Shaders/Definitions.usf
--- a/Engine/Shaders/Definitions.usf
+++ b/Engine/Shaders/Definitions.usf
@@ -284,4 +284,9 @@
 #if !COMPILER_HLSLCC
 #error "Missing COMPILER_HLSLCC define!"
 #endif
-#endif
\ No newline at end of file
+#endif
+
+// Hlslcc platforms ignore the uniform keyword as it can't properly optimize flow
+#if COMPILER_HLSLCC
+#define uniform
+#endif
diff --git a/Engine/Shaders/DistanceFieldScreenGridLighting.usf b/Engine/Shaders/DistanceFieldScreenGridLighting.usf
--- a/Engine/Shaders/DistanceFieldScreenGridLighting.usf
+++ b/Engine/Shaders/DistanceFieldScreenGridLighting.usf
@@ -232,7 +232,6 @@ void ConeTraceObjectOcclusionCS(
 	float3 TangentY;
 	FindBestAxisVectors2(WorldNormal, TangentX, TangentY);
 
-	if (all(DispatchThreadId < ScreenGridConeVisibilitySize))
 	{
 		float2 ScreenUV = GetScreenUVFromScreenGrid(OutputCoordinate);
 		float2 ScreenPosition = (ScreenUV.xy - View.ScreenPositionScaleBias.wz) / View.ScreenPositionScaleBias.xy;
diff --git a/Engine/Shaders/ForEachLight.usf b/Engine/Shaders/ForEachLight.usf
--- a/Engine/Shaders/ForEachLight.usf
+++ b/Engine/Shaders/ForEachLight.usf
@@ -35,7 +35,7 @@ float3 LightingForOneLight(FMaterialPixelParameters Parameters, uint LightIndex,
 	// Only radial lights supported with tiled deferred
 	LightData.bRadialLight = ForwardLightData.SpotAnglesAndSourceRadiusAndDir[LightIndex].w == 0;
 	// The only type of shadowing supported for lights using tiled is static shadowing, so the light should only compute shadowing if it has static shadowing
-	LightData.bShadowed = 0;
+	LightData.ShadowedBits = 0;
 	
 	float2 ScreenUV = Parameters.SvPosition.xy * Frame.BufferSizeAndInvSize.zw;
 
diff --git a/Engine/Shaders/HeightFogCommon.usf b/Engine/Shaders/HeightFogCommon.usf
--- a/Engine/Shaders/HeightFogCommon.usf
+++ b/Engine/Shaders/HeightFogCommon.usf
@@ -7,9 +7,11 @@
 static const float FLT_EPSILON = 0.001f;
 static const float FLT_EPSILON2 = 0.01f;
 
-/** Shader parameters shared between exponential and multi-layer height fog, to reduce constant register usage. */
-float4 SharedFogParameter0;
-float4 SharedFogParameter1;
+/** FogDensity * exp2(-FogHeightFalloff * (CameraWorldPosition.z - FogHeight)) in x, FogHeightFalloff in y, CosTerminatorAngle in z, StartDistance in w. */
+float4 ExponentialFogParameters;
+float4 ExponentialFogColorParameter;
+/** FogDensity in x, FogHeight in y. */
+float2 ExponentialFogParameters3;
 half4 InscatteringLightDirection;
 half4 DirectionalInscatteringColor;
 float DirectionalInscatteringStartDistance;
@@ -17,35 +19,63 @@ float DirectionalInscatteringStartDistance;
 // @param WorldPositionRelativeToCamera = WorldPosition - InCameraPosition
 half4 GetExponentialHeightFog(float3 WorldPositionRelativeToCamera)
 {
-	/** FogDensity * exp2(-FogHeightFalloff * (CameraWorldPosition.z - FogHeight)) in x, FogHeightFalloff in y, CosTerminatorAngle in z, StartDistance in w. */
-	const float4 ExponentialFogParameters = SharedFogParameter0;
 	/** Color to use. */
-	const half3 ExponentialFogColor = SharedFogParameter1.xyz;
-	const half MinFogOpacity = SharedFogParameter1.w;
+	const half3 ExponentialFogColor = ExponentialFogColorParameter.xyz;
+	const half MinFogOpacity = ExponentialFogColorParameter.w;
 
 	float3 CameraToReceiver = WorldPositionRelativeToCamera;
-	float CameraToReceiverLengthSqr = dot( CameraToReceiver, CameraToReceiver );
-	float CameraToReceiverLengthInv = rsqrt( CameraToReceiverLengthSqr );
+	float CameraToReceiverLengthSqr = dot(CameraToReceiver, CameraToReceiver);
+	float CameraToReceiverLengthInv = rsqrt(CameraToReceiverLengthSqr);
 	float CameraToReceiverLength = CameraToReceiverLengthSqr * CameraToReceiverLengthInv;
 	half3 CameraToReceiverNormalized = CameraToReceiver * CameraToReceiverLengthInv;
 
+	float RayOriginTerms = ExponentialFogParameters.x;
+	float RayLength = CameraToReceiverLength;
+	float RayDirectionZ = CameraToReceiver.z;
+
+#if USE_GLOBAL_CLIP_PLANE
+
+	BRANCH
+	// While rendering a planar reflection with a clip plane, we must compute analytical fog using a camera path starting from the plane, rather than the virtual camera origin
+	if (dot(View.GlobalClippingPlane.xyz, 1) > 0.0f)
+	{
+		float CameraOriginPlaneDistance = dot(View.GlobalClippingPlane, float4(View.WorldCameraOrigin, 1));
+		float PlaneIntersectionTime = -CameraOriginPlaneDistance / dot(CameraToReceiver, View.GlobalClippingPlane.xyz);
+
+		// Continue if the reflection plane is between the camera and receiver point
+		if (PlaneIntersectionTime > 0 && PlaneIntersectionTime < 1)
+		{
+			float CameraToPlaneIntersectionZ = PlaneIntersectionTime * CameraToReceiver.z;
+			float PlaneIntersectionZ = View.WorldCameraOrigin.z + CameraToPlaneIntersectionZ;
+			float PlaneIntersectionToReceiverZ = CameraToReceiver.z - CameraToPlaneIntersectionZ;
+
+			// Calculate fog off of the ray starting from the reflection plane, instead of starting from the camera
+			RayLength = (1.0f - PlaneIntersectionTime) * CameraToReceiverLength;
+			RayDirectionZ = PlaneIntersectionToReceiverZ;
+
+			float Exponent = max(-127.0f, ExponentialFogParameters.y * (PlaneIntersectionZ - ExponentialFogParameters3.y));
+			RayOriginTerms = ExponentialFogParameters3.x * exp2(-Exponent);
+		}
+	}
+
+#endif
+
 	// Calculate the line integral of the ray from the camera to the receiver position through the fog density function
 	// The exponential fog density function is d = GlobalDensity * exp(-HeightFalloff * z)
-	float EffectiveZ = (abs(CameraToReceiver.z) > FLT_EPSILON2) ? CameraToReceiver.z : FLT_EPSILON2;
-	float Falloff = max( -127.0f, ExponentialFogParameters.y * EffectiveZ );	// if it's lower than -127.0, then exp2() goes crazy in OpenGL's GLSL.
-	float ExponentialHeightLineIntegralShared = ExponentialFogParameters.x * (1.0f - exp2(-Falloff) ) / Falloff;
-	float ExponentialHeightLineIntegral = ExponentialHeightLineIntegralShared * max(CameraToReceiverLength - ExponentialFogParameters.w, 0.0f);
+	float EffectiveZ = (abs(RayDirectionZ) > FLT_EPSILON2) ? RayDirectionZ : FLT_EPSILON2;
+	float Falloff = max(-127.0f, ExponentialFogParameters.y * EffectiveZ);	// if it's lower than -127.0, then exp2() goes crazy in OpenGL's GLSL.
+	float ExponentialHeightLineIntegralShared = RayOriginTerms * (1.0f - exp2(-Falloff) ) / Falloff;
+	float ExponentialHeightLineIntegral = ExponentialHeightLineIntegralShared * max(RayLength - ExponentialFogParameters.w, 0.0f);
 
 	half3 DirectionalInscattering = 0;
 
 	if (InscatteringLightDirection.w > 0)
 	{
 		// Setup a cosine lobe around the light direction to approximate inscattering from the directional light off of the ambient haze;
-		// clamp to avoid issues on some platforms (Could've used ClampedPow(), but we want the dot to be 0..1)
-		half3 DirectionalLightInscattering = DirectionalInscatteringColor.xyz * pow(saturate(dot(CameraToReceiverNormalized, InscatteringLightDirection.xyz)), clamp(DirectionalInscatteringColor.w, 0.000001f, 1000.0));
+		half3 DirectionalLightInscattering = DirectionalInscatteringColor.xyz * pow(saturate(dot(CameraToReceiverNormalized, InscatteringLightDirection.xyz)), DirectionalInscatteringColor.w);
 
 		// Calculate the line integral of the eye ray through the haze, using a special starting distance to limit the inscattering to the distance
-		float DirExponentialHeightLineIntegral = ExponentialHeightLineIntegralShared * max(CameraToReceiverLength - DirectionalInscatteringStartDistance, 0.0f);
+		float DirExponentialHeightLineIntegral = ExponentialHeightLineIntegralShared * max(RayLength - DirectionalInscatteringStartDistance, 0.0f);
 		// Calculate the amount of light that made it through the fog using the transmission equation
 		half DirectionalInscatteringFogFactor = saturate(exp2(-DirExponentialHeightLineIntegral));
 		// Final inscattering from the light
diff --git a/Engine/Shaders/MaterialTemplate.usf b/Engine/Shaders/MaterialTemplate.usf
--- a/Engine/Shaders/MaterialTemplate.usf
+++ b/Engine/Shaders/MaterialTemplate.usf
@@ -47,12 +47,6 @@
 	#define USE_STENCIL_LOD_DITHER	USE_STENCIL_LOD_DITHER_DEFAULT
 #endif
 
-// scale factor for screen position, for cases where we may render certain things at lower resolutions (e.g. separate translucency)
-// any pass wanting to render at reduced resolution needs to simply set this constant to scale scene texture UVs to make scene texture accesses work
-#if PIXELSHADER && TRANSLUCENCY_RENDERING
-float2 ScreenPositionUVScale;
-#endif
-
 //Platforms that don't run the editor shouldn't need editor features in the shaders.
 #define PLATFORM_SUPPORTS_EDITOR_SHADERS (!PS4_PROFILE && !XBOXONE_PROFILE && !ES31_AEP_PROFILE)
 
@@ -62,6 +56,8 @@ float2 ScreenPositionUVScale;
 //Materials also have to opt in to these features.
 #define USE_EDITOR_COMPOSITING (USE_EDITOR_SHADERS && EDITOR_PRIMITIVE_MATERIAL)
 
+#define MATERIALBLENDING_ANY_TRANSLUCENT (MATERIALBLENDING_TRANSLUCENT || MATERIALBLENDING_ADDITIVE || MATERIALBLENDING_MODULATE)
+
 /**
  * Parameters used by vertex and pixel shaders to access particle properties.
  */
@@ -1124,10 +1120,6 @@ float4 SceneTextureLookup(float2 UV, int SceneTextureIndex, bool bFiltered)
 	return float4(0.0f, 0.0f, 0.0f, 0.0f);
 #endif
 
-#if PIXELSHADER && TRANSLUCENCY_RENDERING
-	UV *= ScreenPositionUVScale;
-#endif
-
 	FScreenSpaceData ScreenSpaceData = GetScreenSpaceData(UV, false);
 
 #if PS4_PROFILE
@@ -1791,9 +1783,6 @@ half3x3 AssembleTangentToWorld( half3 TangentToWorld0, half4 TangentToWorld2 )
 	#define OPTIONAL_OutDepth
 #endif
 
-// Using SV_ClipDistance has overhead (15% slower base pass in triangle bound test scene on PS4) so projects have to opt-in
-#define USE_GLOBAL_CLIP_PLANE (SUPPORT_GLOBAL_CLIP_PLANE && PROJECT_ALLOW_GLOBAL_CLIP_PLANE)
-
 float ApplyPixelDepthOffsetToMaterialParameters(inout FMaterialPixelParameters MaterialParameters, FPixelMaterialInputs PixelMaterialInputs, out float OutDepth)
 {
 	// Can only push into the scene for conservative depth writes
diff --git a/Engine/Shaders/MaterialTexCoordScalesPixelShader.usf b/Engine/Shaders/MaterialTexCoordScalesPixelShader.usf
--- a/Engine/Shaders/MaterialTexCoordScalesPixelShader.usf
+++ b/Engine/Shaders/MaterialTexCoordScalesPixelShader.usf
@@ -12,7 +12,7 @@
 
  // If -1, output everything to the RW texture.
 float4 AccuracyColors[5];
-uint TextureAnalysisIndex;
+int TextureAnalysisIndex;
 
 float4 OneOverCPUTexCoordScales[MAX_NUM_TEXTURE_REGISTER / 4];
 int4 TexCoordIndices[MAX_NUM_TEXTURE_REGISTER / 4];
@@ -31,7 +31,7 @@ float GetComponent(float4 V, int Index)
 struct FTexCoordScalesParams
 {
 	// Global var used to reduce the number of param passed to the shaders.
-	uint2 PixelPosition;
+	int2 PixelPosition;
 
 	// ddU, ddV for each coord index (up to 4 currently)
 	float4 OneOverDDU;
@@ -102,7 +102,7 @@ void Main(
 	FTexCoordScalesParams Params;
 
 	// Set global var used for outputting data
-	Params.PixelPosition = DebugInputs.SvPosition;
+	Params.PixelPosition = DebugInputs.SvPosition.xy;
 	Params.OneOverDDU = 1 / float4(length(ddx(DebugInputs.TexCoord01.xy)), length(ddx(DebugInputs.TexCoord01.zw)), length(ddx(DebugInputs.TexCoord23.xy)), length(ddx(DebugInputs.TexCoord23.zw)));
 	Params.OneOverDDV = 1 / float4(length(ddy(DebugInputs.TexCoord01.xy)), length(ddy(DebugInputs.TexCoord01.zw)), length(ddy(DebugInputs.TexCoord23.xy)), length(ddy(DebugInputs.TexCoord23.zw)));
 
@@ -142,18 +142,22 @@ void Main(
 		CalcMaterialParameters(MaterialParameters, PixelMaterialInputs, DebugInputs.SvPosition, bIsFrontFace);
 
 		// Sample material properties. The results are not used, but the calls to StoreTexCoordScale are still be made.
-		BaseColor = GetMaterialBaseColor(PixelMaterialInputs);
-		half  Metallic = GetMaterialMetallic(PixelMaterialInputs);
-		half  Specular = GetMaterialSpecular(PixelMaterialInputs);
-		float Roughness = GetMaterialRoughness(PixelMaterialInputs);
-		half3 Emissive = GetMaterialEmissive(PixelMaterialInputs);
-		half Opacity = GetMaterialOpacity(PixelMaterialInputs);
+		BaseColor = GetMaterialBaseColorRaw(PixelMaterialInputs);
+		half  Metallic = GetMaterialMetallicRaw(PixelMaterialInputs);
+		half  Specular = GetMaterialSpecularRaw(PixelMaterialInputs);
+		float Roughness = GetMaterialRoughnessRaw(PixelMaterialInputs);
+		half3 Normal = GetMaterialNormalRaw(PixelMaterialInputs);
+		half3 Emissive = GetMaterialEmissiveRaw(PixelMaterialInputs);
+		half Opacity = GetMaterialOpacityRaw(PixelMaterialInputs);
 		half Mask = GetMaterialMask(PixelMaterialInputs);
-		half4 SSData = GetMaterialSubsurfaceData(MaterialParameters);
+		half4 SSData = GetMaterialSubsurfaceDataRaw(MaterialParameters);
 		float Custom0 = GetMaterialCustomData0(MaterialParameters);
 		float Custom1 = GetMaterialCustomData1(MaterialParameters);
-		float MaterialAO = GetMaterialAmbientOcclusion(PixelMaterialInputs);
+		float MaterialAO = GetMaterialAmbientOcclusionRaw(PixelMaterialInputs);
 		float PixelDepthOffset = GetMaterialPixelDepthOffset(PixelMaterialInputs);
+#if CLEAR_COAT_BOTTOM_NORMAL && NUM_MATERIAL_OUTPUTS_CLEARCOATBOTTOMNORMAL > 0
+		float3 BottomNormal =  ClearCoatBottomNormal0(MaterialParameters);
+#endif
 	}
 	Params = MaterialParameters.TexCoordScalesParams;
 
diff --git a/Engine/Shaders/PixelShaderOutputCommon.usf b/Engine/Shaders/PixelShaderOutputCommon.usf
new file mode 100644
--- /dev/null
+++ b/Engine/Shaders/PixelShaderOutputCommon.usf
@@ -0,0 +1,148 @@
+// Copyright 1998-2016 Epic Games, Inc. All Rights Reserved.
+
+/*=============================================================================
+PixelShaderOutputCommon.usf: To allow PS input/output passed int functions through a single struct, allowing for a more readable code (less #ifdefs, reducing the boolean hell)
+=============================================================================*/
+
+#ifndef PIXELSHADEROUTPUT_BASEPASS
+	#define PIXELSHADEROUTPUT_BASEPASS 0
+#endif
+#ifndef PIXELSHADEROUTPUT_MRT0
+	#define PIXELSHADEROUTPUT_MRT0 0
+#endif
+#ifndef PIXELSHADEROUTPUT_MRT1
+	#define PIXELSHADEROUTPUT_MRT1 0
+#endif
+#ifndef PIXELSHADEROUTPUT_MRT2
+	#define PIXELSHADEROUTPUT_MRT2 0
+#endif
+#ifndef PIXELSHADEROUTPUT_MRT3
+	#define PIXELSHADEROUTPUT_MRT3 0
+#endif
+#ifndef PIXELSHADEROUTPUT_MRT4
+	#define PIXELSHADEROUTPUT_MRT4 0
+#endif
+#ifndef PIXELSHADEROUTPUT_MRT5
+	#define PIXELSHADEROUTPUT_MRT5 0
+#endif
+#ifndef PIXELSHADEROUTPUT_MRT6
+	#define PIXELSHADEROUTPUT_MRT6 0
+#endif
+#ifndef PIXELSHADEROUTPUT_MRT7
+	#define PIXELSHADEROUTPUT_MRT7 0
+#endif
+
+
+// ---------------------------------------------------------------------------------
+
+void MainPS
+	(
+#if PIXELSHADEROUTPUT_BASEPASS
+		FVertexFactoryInterpolantsVSToPS Interpolants,
+		FBasePassInterpolantsVSToPS BasePassInterpolants,
+#endif
+
+		in INPUT_POSITION_QUALIFIERS float4 SvPosition : SV_Position		// after all interpolators
+
+		OPTIONAL_IsFrontFace
+
+#if PIXELSHADEROUTPUT_MRT0
+		, out float4 OutTarget0 : SV_Target0
+#endif
+
+#if PIXELSHADEROUTPUT_MRT1
+		, out float4 OutTarget1 : SV_Target1
+#endif
+
+#if PIXELSHADEROUTPUT_MRT2
+		, out float4 OutTarget2 : SV_Target2
+#endif
+
+#if PIXELSHADEROUTPUT_MRT3
+		, out float4 OutTarget3 : SV_Target3
+#endif
+
+#if PIXELSHADEROUTPUT_MRT4
+		, out float4 OutTarget4 : SV_Target4
+#endif
+
+#if PIXELSHADEROUTPUT_MRT5
+		, out float4 OutTarget5 : SV_Target5
+#endif
+
+#if PIXELSHADEROUTPUT_MRT6
+		, out float4 OutTarget6 : SV_Target6
+#endif
+
+#if PIXELSHADEROUTPUT_MRT7
+		, out float4 OutTarget7 : SV_Target7
+#endif
+
+		// todo: if we are not inside a volumetric decal we could use OPTIONAL_OutDepthConservative for better performance, this would require a shader permutation
+		OPTIONAL_OutDepthConservative
+
+#if PIXELSHADEROUTPUT_A2C
+		// alpha to coverage
+		, in uint InCoverage : SV_Coverage
+		, out uint OutCoverage : SV_Coverage
+#endif
+	)
+{
+	// ---------------------------------------------------------------------------------
+
+	FPixelShaderInOut PixelShaderInOut = (FPixelShaderInOut)0;
+
+#if PIXELSHADEROUTPUT_A2C
+	PixelShaderInOut.InCoverage = InCoverage;
+	PixelShaderInOut.OutCoverage = InCoverage;
+#endif 
+
+	PixelShaderInOut.SvPosition = SvPosition;
+	PixelShaderInOut.bIsFrontFace = bIsFrontFace;
+
+#if PIXELSHADEROUTPUT_BASEPASS
+	FPixelShaderInOut_MainPS(Interpolants, BasePassInterpolants, PixelShaderInOut);
+#else
+	FPixelShaderInOut_MainPS(PixelShaderInOut);
+#endif
+
+#if PIXELSHADEROUTPUT_MRT0
+	OutTarget0 = PixelShaderInOut.MRT[0];
+#endif
+
+#if PIXELSHADEROUTPUT_MRT1
+	OutTarget1 = PixelShaderInOut.MRT[1];
+#endif
+
+#if PIXELSHADEROUTPUT_MRT2
+	OutTarget2 = PixelShaderInOut.MRT[2];
+#endif
+
+#if PIXELSHADEROUTPUT_MRT3
+	OutTarget3 = PixelShaderInOut.MRT[3];
+#endif
+
+#if PIXELSHADEROUTPUT_MRT4
+	OutTarget4 = PixelShaderInOut.MRT[4];
+#endif
+
+#if PIXELSHADEROUTPUT_MRT5
+	OutTarget5 = PixelShaderInOut.MRT[5];
+#endif
+
+#if PIXELSHADEROUTPUT_MRT6
+	OutTarget6 = PixelShaderInOut.MRT[6];
+#endif
+
+#if PIXELSHADEROUTPUT_MRT7
+	OutTarget7 = PixelShaderInOut.MRT[7];
+#endif
+
+#if PIXELSHADEROUTPUT_A2C
+	OutCoverage = PixelShaderInOut.OutCoverage;
+#endif 
+
+#if OUTPUT_PIXEL_DEPTH_OFFSET
+	OutDepth = PixelShaderInOut.OutDepth;
+#endif 
+}
\ No newline at end of file
diff --git a/Engine/Shaders/PlanarReflectionShaders.usf b/Engine/Shaders/PlanarReflectionShaders.usf
--- a/Engine/Shaders/PlanarReflectionShaders.usf
+++ b/Engine/Shaders/PlanarReflectionShaders.usf
@@ -129,6 +129,6 @@ void PlanarReflectionPS(
 	BRANCH
 	if (ScreenSpaceData.GBuffer.ShadingModelID > 0)
 	{
-		OutColor = ComputePlanarReflections(WorldPosition, ScreenSpaceData.GBuffer.WorldNormal, PlanarReflectionSampler);
+		OutColor = ComputePlanarReflections(WorldPosition, ScreenSpaceData.GBuffer.WorldNormal, ScreenSpaceData.GBuffer.Roughness, PlanarReflectionSampler);
 	}
 }
diff --git a/Engine/Shaders/PlanarReflectionShared.usf b/Engine/Shaders/PlanarReflectionShared.usf
--- a/Engine/Shaders/PlanarReflectionShared.usf
+++ b/Engine/Shaders/PlanarReflectionShared.usf
@@ -26,7 +26,7 @@ float4x4 ProjectionWithExtraFOV;
 Texture2D PlanarReflectionTexture;
 SamplerState PlanarReflectionSampler;
 
-float4 ComputePlanarReflections(float3 WorldPosition, float3 WorldNormal, SamplerState SharedClampSampler)
+float4 ComputePlanarReflections(float3 WorldPosition, float3 WorldNormal, float Roughness, SamplerState SharedClampSampler)
 {
 	float4 OutPlanarReflection = 0;
 
@@ -64,7 +64,8 @@ float4 ComputePlanarReflections(float3 WorldPosition, float3 WorldNormal, Sample
 		float2 ViewportUV = ClipVirtualReflectionSpherePosition.xy / ClipVirtualReflectionSpherePosition.w * float2(.5f, -.5f) + .5f;
 
 		float AngleFade = saturate(dot(ReflectionPlane.xyz, WorldNormal) * PlanarReflectionParameters2.x + PlanarReflectionParameters2.y);
-		float FinalFade = DistanceFade * AngleFade;
+		float RoughnessFade = 1 - saturate((Roughness - .2f) * 10.0f);
+		float FinalFade = DistanceFade * AngleFade * RoughnessFade;
 
 		float4 PlanarReflectionTextureValue = Texture2DSampleLevel(
 			PlanarReflectionTexture, 
diff --git a/Engine/Shaders/PostProcessMotionBlur.usf b/Engine/Shaders/PostProcessMotionBlur.usf
--- a/Engine/Shaders/PostProcessMotionBlur.usf
+++ b/Engine/Shaders/PostProcessMotionBlur.usf
@@ -9,729 +9,38 @@
 #include "DeferredShadingCommon.usf"		// FGBufferData
 #include "FastMath.usf"
 
-// Spaces:
-//   screen space: -1..1 -1..1
-//   normalized_motionblur_velocity: 2d vector where the max motionblur is defined in a unit circle around 0 
-//   world_space
-//   half_res_pixels
-
-// similar implementation to: http://graphics.cs.williams.edu/papers/MotionBlurI3D12/McGuire12Blur.pdf
-// but adapted for half resolution and not using randomization/noise
-
-// from the paper: We use SOFT Z EXTENT = 1mm to 10cm for our results
-#define SOFT_Z_EXTENT 1
-
-// This,
-//  - Gets the MainPS pass to low cost under no motion.
-//  - Somewhat cleans up the blur/no blur transition.
-//  - Reduces the motion blur on small motion (where it is not needed).
-//  - Removes the background blur (seems higher quality).
-#define MOTION_BLUR_EARLY_EXIT 1
-
-// 0: use more texture lookups avoiding .Gather() lookups
-#define GATHER_OPTIMIZATION (FEATURE_LEVEL >= FEATURE_LEVEL_SM5)
-
-// 0:off / 1:on, useful to debug the motionblur algorithm
-#define MOTIONBLUR_TESTCHART 0
-
-// --------------------------------------------------------------------------
-
-#if MOTION_BLUR_QUALITY == 0
-	// this would break VisualizeMotionBlur so we disable it
-	#undef MOTION_BLUR_EARLY_EXIT 
-	#define MOTION_BLUR_EARLY_EXIT 0
-#endif
-
-// helper function, from the paper but adjusted to avoid division by zero
-float cone(float2 X, float2 Y, float2 v)
-{
-	// to avoid div by 0
-	float Bias = 0.001f;
-
-	// better for no background velocity
-	return length(X - Y) < length(v);
-}
-
-// helper function, from the paper but adjusted to avoid division by zero
-float cylinder(float2 X, float2 Y, float2 v)
-{
-	// to avoid div by 0
-	float Bias = 0.001f;
-
-	return 1 - smoothstep(0.95f * length(v), 1.05f * length(v) + Bias, length(X - Y) );
-}
-
-// helper function, from the paper
-// note this assumes negative z values
-// is zb closer than za?
-float softDepthCompare(float za, float zb)
-{
-	return saturate(1 - (za - zb) / SOFT_Z_EXTENT);
-}
-
-
-// ------------------------------------------
-
-// MOTION_BLUR_QUALITY == 0:visualize, 1:low, 2:medium, 3:high, 4:very high
-
-// to scale to normalized motionblur velocity
-// xy:includes y flip, zw:unused
-float4 VelocityScale;
-
-// Last frame's view projection matrix (world-camera) to clip)
-float4x4 PrevViewProjMatrix;
-
-// .xy multiply, .zw:add 
-// to transform the UV to a normalized view  0..1
-float4 TextureViewMad;
-
-// xy:IndividualVelocityScale.xy zw:unused, from postprocess settings
-float4 MotionBlurParameters;
-
-// ------------------------------------------
-
-
-// debug motionblur (very useful, keep)
-// @param ScreenPos -1..1 -1..1 for viewport
-// @param Velocity in -1..1 range for full motionblur
-// @apram Color RGB and depth in alpha
-// @param AvgObject 0:background, 1:foregound
-void OverrideWithTestChart(float2 ScreenPos, inout float2 ObjectVelocity, inout float2 BackgroundVelocity, inout float4 Color, inout float AvgObject)
-{
-#if MOTIONBLUR_TESTCHART == 1
-	// needs to be inside the loop to prevent NVIDIA driver optimizetion (blinking)
-	float2 PixelPos = ScreenPos * ScreenPosToPixel.xy + ScreenPosToPixel.zw + 0.5f - 25;
-	float3 BackgroundColor = lerp(0.0, 0.3f, PseudoRandom(PixelPos)); 
-	float3 ForegroundColor = lerp(float3(1, 0, 0), float3(1, 1, 0), PseudoRandom(PixelPos)); 
-
-	int2 tile = (int2)floor(PixelPos / 12.0f);
-	int2 experiment = (int2)floor(tile / 5.0f);
-
-	if(experiment.x >= 0 && experiment.y >= 0 && experiment.x < 10 && experiment.y < 5)
-	{
-		int2 localtile = uint2(tile) % 5;
-
-		bool bForeground = localtile.x == 2 && localtile.y == 2;
-
-		Color.rgb = bForeground ? ForegroundColor : BackgroundColor;
-
-		// depth
-		Color.a = bForeground ? 100.0f : 1000.0f;
-
-		bool bLeftSide = experiment.x < 5;
-
-		if(!bLeftSide)
-		{
-			experiment.x -= 5;
-		}
-
-		float ForegroundAngle = (experiment.x - 1) * (6.283f / 12);
-		float BackgroundAngle = (experiment.y - 1) * (6.283f / 12) + 3.1415f/2;
-
-		// ForegroundR with very small amounts needs extra testing so we do a non linear scale
-		float ForegroundR = pow(experiment.x / 5.0f, 2);
-		float BackgroundR = pow(experiment.y / 5.0f, 2);
-
-		float2 ForegroundXY = ForegroundR * float2(sin(ForegroundAngle), cos(ForegroundAngle));
-		float2 BackgroundXY = BackgroundR * float2(sin(BackgroundAngle), cos(BackgroundAngle));
-
-		BackgroundVelocity.xy = BackgroundXY;
-
-		if(bLeftSide)
-		{
-			ObjectVelocity.xy = ForegroundXY;
-			AvgObject = bForeground;
-		}
-		else
-		{
-			ObjectVelocity.xy = bForeground ? ForegroundXY : BackgroundXY;
-			AvgObject = 1.0f;
-		}
-	}
-
-#endif
-}
-
-// used to visualize the motion blur
-// @return 0/1
-float Compute3DCheckerBoard(float3 Pos)
-{
-	float3 TiledWorldPos = frac(Pos) > 0.5f;
-	return (float)((uint)dot(float3(1,1,1), TiledWorldPos) % 2);
-}
-
-void VisualizeMotionBlurPS(
-	noperspective float4 UVAndScreenPos : TEXCOORD0,
-	out float4 OutColor : SV_Target0)
-{
-	const float TileSize = 48;
-
-	float2 ScreenSpacePos = UVAndScreenPos.zw;
-	float2 PixelPos = ComputePixelPosCenter(ScreenSpacePos, true);
-	float2 UV = PixelPos * PostprocessInput0Size.zw;
-
-	float2 PixelPosAtTileCenter = PixelPos - (frac(PixelPos / TileSize) - 0.5f) * TileSize;
-	float2 UVAtTileCenter = PixelPosAtTileCenter * PostprocessInput0Size.zw;
-	float2 ScreenPosAtTileCenter = ScreenPosFromPixelPos(PixelPosAtTileCenter, true);
-
-	// World aligned checkerboards
-	{
-		float DeviceZ = LookupDeviceZ(UV);
-		float SceneDepth = ConvertFromDeviceZ(DeviceZ);
-
-		float3 ScreenVector = mul(float4(ScreenSpacePos, 1, 0), View.ScreenToWorld).xyz;
-
-		// world space position of the current pixel
-		float3 OffsetWorldPos = ScreenVector * SceneDepth;
-
-		float3 AbsWorldPos = View.WorldCameraOrigin + OffsetWorldPos;
-		float3 WorldCheckerboard = Compute3DCheckerBoard(AbsWorldPos * 0.02f) * 0.1f + Compute3DCheckerBoard(AbsWorldPos * 0.002f) * 0.3f + Compute3DCheckerBoard(AbsWorldPos * 0.0002f) * 0.6f;
-		OutColor = float4(lerp(WorldCheckerboard, float3(0,0,0), 0.7f), 1);
-	}
-
-	bool bSelectorOpaqueAtTileCenter = PostprocessInput2.SampleLevel(PostprocessInput2Sampler, UVAtTileCenter, 0).x > 0;
-	bool bSelectorOpaque = PostprocessInput2.SampleLevel(PostprocessInput2Sampler, UV, 0).x > 0;
-
-	// relative, in screen space -1...1 -1..1, can be even outside of that range, points into the movement direction
-	float2 VelocityAtTileCenter;
-	{
-		VelocityAtTileCenter = DecodeVelocityFromTexture(PostprocessInput2.SampleLevel(PostprocessInput2Sampler, UVAtTileCenter, 0).xy);
-
-		// reconstruct from camera motion if we don't have velocity data
-		if(!bSelectorOpaqueAtTileCenter)
-		{
-			// UV is per tile
-			float4 ThisClip = float4(ScreenPosAtTileCenter, LookupDeviceZ(UVAtTileCenter), 1);
-			float4 PrevClip = mul( ThisClip, View.ClipToPrevClip );
-			float2 PrevScreen = PrevClip.xy / PrevClip.w;
-
-			// points into the movement direction
-			VelocityAtTileCenter = ScreenPosAtTileCenter - PrevScreen;
-		}
-	}
-
-	// tint yellow if velocity data is stored in texture
-	{
-		float3 TintColor = bSelectorOpaque ? float3(0.5f, 0.5f, 0.2f) : float3(0.5f, 0.5f, 0.5f);
-	
-		OutColor.rgb = lerp(OutColor.rgb, TintColor, 0.45f);
-	}
-
-	// tile center
-	{
-		float2 Delta = PixelPos - PixelPosAtTileCenter;
-
-		float HMask = lerp(saturate(abs(Delta.y)), 1, saturate(abs(Delta.x) / TileSize * 4));
-		float VMask = lerp(saturate(abs(Delta.x)), 1, saturate(abs(Delta.y) / TileSize * 4));
-
-		float Dist = length(Delta);
-		OutColor.rgb *= lerp(1, 0.3f, saturate(3 - Dist));
-	}
-
-	float3 LineColor = bSelectorOpaqueAtTileCenter ? float3(1,1,0) : float3(0.7,0.7,0.7);
-
-	// points into the movement direction
-	float2 PixelDirection = VelocityAtTileCenter * ScreenPosToPixel.xy;
-
-	// arrow
-	{
-		float2 PerpPixelDirection = float2(PixelDirection.y, -PixelDirection.x);
-		float2 DirectionInTile = PixelPos - PixelPosAtTileCenter;
-
-		float DistOnLine = dot(normalize(-PixelDirection), DirectionInTile) + length(PixelDirection);
-
-		bool bArrowHead = DistOnLine < 8;
-
-//		float LocalThickness = bArrowHead ? (DistOnLine * 0.5f) : 0.9f;
-//		float LocalThickness = bArrowHead ? (DistOnLine * 0.3f) : 1 + (frac(DistOnLine/8)*8)*0.2f;
-		float LocalThickness = 1 + (frac(DistOnLine/8)*8)*0.25f;
-
-		float PerpDirectionMask = saturate(LocalThickness - abs(dot(normalize(PerpPixelDirection), DirectionInTile)));
-		float DirectionMask = saturate(length(PixelDirection) - length(DirectionInTile));
-
-		float3 LineMask = PerpDirectionMask * DirectionMask;
-		OutColor.rgb = lerp(OutColor.rgb, LineColor, LineMask);
-	}
-
-	// previous pos is a dot
-	{
-		float3 DotColor = float3(0,1,0);
-		// PixelPos of the previous position
-		float2 PreviousPixelPos = PixelPosAtTileCenter - PixelDirection;
-		float Dist = length(PreviousPixelPos - PixelPos);
-		OutColor.rgb = lerp(OutColor.rgb, LineColor, saturate(4 - Dist));
-		OutColor.rgb = lerp(OutColor.rgb, 0, saturate(2.5f - Dist));
-	}
-}
-
-// ------------------------------------------
-
-// motion blur setup vertex shader
-void SetupVS(
-	in float4 InPosition : ATTRIBUTE0,
-	in float2 InTexCoord : ATTRIBUTE1,
-	out noperspective float2 OutUV: TEXCOORD0,
-	out noperspective float2 OutUVs[4]: TEXCOORD1,
-	out float4 OutPosition : SV_POSITION
-	)
-{
-	DrawRectangle(InPosition, InTexCoord, OutPosition, OutUV);
-
-//	float2 ScreenPos =  OutPosition.xy;
-
-	float2 HalfPixelOffset = PostprocessInput1Size.zw * float2(0.5f, 0.5f);
-
-	// no filtering (2x2 kernel) to get no leaking in Depth of Field, lefttop,righttop,leftbottom,rightbottom
-	OutUVs[0] = OutUV + HalfPixelOffset * float2(-1, -1);
-	OutUVs[1] = OutUV + HalfPixelOffset * float2( 1, -1);
-	OutUVs[2] = OutUV + HalfPixelOffset * float2(-1,  1);
-	OutUVs[3] = OutUV + HalfPixelOffset * float2( 1,  1);
-}
-
-// motion blur setup pixel shader, downsamples to half resolution
-// MRT0: velocity in normalized_motionblur_velocity and mask in 0:background..1:object (needed for soft masked)
-// MRT1: color and depth in alpha
-void SetupPS(
-	noperspective float2 UV : TEXCOORD0,				// UV into the full resolution source RT
-	noperspective float2 UVs[4] : TEXCOORD1,			// UV for high quality 4 samples with blurring
-	out float4 OutVelocity : SV_Target0,
-	out float4 OutColor : SV_Target1)
-{
-	// detect find objects
-	float4 ObjectMask;
-	float2 VelocitySamples[4];
-
-#if GATHER_OPTIMIZATION
-	{
-		// using Gather: xyzw in counter clockwise order starting with the sample to the lower left of the queried location
-		float4 Red = PostprocessInput0.GatherRed(PostprocessInput0Sampler, UV);
-		float4 Green = PostprocessInput0.GatherGreen(PostprocessInput0Sampler, UV);
-
-		VelocitySamples[0] = float2(Red.w, Green.w);
-		VelocitySamples[1] = float2(Red.z, Green.z);
-		VelocitySamples[2] = float2(Red.x, Green.x);
-		VelocitySamples[3] = float2(Red.y, Green.y);
-	}
-#else
-	UNROLL for( int i = 0; i < 4; i++ )
-	{
-		VelocitySamples[i] = PostprocessInput0.SampleLevel(PostprocessInput0Sampler, UVs[i], 0).xy;
-	}
-#endif
-
-	UNROLL for( int i = 0; i < 4; i++ )
-	{
-		ObjectMask[i] = VelocitySamples[i].x > 0;
-		VelocitySamples[i] = DecodeVelocityFromTexture(VelocitySamples[i]);
-		VelocitySamples[i] = ObjectMask[i] ? VelocitySamples[i] : 0;
-	}
-
-	float AvgObject = dot(ObjectMask, 0.25f);
-
-	ObjectMask = (AvgObject > 0.5f) ? ObjectMask : (1 - ObjectMask);
-
-	// bias to avoid div by 0
-	float InvTotalWeight = 1.0 / ( dot( ObjectMask, 1 ) + 0.000001f );
-
-	float4 SumColorAndDepth = 0;
-
-#if GATHER_OPTIMIZATION
-	float4 DepthValues = PostprocessInput2.GatherRed(PostprocessInput2Sampler, UV, 0).r;
-#else
-	float4 DepthValues = float4(
-		PostprocessInput2.SampleLevel(PostprocessInput2Sampler, UVs[2], 0).r,
-		PostprocessInput2.SampleLevel(PostprocessInput2Sampler, UVs[3], 0).r,
-		PostprocessInput2.SampleLevel(PostprocessInput2Sampler, UVs[1], 0).r,
-		PostprocessInput2.SampleLevel(PostprocessInput2Sampler, UVs[0], 0).r);
-#endif
-
-	SumColorAndDepth += float4(PostprocessInput1.SampleLevel(PostprocessInput1Sampler, UVs[0], 0).rgb, DepthValues.x) * ObjectMask.x;
-	SumColorAndDepth += float4(PostprocessInput1.SampleLevel(PostprocessInput1Sampler, UVs[1], 0).rgb, DepthValues.y) * ObjectMask.y;
-	SumColorAndDepth += float4(PostprocessInput1.SampleLevel(PostprocessInput1Sampler, UVs[2], 0).rgb, DepthValues.z) * ObjectMask.z;
-	SumColorAndDepth += float4(PostprocessInput1.SampleLevel(PostprocessInput1Sampler, UVs[3], 0).rgb, DepthValues.w) * ObjectMask.w;
-
-	OutColor = SumColorAndDepth * InvTotalWeight;
-
-	// clamp to avoid artifacts from exceeding fp16 through framebuffer blending of multiple very bright lights
-	OutColor.rgb = min(float3(256 * 256, 256 * 256, 256 * 256), OutColor.rgb);
-
-	OutColor.a = ConvertFromDeviceZ(OutColor.a);
-
-	{
-		float2 SumVelocity = 0;
-
-		FLATTEN for( int i = 0; i < 4; i++ )
-		{
-			SumVelocity += VelocitySamples[i] * ObjectMask[i];
-		}
-
-		OutVelocity.xy = SumVelocity * InvTotalWeight * VelocityScale.xy;
-	}
-
-	// for debugging
-/*
-	{
-		float2 BackgroundVelocity = 0;
-
-		OverrideWithTestChart(UVAndScreenPos.zw, OutVelocity.xy, BackgroundVelocity, OutColor, AvgObject);
-	}
-*/
-
-	// 0:background, 1:object layer
-	OutVelocity.b = AvgObject > 0.5f;
-
-	// clamp motion vector in a disc from -1 to 1 (the maximum motion vector range)
-	{
-		half Len = dot(OutVelocity.xy, OutVelocity.xy);
-
-		float ScaleFix = rsqrt(Len);
-		
-		FLATTEN if(Len < 1)
-		{
-			ScaleFix = 1.0f;
-		}
-		FLATTEN if(Len < 0.01f)
-		{
-			ScaleFix = 0;
-		}
-
-		OutVelocity.xy *= ScaleFix * OutVelocity.b;
-	}
-
-	// alpha is used to normalize the velocity after blurring (0:background, 1:object)
-	OutVelocity.a = 1;
-
-	// debug, uncomment to geneate small and large horizontal motion and the objects
-//	OutVelocity.xy = lerp(float2(0,0), (InUV.y > 0.8f) ? float2(1,0) : float2(0.02f,0), OutVelocity.b);
-}
-
-uint GetStepCountFromQuality()
-{
-#if MOTION_BLUR_QUALITY == 1
-	return 4;
-#elif MOTION_BLUR_QUALITY == 2
-	return 6;
-#elif MOTION_BLUR_QUALITY == 3
-	return 8;
-#else // MOTION_BLUR_QUALITY == 4
-	return 16;
-#endif
-}
-
-// motionblur pixel shader
-// input:
-//   0: RGB:scene color, A: depth in half resolution from MotionBlurSetup
-//   1: blurred quarter res velocity from MotionBlurSetup
-//   2: half res velocity from MotionBlurSetup
-// half resolution output: RGB: color, A:blend to full res factor
-void MainPS(noperspective float4 UVAndScreenPos : TEXCOORD0, float4 SvPosition : SV_POSITION, out float4 OutColor : SV_Target0)
-{
-	// Screen Quad UV 0..1
-	float2 UV = UVAndScreenPos.xy;
-	// screen position in [-1, 1] screen space
-	float2 ScreenSpacePos = UVAndScreenPos.zw;
-
-	OutColor = 0;
-
-	// RGB:color, A:depth
-	float4 ColorAndDepth = PostprocessInput0.SampleLevel(PostprocessInput0Sampler, UV.xy, 0);
-
-	// can be moved to VS
-	float3 ScreenVector = mul(float4(ScreenSpacePos, 1, 0), View.ScreenToWorld).xyz;
-
-	// world space position of the current pixel
-	float3 OffsetWorldPos = ScreenVector * ColorAndDepth.a;
-	// previous frame clip space position of the current pixel
-	float4 PrevClipPos = mul(float4(OffsetWorldPos, 1), PrevViewProjMatrix);
-	// previous frame screen coordinates of the current pixel
-	float3 PrevScreenCoord = PrevClipPos.xyz / PrevClipPos.w;
-	
-	// we split the content in a object and background layer
-
-	// background velocity in half_res_pixels
-	float2 PixelBackgroundVelocity;
-	{
-		// the background velocity in normalized_motionblur_velocity
-		float2 NormBackgroundVelocity;
-		{
-			NormBackgroundVelocity = (UVAndScreenPos.zw - PrevScreenCoord.xy) * MotionBlurParameters.xy;
-
-			// for debugging
-			{
-				float2 Velocity = 0;
-				float4 Color = 0;
-				float AvgObject = 0;
-			
-				OverrideWithTestChart(ScreenSpacePos, Velocity, NormBackgroundVelocity, OutColor, AvgObject);
-			}
-
-			// clamp the max motionblur within the unit radius 
-			float MotionLength = length(NormBackgroundVelocity);
-
-			FLATTEN if(MotionLength > 1.0f)
-			{
-				NormBackgroundVelocity /= MotionLength;
-			}
-		}
-
-		// in (-1..1 -1..1 for the screen) 
-		float2 ScreenBackgroundVelocity = NormBackgroundVelocity * MotionBlurParameters.zw;
-
-
-		// in half_res_pixels
-		PixelBackgroundVelocity = ScreenBackgroundVelocity * ScreenPosToPixel.xy;
-
-
-	}
-
-	// RG: xy motion in normalized_motionblur_velocity, B: object weight 0:background..1:object
-	float3 NormSoftMaskedVelocity;
-	{
-		float4 SoftMaskedTexture = PostprocessInput1.SampleLevel(PostprocessInput1Sampler, UV.xy, 0);
-
-		NormSoftMaskedVelocity = SoftMaskedTexture.rgb / (SoftMaskedTexture.a + 1.0f / 255.0f);
-	}
-
-	// how many motionblur samples affects quality and performance (we do multiple texture lookups per sample)
-	const uint StepCount = GetStepCountFromQuality();
-
-	// we start accumulation with the current pixel with a small weight to define div by near 0 case to be the current pixel color
-	const float4 ColorAccumStartValue = float4(ColorAndDepth.rgb, 1) * 0.001f;
-
-	// RGB:weighed Sum, A:weight to compute average color
-	float4 BackgroundColorAccum = ColorAccumStartValue;
-
-#if MOTION_BLUR_EARLY_EXIT
-	if(true)
-	{
-		BackgroundColorAccum = float4(PostprocessInput0.SampleLevel(PostprocessInput0Sampler, UV.xy, 0).rgb, 1);
-	}
-	else
-	{
-#endif
-
-	// create background motion layer
-	UNROLL for(uint i = 0; i < StepCount; ++i)
-	{
-		float delta = (i / (float)(StepCount - 1)) - 0.5f;
-		float2 LocalUV = UV + delta * PixelBackgroundVelocity * PostprocessInput0Size.zw;	// can be optimized
-		float2 NormalizedView = LocalUV * TextureViewMad.xy + TextureViewMad.zw;
-
-		FLATTEN if(all(NormalizedView > 0 && NormalizedView < 1))
-		{
-			float ObjectSample = PostprocessInput2.SampleLevel(PostprocessInput2Sampler, LocalUV.xy, 0).b;
-
-			ObjectSample = saturate(ObjectSample * 4);
-
-			float4 ColorSample = float4(PostprocessInput0.SampleLevel(PostprocessInput0Sampler, LocalUV.xy, 0).rgb, 1);
-
-			BackgroundColorAccum += ColorSample * (1 - ObjectSample);
-		}
-	}
-
-#if MOTION_BLUR_EARLY_EXIT
-	}
-#endif
-
-	// object velocity in half_res_pixel
-	float2 PixelObjectBlurredVelocity;
-	{
-		// SoftMasked Velocity in normalized_motionblur_velocity, lower resolution and blurred,
-		// divide normalizes which is needed because of gaussian blurs, bias to avoid division by 0
-		half2 NormBlurredVelocity = NormSoftMaskedVelocity.xy / (NormSoftMaskedVelocity.b + 1.0f / 255.0f);
-		
-		// Update the pixel velocity
-		float2 ScreenBlurredVelocity = NormBlurredVelocity * MotionBlurParameters.zw;
-		PixelObjectBlurredVelocity = ScreenBlurredVelocity * ScreenPosToPixel.xy;
-	}
-
-	
-	// --------------------------------------
-	// camera and background
-	float2 PixelCombinedVelocity = lerp(PixelBackgroundVelocity, PixelObjectBlurredVelocity, NormSoftMaskedVelocity.b);
-
-#if MOTION_BLUR_EARLY_EXIT
-	float deBlur = dot(PixelCombinedVelocity, PixelCombinedVelocity);
-	deBlur = sqrt(max(deBlur, 1.0/65536.0));
-	deBlur = saturate((deBlur - 2.0) * (1.0/4.0));
-	if(deBlur == 0.0)
-	{
-		OutColor = float4(0.0, 0.0, 0.0, 0.0);
-		return;
-	}
-	else
-	{
-#endif
-
-	{
-		// in pixels
-		float2 X = UV * PostprocessInput0Size.xy;
-		float4 cXzX = PostprocessInput0.SampleLevel(PostprocessInput0Sampler, UV.xy, 0);
-		float3 cX = cXzX.rgb;
-		float zX = -cXzX.w;		// negated as this was in paper
-		
-		// in pixels
-		float4 vXvX = PostprocessInput2.SampleLevel(PostprocessInput2Sampler, UV.xy, 0);
-		float2 vX = vXvX.xy;
-		vX = (vXvX.b > 0.5f) ? (vX * MotionBlurParameters.zw * ScreenPosToPixel.xy) : PixelBackgroundVelocity;
-
-		// to avoid div by 0
-		float Bias = 1.0f;
-		
-		float StartAlpha = 1 / (length(vX) + Bias);
-
-		float3 BackgroundColor = BackgroundColorAccum.rgb / BackgroundColorAccum.a;
-
-		float ColorAccumSum = 0.0001f;
-
-		float4 ColorAccum = float4(BackgroundColor, 1) * ColorAccumSum;
-		float4 SecondColorAccum = float4(BackgroundColor, 1) * ColorAccumSum;
-
-		UNROLL for(uint e = 0; e < StepCount; ++e)
-		{
-			// we want to have the samples starting from inside going outwards
-			uint iMid = StepCount / 2;
-			uint iSign = (e % 2) * 2 - 1;
-			uint i = iMid + (e / 2) * iSign;
-
-			// -0.5 .. 0.5
-			float delta = (i / (float)(StepCount - 1)) - 0.5f;
-			float2 LocalUV = UV + delta * PixelCombinedVelocity * PostprocessInput0Size.zw;	// can be optimized
-			float2 NormalizedView = LocalUV * TextureViewMad.xy + TextureViewMad.zw;
-
-			FLATTEN if(all(NormalizedView > 0 && NormalizedView < 1))
-			{
-				// in pixels
-				float2 Y = LocalUV * PostprocessInput0Size.xy;
-				float4 cYzY = PostprocessInput0.SampleLevel(PostprocessInput0Sampler, LocalUV.xy, 0);
-				float3 cY = cYzY.rgb;
-				float zY = -cYzY.w;		// negated as this was in paper
-
-				// in pixels
-				float4 vYvY = PostprocessInput2.SampleLevel(PostprocessInput2Sampler, LocalUV.xy, 0);
-				float2 vY = vYvY.xy;
-				vY = (vYvY.b > 0.5f) ? (vY * MotionBlurParameters.zw * ScreenPosToPixel.xy) : PixelBackgroundVelocity;
-//				vY = PixelBackgroundVelocity;
-
-				float f = softDepthCompare(zX, zY);
-				float b = softDepthCompare(zY, zX);
-
-				float ay = 0;
-
-				// Blurry Y in front of any X
-				// the sample we look at (Y) is fast enough to affect me (direction not taken into account)
-				ay += f * cone(Y, X, vY);
-				// Any Y behind blurry X, estimate background
-				ay += b * cone(X, Y, vX); 
-				// Simultaneously blurry X and Y
-				ay += cylinder(Y, X, vY) * cylinder(X, Y, vX) * 2;
-
-				// mask out Y that are not part of the moving direction we current process
-//				ay *= 1 - saturate((length(vY - PixelCombinedVelocity) - length(vY - PixelBackgroundVelocity)) * 0.12f);
-
-				ay = saturate(ay);
-
-				ColorAccum += float4(cY.rgb, 0) * ay;
-				ColorAccumSum += ay;
-
-				float MaskAlreadyFound = saturate(1 - SecondColorAccum.a);
-				SecondColorAccum += float4(cY.rgb, 1) * (1 - ay) * MaskAlreadyFound;
-			}
-		}
-
-		// color content behind motion vector
-//		float3 SecondColor = SecondColorAccum / (StepCount - ColorAccumSum);
-		float3 SecondColor = SecondColorAccum .rgb / SecondColorAccum.a;
-
-		// camera blurred background with the moving object removed
-//		float LerpFactor = length(vX) - 1 > length(PixelBackgroundVelocity);
-		float LerpFactor = saturate((length(PixelBackgroundVelocity) - length(vX)) );	// better
-		float3 BehindMovingObject = lerp(BackgroundColor, SecondColor, LerpFactor);
-
-		OutColor = float4(ColorAccum / ColorAccumSum);
-
-		float BackgroundFraction = length(PixelCombinedVelocity) < length(vX) - 1;
-
-		float MovingObjectFraction = ColorAccumSum / StepCount;
-
-		// take the faster movement (background or non background)
-		OutColor = lerp(float4(BehindMovingObject, 1), OutColor, MovingObjectFraction);
-		
-		OutColor.a = 1;
-
-#if MOTIONBLUR_TESTCHART != 1
-		// compute how much full res should blend in
-		// works but without softedge (seems to be the issues with half res already)
-		// comment tyhe next line to see the half res result
-		OutColor.a = saturate( lerp(length(PixelBackgroundVelocity), length(PixelCombinedVelocity), MovingObjectFraction) * 0.25f);
-#endif
-
-#if MOTION_BLUR_EARLY_EXIT
-		OutColor.a = deBlur;
-#endif
-
-		// prepare for the following compositing pass
-		OutColor.rgb *= OutColor.a;
-	}
-	// --------------------------------------
-
-#if MOTION_BLUR_EARLY_EXIT
-	}
-#endif
-
-
-#if MOTION_BLUR_QUALITY == 0
-	// visualize motion blur
-	{
-		float3 AbsWorldPos = View.WorldCameraOrigin + OffsetWorldPos;
-		float3 WorldCheckerboard = Compute3DCheckerBoard(AbsWorldPos * 0.02f) * 0.1f + Compute3DCheckerBoard(AbsWorldPos * 0.002f) * 0.3f + Compute3DCheckerBoard(AbsWorldPos * 0.0002f) * 0.6f;
-
-		OutColor = float4(lerp(WorldCheckerboard, OutColor.rgb, 0.7f), 1);
-
-		float2 NewPixelPosCentered = UV * PostprocessInput0Size.xy;
+// from the paper: We use SOFT Z EXTENT = 1mm to 10cm for our results
+#define SOFT_Z_EXTENT 1
 
-		// -1..1
-		float2 NormalizedDirection = PixelObjectBlurredVelocity / ScreenPosToPixel.xy * MotionBlurParameters.xy;
-		// -1..1
-		float2 ScreenLocalTilePos = frac(NewPixelPosCentered / 16.0f) * 2 * 1.2f - 1;
-		// -1..1
-		float2 PerpDirection = float2(ScreenLocalTilePos.y, -ScreenLocalTilePos.x);
+// 0:off / 1:on, useful to debug the motionblur algorithm
+#define MOTIONBLUR_TESTCHART 0
 
-		float DirectionMask = 1 - saturate(abs(dot(PerpDirection, normalize(NormalizedDirection))) * 6);
-		float StrengthMask = 1 - saturate((length(PerpDirection) - length(NormalizedDirection)) * 6);
-		float OrientationMask = saturate(dot(normalize(ScreenLocalTilePos), NormalizedDirection) * 6);
-		float DiskMask = saturate((length(PerpDirection) - 1.0f) * 6);
+// --------------------------------------------------------------------------
 
-		bool bSelectorOpaque = PostprocessInput2.SampleLevel(PostprocessInput2Sampler, UV.xy, 0).b > 0.5f;
+// MOTION_BLUR_QUALITY == 0:visualize, 1:low, 2:medium, 3:high, 4:very high
 
-		float3 LineColor = lerp(float3(1,0,0), float3(0,1,0), OrientationMask); 
-		float3 LineMask = DirectionMask * StrengthMask; 
-		float3 TintColor = bSelectorOpaque ? float3(0.2f, 0.2f, 0.6f) : float3(0.5f, 0.5f, 0.5f);
+// to scale to normalized motionblur velocity
+// xy:includes y flip, zw:unused
+float4 VelocityScale;
 
-		OutColor.rgb = lerp(lerp(WorldCheckerboard, TintColor, 0.5f), LineColor, LineMask);
-		OutColor.rgb *= lerp(1.0f, 0.9f, DiskMask);
-	}
-#endif
-}
+// xy:IndividualVelocityScale.xy zw:unused, from postprocess settings
+float4 MotionBlurParameters;
 
+// ------------------------------------------
 
-void MainRecombinePS(noperspective float4 UVAndScreenPos : TEXCOORD0, out float4 OutColor : SV_Target0)
+uint GetStepCountFromQuality()
 {
-	float2 UV = UVAndScreenPos.xy;
-
-	float3 FullResSceneColor = PostprocessInput0.SampleLevel(PostprocessInput0Sampler, UV, 0).rgb;
-	float4 MotionBlurOutput = PostprocessInput1.SampleLevel(PostprocessInput1Sampler, UV, 0);
-	
-	float Mask = MotionBlurOutput.a;
-
-	OutColor.rgb = FullResSceneColor * (1 - Mask) + MotionBlurOutput.rgb;
-	OutColor.a = 0;
+#if MOTION_BLUR_QUALITY == 1
+	return 4;
+#elif MOTION_BLUR_QUALITY == 2
+	return 6;
+#elif MOTION_BLUR_QUALITY == 3
+	return 8;
+#else // MOTION_BLUR_QUALITY == 4
+	return 16;
+#endif
 }
 
-
-
-
-
 uint	bDrawMax;
 
 void VelocityScatterVS(
@@ -816,7 +125,7 @@ void VelocityScatterPS(
 
 
 
-
+//-----------------------------------
 
 float2 DepthCmp( float CenterDepth, float SampleDepth, float DepthScale )
 {
@@ -882,7 +191,7 @@ float GetVelocityLengthPixels( float2 EncodedVelocity )
 #endif
 }
 
-void MainNewPS(
+void MainPS(
 	in float4 UVAndScreenPos : TEXCOORD0,
 	in float4 SvPosition : SV_Position,		// after all interpolators
 	out float4 OutColor : SV_Target0
@@ -1027,4 +336,181 @@ void MainNewPS(
 
 	//OutColor.rgb *= bFastPath ? float3(1,0,0) : float3(0,1,0);
 	//OutColor.rgb = lerp( OutColor.rgb, float3( abs( MinMaxVelocity.zw ), 0 ) * 0.2, 0.9 );
+}
+
+
+// ------------------------------------------
+
+
+// debug motionblur (very useful, keep)
+// @param ScreenPos -1..1 -1..1 for viewport
+// @param Velocity in -1..1 range for full motionblur
+// @apram Color RGB and depth in alpha
+// @param AvgObject 0:background, 1:foregound
+void OverrideWithTestChart(float2 ScreenPos, inout float2 ObjectVelocity, inout float2 BackgroundVelocity, inout float4 Color, inout float AvgObject)
+{
+#if MOTIONBLUR_TESTCHART == 1
+	// needs to be inside the loop to prevent NVIDIA driver optimizetion (blinking)
+	float2 PixelPos = ScreenPos * ScreenPosToPixel.xy + ScreenPosToPixel.zw + 0.5f - 25;
+	float3 BackgroundColor = lerp(0.0, 0.3f, PseudoRandom(PixelPos)); 
+	float3 ForegroundColor = lerp(float3(1, 0, 0), float3(1, 1, 0), PseudoRandom(PixelPos)); 
+
+	int2 tile = (int2)floor(PixelPos / 12.0f);
+	int2 experiment = (int2)floor(tile / 5.0f);
+
+	if(experiment.x >= 0 && experiment.y >= 0 && experiment.x < 10 && experiment.y < 5)
+	{
+		int2 localtile = uint2(tile) % 5;
+
+		bool bForeground = localtile.x == 2 && localtile.y == 2;
+
+		Color.rgb = bForeground ? ForegroundColor : BackgroundColor;
+
+		// depth
+		Color.a = bForeground ? 100.0f : 1000.0f;
+
+		bool bLeftSide = experiment.x < 5;
+
+		if(!bLeftSide)
+		{
+			experiment.x -= 5;
+		}
+
+		float ForegroundAngle = (experiment.x - 1) * (6.283f / 12);
+		float BackgroundAngle = (experiment.y - 1) * (6.283f / 12) + 3.1415f/2;
+
+		// ForegroundR with very small amounts needs extra testing so we do a non linear scale
+		float ForegroundR = pow(experiment.x / 5.0f, 2);
+		float BackgroundR = pow(experiment.y / 5.0f, 2);
+
+		float2 ForegroundXY = ForegroundR * float2(sin(ForegroundAngle), cos(ForegroundAngle));
+		float2 BackgroundXY = BackgroundR * float2(sin(BackgroundAngle), cos(BackgroundAngle));
+
+		BackgroundVelocity.xy = BackgroundXY;
+
+		if(bLeftSide)
+		{
+			ObjectVelocity.xy = ForegroundXY;
+			AvgObject = bForeground;
+		}
+		else
+		{
+			ObjectVelocity.xy = bForeground ? ForegroundXY : BackgroundXY;
+			AvgObject = 1.0f;
+		}
+	}
+
+#endif
+}
+
+// used to visualize the motion blur
+// @return 0/1
+float Compute3DCheckerBoard(float3 Pos)
+{
+	float3 TiledWorldPos = frac(Pos) > 0.5f;
+	return (float)((uint)dot(float3(1,1,1), TiledWorldPos) % 2);
+}
+
+void VisualizeMotionBlurPS(
+	noperspective float4 UVAndScreenPos : TEXCOORD0,
+	out float4 OutColor : SV_Target0)
+{
+	const float TileSize = 48;
+
+	float2 ScreenSpacePos = UVAndScreenPos.zw;
+	float2 PixelPos = ComputePixelPosCenter(ScreenSpacePos, true);
+	float2 UV = PixelPos * PostprocessInput0Size.zw;
+
+	float2 PixelPosAtTileCenter = PixelPos - (frac(PixelPos / TileSize) - 0.5f) * TileSize;
+	float2 UVAtTileCenter = PixelPosAtTileCenter * PostprocessInput0Size.zw;
+	float2 ScreenPosAtTileCenter = ScreenPosFromPixelPos(PixelPosAtTileCenter, true);
+
+	// World aligned checkerboards
+	{
+		float DeviceZ = LookupDeviceZ(UV);
+		float SceneDepth = ConvertFromDeviceZ(DeviceZ);
+
+		float3 ScreenVector = mul(float4(ScreenSpacePos, 1, 0), View.ScreenToWorld).xyz;
+
+		// world space position of the current pixel
+		float3 OffsetWorldPos = ScreenVector * SceneDepth;
+
+		float3 AbsWorldPos = View.WorldCameraOrigin + OffsetWorldPos;
+		float3 WorldCheckerboard = Compute3DCheckerBoard(AbsWorldPos * 0.02f) * 0.1f + Compute3DCheckerBoard(AbsWorldPos * 0.002f) * 0.3f + Compute3DCheckerBoard(AbsWorldPos * 0.0002f) * 0.6f;
+		OutColor = float4(lerp(WorldCheckerboard, float3(0,0,0), 0.7f), 1);
+	}
+
+	bool bSelectorOpaqueAtTileCenter = PostprocessInput2.SampleLevel(PostprocessInput2Sampler, UVAtTileCenter, 0).x > 0;
+	bool bSelectorOpaque = PostprocessInput2.SampleLevel(PostprocessInput2Sampler, UV, 0).x > 0;
+
+	// relative, in screen space -1...1 -1..1, can be even outside of that range, points into the movement direction
+	float2 VelocityAtTileCenter;
+	{
+		VelocityAtTileCenter = DecodeVelocityFromTexture(PostprocessInput2.SampleLevel(PostprocessInput2Sampler, UVAtTileCenter, 0).xy);
+
+		// reconstruct from camera motion if we don't have velocity data
+		if(!bSelectorOpaqueAtTileCenter)
+		{
+			// UV is per tile
+			float4 ThisClip = float4(ScreenPosAtTileCenter, LookupDeviceZ(UVAtTileCenter), 1);
+			float4 PrevClip = mul( ThisClip, View.ClipToPrevClip );
+			float2 PrevScreen = PrevClip.xy / PrevClip.w;
+
+			// points into the movement direction
+			VelocityAtTileCenter = ScreenPosAtTileCenter - PrevScreen;
+		}
+	}
+
+	// tint yellow if velocity data is stored in texture
+	{
+		float3 TintColor = bSelectorOpaque ? float3(0.5f, 0.5f, 0.2f) : float3(0.5f, 0.5f, 0.5f);
+	
+		OutColor.rgb = lerp(OutColor.rgb, TintColor, 0.45f);
+	}
+
+	// tile center
+	{
+		float2 Delta = PixelPos - PixelPosAtTileCenter;
+
+		float HMask = lerp(saturate(abs(Delta.y)), 1, saturate(abs(Delta.x) / TileSize * 4));
+		float VMask = lerp(saturate(abs(Delta.x)), 1, saturate(abs(Delta.y) / TileSize * 4));
+
+		float Dist = length(Delta);
+		OutColor.rgb *= lerp(1, 0.3f, saturate(3 - Dist));
+	}
+
+	float3 LineColor = bSelectorOpaqueAtTileCenter ? float3(1,1,0) : float3(0.7,0.7,0.7);
+
+	// points into the movement direction
+	float2 PixelDirection = VelocityAtTileCenter * ScreenPosToPixel.xy;
+
+	// arrow
+	{
+		float2 PerpPixelDirection = float2(PixelDirection.y, -PixelDirection.x);
+		float2 DirectionInTile = PixelPos - PixelPosAtTileCenter;
+
+		float DistOnLine = dot(normalize(-PixelDirection), DirectionInTile) + length(PixelDirection);
+
+		bool bArrowHead = DistOnLine < 8;
+
+//		float LocalThickness = bArrowHead ? (DistOnLine * 0.5f) : 0.9f;
+//		float LocalThickness = bArrowHead ? (DistOnLine * 0.3f) : 1 + (frac(DistOnLine/8)*8)*0.2f;
+		float LocalThickness = 1 + (frac(DistOnLine/8)*8)*0.25f;
+
+		float PerpDirectionMask = saturate(LocalThickness - abs(dot(normalize(PerpPixelDirection), DirectionInTile)));
+		float DirectionMask = saturate(length(PixelDirection) - length(DirectionInTile));
+
+		float3 LineMask = PerpDirectionMask * DirectionMask;
+		OutColor.rgb = lerp(OutColor.rgb, LineColor, LineMask);
+	}
+
+	// previous pos is a dot
+	{
+		float3 DotColor = float3(0,1,0);
+		// PixelPos of the previous position
+		float2 PreviousPixelPos = PixelPosAtTileCenter - PixelDirection;
+		float Dist = length(PreviousPixelPos - PixelPos);
+		OutColor.rgb = lerp(OutColor.rgb, LineColor, saturate(4 - Dist));
+		OutColor.rgb = lerp(OutColor.rgb, 0, saturate(2.5f - Dist));
+	}
 }
\ No newline at end of file
diff --git a/Engine/Shaders/PostProcessTemporalCommon.usf b/Engine/Shaders/PostProcessTemporalCommon.usf
--- a/Engine/Shaders/PostProcessTemporalCommon.usf
+++ b/Engine/Shaders/PostProcessTemporalCommon.usf
@@ -7,9 +7,6 @@ This is the common bits of the temporal AA shader
 which can be configured via defines for various temporal AA passes.
 =============================================================================*/
 
-// 0 / 1(default) This can be set to 0 to debug the velocity vectors
-#define NEIGHBORHOOD_CLAMPING 1
-
 // 1 = Use tighter AABB clamp for history.
 // 0 = Use simple min/max clamp.
 #ifndef AA_AABB
diff --git a/Engine/Shaders/SceneCapturePixelShader.usf b/Engine/Shaders/SceneCapturePixelShader.usf
new file mode 100644
--- /dev/null
+++ b/Engine/Shaders/SceneCapturePixelShader.usf
@@ -0,0 +1,30 @@
+// Copyright 1998-2016 Epic Games, Inc. All Rights Reserved.
+
+/*=============================================================================
+	SceneCapturePixelShader.usf: 
+=============================================================================*/
+
+#include "Common.usf"
+#include "DeferredShadingCommon.usf"
+
+void Main(
+	FScreenVertexOutput Input,
+	out float4 OutColor : SV_Target0
+	)
+{
+#if SOURCE_MODE_SCENE_COLOR_AND_OPACITY
+	float4 SceneColorValue = CalcFullSceneColor(Input.UV);
+	// Convert scene color alpha to opacity that that a translucent blend mode expects (0: no foreground content, 1: full foreground content)
+	OutColor = float4(SceneColorValue.rgb, saturate(1.0f - SceneColorValue.a));
+#elif SOURCE_MODE_SCENE_COLOR_SCENE_DEPTH
+	OutColor = float4(CalcSceneColor(Input.UV), CalcSceneDepth(Input.UV));
+#elif SOURCE_MODE_SCENE_DEPTH 
+	OutColor = float4(CalcSceneDepth(Input.UV), 0, 0, 0);
+#elif SOURCE_MODE_NORMAL
+	FGBufferData GBufferData = GetGBufferData(Input.UV);
+	OutColor = float4(GBufferData.WorldNormal, 0);
+#elif SOURCE_MODE_BASE_COLOR
+	FGBufferData GBufferData = GetGBufferData(Input.UV);
+	OutColor = float4(GBufferData.BaseColor, 0);
+#endif
+}
diff --git a/Engine/Shaders/SkyLighting.usf b/Engine/Shaders/SkyLighting.usf
--- a/Engine/Shaders/SkyLighting.usf
+++ b/Engine/Shaders/SkyLighting.usf
@@ -74,6 +74,9 @@ void SkyLightDiffusePS(
 		
 		float ScalarFactors = lerp(SkyVisibility, 1, OcclusionTintAndMinOcclusion.w);
 
+		float3 DiffuseColor = GBuffer.DiffuseColor;
+
+		BRANCH 
 		if (ShadingModelId == SHADINGMODELID_TWOSIDED_FOLIAGE)
 		{
 			float3 SubsurfaceLookup = GetSkySHDiffuse(-GBuffer.WorldNormal) * Frame.SkyLightColor.rgb;
@@ -81,20 +84,26 @@ void SkyLightDiffusePS(
 			Lighting += ScalarFactors * SubsurfaceLookup * SubsurfaceColor;
 		}
 
-		float3 DiffuseColor = GBuffer.DiffuseColor;
+		if (ShadingModelId == SHADINGMODELID_SUBSURFACE || ShadingModelId == SHADINGMODELID_PREINTEGRATED_SKIN)
+		{
+			float3 SubsurfaceColor = ExtractSubsurfaceColor(GBuffer);
+			// Add subsurface energy to diffuse
+			DiffuseColor += SubsurfaceColor;
+		}
 
-		BRANCH if( GBuffer.ShadingModelID == SHADINGMODELID_HAIR )
+		BRANCH 
+		if (ShadingModelId == SHADINGMODELID_HAIR)
 		{
 			float3 N = GBuffer.WorldNormal;
 			float3 V = -normalize(mul(float4(UVAndScreenPos.zw, 1, 0), View.ScreenToWorld).xyz);
-			float3 L = normalize( V - N * dot(V,N) );
+			float3 L = normalize(V - N * dot(V,N));
 			SkyLightingNormal = L;
 			DiffuseColor = PI * HairShading( GBuffer, L, V, N, 1, 0, 0.2, uint2(0,0) );
 		}
 
-		BRANCH if( GBuffer.ShadingModelID == SHADINGMODELID_CLOTH )
+		if (ShadingModelId == SHADINGMODELID_CLOTH)
 		{
-			DiffuseColor +=  GBuffer.CustomData.rgb * GBuffer.CustomData.a;
+			DiffuseColor += GBuffer.CustomData.rgb * GBuffer.CustomData.a;
 		}
 
 		// Compute the preconvolved incoming lighting with the bent normal direction
diff --git a/Engine/Shaders/SkyLightingShared.usf b/Engine/Shaders/SkyLightingShared.usf
--- a/Engine/Shaders/SkyLightingShared.usf
+++ b/Engine/Shaders/SkyLightingShared.usf
@@ -35,28 +35,36 @@ void GetDistanceFieldAOSpecularOcclusion(float2 ScreenUV, float3 ReflectionVecto
 		float3 BentNormalAO = TextureValue.xyz;
 		float BentNormalLength = length(BentNormalAO);
 
-#define USE_CONE_TRACED_SPECULAR_OCCLUSION 0
-#if USE_CONE_TRACED_SPECULAR_OCCLUSION
-		IndirectSpecularOcclusion = TextureValue.w;
-#else
 		BRANCH
-		if (bTwoSidedFoliage)
+		if (Frame.DistanceFieldAOSpecularOcclusionMode == 0)
 		{
-			IndirectSpecularOcclusion = BentNormalLength;
+			IndirectSpecularOcclusion = IndirectDiffuseOcclusion;
 		}
-		else
+		else if (Frame.DistanceFieldAOSpecularOcclusionMode == 1)
 		{
-			float ReflectionConeAngle = Roughness * PI;
-			float UnoccludedAngle = BentNormalLength * PI * InvSkySpecularOcclusionStrength;
-			float AngleBetween = acos(dot(BentNormalAO, ReflectionVector) / max(BentNormalLength, .001f));
-			IndirectSpecularOcclusion = ApproximateConeConeIntersection(ReflectionConeAngle, UnoccludedAngle, AngleBetween);
+			BRANCH
+			if (bTwoSidedFoliage)
+			{
+				IndirectSpecularOcclusion = BentNormalLength;
+			}
+			else
+			{
+				float ReflectionConeAngle = Roughness * PI;
+				float UnoccludedAngle = BentNormalLength * PI * InvSkySpecularOcclusionStrength;
+				float AngleBetween = acos(dot(BentNormalAO, ReflectionVector) / max(BentNormalLength, .001f));
+				IndirectSpecularOcclusion = ApproximateConeConeIntersection(ReflectionConeAngle, UnoccludedAngle, AngleBetween);
 
-			// Can't rely on the direction of the bent normal when close to fully occluded, lerp to shadowed
-			IndirectSpecularOcclusion = lerp(0, IndirectSpecularOcclusion, saturate((UnoccludedAngle - .1f) / .2f));
-			IndirectSpecularOcclusion = lerp(IndirectSpecularOcclusion, 1, MinSkySpecularOcclusion);
+				// Can't rely on the direction of the bent normal when close to fully occluded, lerp to shadowed
+				IndirectSpecularOcclusion = lerp(0, IndirectSpecularOcclusion, saturate((UnoccludedAngle - .1f) / .2f));
+			}
+		}
+		else if (Frame.DistanceFieldAOSpecularOcclusionMode == 2)
+		{
+			// Cone traced specular occlusion in w
+			IndirectSpecularOcclusion = TextureValue.w;
 		}
-#endif
 
+		IndirectSpecularOcclusion = lerp(IndirectSpecularOcclusion, 1, MinSkySpecularOcclusion);
 		IndirectDiffuseOcclusion = lerp(BentNormalLength, 1, MinSkySpecularOcclusion);
 	}
 }
diff --git a/Engine/Source/Developer/AssetTools/Private/AssetToolsPrivatePCH.h b/Engine/Source/Developer/AssetTools/Private/AssetToolsPrivatePCH.h
--- a/Engine/Source/Developer/AssetTools/Private/AssetToolsPrivatePCH.h
+++ b/Engine/Source/Developer/AssetTools/Private/AssetToolsPrivatePCH.h
@@ -53,6 +53,7 @@ DECLARE_LOG_CATEGORY_EXTERN(LogAssetTools, Log, All);
 #include "AssetTypeActions/AssetTypeActions_Font.h"
 #include "AssetTypeActions/AssetTypeActions_ForceFeedbackEffect.h"
 #include "AssetTypeActions/AssetTypeActions_SubsurfaceProfile.h"
+#include "AssetTypeActions/AssetTypeActions_SubDSurface.h"
 #include "AssetTypeActions/AssetTypeActions_InstancedFoliageSettings.h"
 #include "AssetTypeActions/AssetTypeActions_InterpData.h"
 #include "AssetTypeActions/AssetTypeActions_LandscapeLayer.h"
diff --git a/Engine/Source/Developer/AssetTools/Private/AssetTypeActions/AssetTypeActions_SubDSurface.cpp b/Engine/Source/Developer/AssetTools/Private/AssetTypeActions/AssetTypeActions_SubDSurface.cpp
new file mode 100644
--- /dev/null
+++ b/Engine/Source/Developer/AssetTools/Private/AssetTypeActions/AssetTypeActions_SubDSurface.cpp
@@ -0,0 +1,9 @@
+// Copyright 1998-2014 Epic Games, Inc. All Rights Reserved.
+
+#include "AssetToolsPrivatePCH.h"
+#include "Engine/SubDSurface.h"
+
+UClass* FAssetTypeActions_SubDSurface::GetSupportedClass() const
+{
+	return USubDSurface::StaticClass();
+}
diff --git a/Engine/Source/Developer/AssetTools/Private/AssetTypeActions/AssetTypeActions_SubDSurface.h b/Engine/Source/Developer/AssetTools/Private/AssetTypeActions/AssetTypeActions_SubDSurface.h
new file mode 100644
--- /dev/null
+++ b/Engine/Source/Developer/AssetTools/Private/AssetTypeActions/AssetTypeActions_SubDSurface.h
@@ -0,0 +1,14 @@
+// Copyright 1998-2014 Epic Games, Inc. All Rights Reserved.
+
+#pragma once
+
+class FAssetTypeActions_SubDSurface : public FAssetTypeActions_Base
+{
+public:
+	// IAssetTypeActions Implementation
+	virtual FText GetName() const override { return NSLOCTEXT("AssetTypeActions", "AssetTypeActions_SubDSurface", "Subdivision Surface"); }
+	virtual FColor GetTypeColor() const override { return FColor(128, 196, 128); }
+	virtual UClass* GetSupportedClass() const override;
+
+	virtual uint32 GetCategories() override { return EAssetTypeCategories::MaterialsAndTextures; }
+};
\ No newline at end of file
