New shader version to resolve issues integrated from Release-Staging-4.12 branch.

diff --git a/Engine/Shaders/ShaderVersion.usf b/Engine/Shaders/ShaderVersion.usf
--- a/Engine/Shaders/ShaderVersion.usf
+++ b/Engine/Shaders/ShaderVersion.usf
@@ -2,4 +2,4 @@
 // Each time the console command is executed it generates a new GUID. As this file is included
 // in common.usf (which should be included in any shader) it allows to invalidate the shader DDC.
 // 
-// GUID = A916320A8AE0495495AF2372472BC6EF
+// GUID = FFA66C634A2D7BE9C53AB8B326FCF4D0
\ No newline at end of file
