Removing the ResolvedView references as capsule shadows don't yet support instanced stereo.
Will reevaluate if/when we add isr support.

diff --git a/Engine/Shaders/CapsuleShadowShaders.usf b/Engine/Shaders/CapsuleShadowShaders.usf
--- a/Engine/Shaders/CapsuleShadowShaders.usf
+++ b/Engine/Shaders/CapsuleShadowShaders.usf
@@ -196,8 +196,8 @@ void SetupTileCullingData(
 		float2 TanViewFOV = GetTanHalfFieldOfView();
 		// tan(FOV) = HalfUnitPlaneWidth / 1, so TanViewFOV * 2 is the size of the whole unit view plane
 		// We are operating on a subset of that defined by ScissorRectMinAndSize
-		float2 TileSize = TanViewFOV * 2 * ScissorRectMinAndSize.zw / ((float2)ResolvedView.ViewSizeAndInvSize.xy * NumGroups);
-		float2 UnitPlaneMin = -TanViewFOV + TanViewFOV * 2 * (ScissorRectMinAndSize.xy - ResolvedView.ViewRectMin.xy) * ResolvedView.ViewSizeAndInvSize.zw;
+		float2 TileSize = TanViewFOV * 2 * ScissorRectMinAndSize.zw / ((float2)View.ViewSizeAndInvSize.xy * NumGroups);
+		float2 UnitPlaneMin = -TanViewFOV + TanViewFOV * 2 * (ScissorRectMinAndSize.xy - View.ViewRectMin.xy) * View.ViewSizeAndInvSize.zw;
 
 		float2 UnitPlaneTileMin = (GroupId.xy * TileSize + UnitPlaneMin) * float2(1, -1);
 		float2 UnitPlaneTileMax = ((GroupId.xy + 1) * TileSize + UnitPlaneMin) * float2(1, -1);
@@ -213,8 +213,8 @@ void SetupTileCullingData(
 	}
 	else
 	{
-		float2 TileSize = 2 * ScissorRectMinAndSize.zw / ((float2)ResolvedView.ViewSizeAndInvSize.xy * NumGroups);
-		float2 UnitPlaneMin = -1 + 2 * (ScissorRectMinAndSize.xy - ResolvedView.ViewRectMin.xy) * ResolvedView.ViewSizeAndInvSize.zw;
+		float2 TileSize = 2 * ScissorRectMinAndSize.zw / ((float2)View.ViewSizeAndInvSize.xy * NumGroups);
+		float2 UnitPlaneMin = -1 + 2 * (ScissorRectMinAndSize.xy - View.ViewRectMin.xy) * View.ViewSizeAndInvSize.zw;
 
 		float2 UnitPlaneTileMin = (GroupId.xy * TileSize + UnitPlaneMin) * float2(1, -1);
 		float2 UnitPlaneTileMax = ((GroupId.xy + 1) * TileSize + UnitPlaneMin) * float2(1, -1);
