Integrate - Fixed stationary skylight brightness

diff --git a/Engine/Shaders/ReflectionEnvironmentShared.usf b/Engine/Shaders/ReflectionEnvironmentShared.usf
--- a/Engine/Shaders/ReflectionEnvironmentShared.usf
+++ b/Engine/Shaders/ReflectionEnvironmentShared.usf
@@ -30,21 +30,14 @@ SamplerState SkyLightCubemapSampler;
 
 /** X = max mip, Y = 1 if sky light should be rendered, 0 otherwise, Z = 1 if sky light is dynamic, 0 otherwise, W = blend fraction. */
 float4 SkyLightParameters;
+float SkyLightCubemapBrightness;
 
-float3 GetSkyLightReflection(float3 ReflectionVector, float Roughness, bool bNormalized)
+float3 GetSkyLightReflection(float3 ReflectionVector, float Roughness, out float OutSkyAverageBrightness)
 {
 	float AbsoluteSpecularMip = ComputeReflectionCaptureMipFromRoughness(Roughness, SkyLightParameters.x);
 	float3 Reflection = TextureCubeSampleLevel(SkyLightCubemap, SkyLightCubemapSampler, ReflectionVector, AbsoluteSpecularMip).rgb;
 
-	BRANCH
-	if (bNormalized)
-	{
-		// Sample the lowest resolution mip to get the average color
-		//@todo - can't normalize sky lighting and reflection capture lighting separately
-		float3 LowFrequencyReflection = TextureCubeSampleLevel(SkyLightCubemap, SkyLightCubemapSampler, ReflectionVector, SkyLightParameters.x).rgb;
-		float LowFrequencyBrightness = Luminance(LowFrequencyReflection);
-		Reflection /= max(LowFrequencyBrightness, .00001f);
-	}
+	OutSkyAverageBrightness = SkyLightCubemapBrightness;
 
 	return Reflection * View.SkyLightColor.rgb;
 }
@@ -52,9 +45,9 @@ float3 GetSkyLightReflection(float3 ReflectionVector, float Roughness, bool bNor
 TextureCube SkyLightBlendDestinationCubemap;
 SamplerState SkyLightBlendDestinationCubemapSampler;
 
-float3 GetSkyLightReflectionSupportingBlend(float3 ReflectionVector, float Roughness, bool bNormalized)
+float3 GetSkyLightReflectionSupportingBlend(float3 ReflectionVector, float Roughness, out float OutSkyAverageBrightness)
 {
-	float3 Reflection = GetSkyLightReflection(ReflectionVector, Roughness, bNormalized);
+	float3 Reflection = GetSkyLightReflection(ReflectionVector, Roughness, OutSkyAverageBrightness);
 
 	BRANCH
 	if (SkyLightParameters.w > 0)
@@ -62,16 +55,6 @@ float3 GetSkyLightReflectionSupportingBlend(float3 ReflectionVector, float Rough
 		float AbsoluteSpecularMip = ComputeReflectionCaptureMipFromRoughness(Roughness, SkyLightParameters.x);
 		float3 BlendDestinationReflection = TextureCubeSampleLevel(SkyLightBlendDestinationCubemap, SkyLightBlendDestinationCubemapSampler, ReflectionVector, AbsoluteSpecularMip).rgb;
 
-		BRANCH
-		if (bNormalized)
-		{
-			// Sample the lowest resolution mip to get the average color
-			//@todo - can't normalize sky lighting and reflection capture lighting separately
-			float3 LowFrequencyReflection = TextureCubeSampleLevel(SkyLightBlendDestinationCubemap, SkyLightBlendDestinationCubemapSampler, ReflectionVector, SkyLightParameters.x).rgb;
-			float LowFrequencyBrightness = Luminance(LowFrequencyReflection);
-			BlendDestinationReflection /= max(LowFrequencyBrightness, .00001f);
-		}
-
 		Reflection = lerp(Reflection, BlendDestinationReflection * View.SkyLightColor.rgb, SkyLightParameters.w);
 	}
 
