UE-37981: Reverts 3165809's change in MaterialTemplate.usf

diff --git a/Engine/Shaders/MaterialTemplate.usf b/Engine/Shaders/MaterialTemplate.usf
--- a/Engine/Shaders/MaterialTemplate.usf
+++ b/Engine/Shaders/MaterialTemplate.usf
@@ -433,7 +433,7 @@ float4 GetScreenPosition(FMaterialPixelParameters Parameters)
 	{
 		return DecalToWorld[3].xyz;
 	}
-#elif HAS_PRIMITIVE_UNIFORM_BUFFER
+#else
 	float3 GetActorWorldPosition()
 	{
 		return Primitive.ActorWorldPosition;
