Fixed: Slate UI asset colors washed out on iOS
Replicated from [REDACTED_PHONE])

diff --git a/Engine/Shaders/SlateVertexShader.usf b/Engine/Shaders/SlateVertexShader.usf
--- a/Engine/Shaders/SlateVertexShader.usf
+++ b/Engine/Shaders/SlateVertexShader.usf
@@ -28,7 +28,7 @@ VertexToPixelInterpolants Main(
 	// save off the original input position value, Line AA will require this
 	VOut.OriginalPosition = WorldPosition;
 
-#if (!(ES2_PROFILE || ES3_1_PROFILE)) || USE_MATERIALS
+#if (!(ES2_PROFILE || ES3_1_PROFILE)) || USE_MATERIALS || METAL_PROFILE
 	InColor.rgb = sRGBToLinear(InColor.rgb);
 #endif
 
