Fixing vertex factory crash on Mac

diff --git a/Engine/Shaders/ParticleSpriteVertexFactory.usf b/Engine/Shaders/ParticleSpriteVertexFactory.usf
--- a/Engine/Shaders/ParticleSpriteVertexFactory.usf
+++ b/Engine/Shaders/ParticleSpriteVertexFactory.usf
@@ -18,8 +18,6 @@
 // We don't compute a tangent basis in ES2 to save shader instructions.
 #define NEEDS_TANGENT_BASIS (FEATURE_LEVEL >= FEATURE_LEVEL_SM4)
 
-bool UseCustomAlignment;
-
 struct FVertexFactoryInput
 {
 	float4	Position				: ATTRIBUTE0;
@@ -31,7 +29,6 @@ struct FVertexFactoryInput
 #endif	//USE_DYNAMIC_PARAMETERS
 	float2	TexCoord				: ATTRIBUTE4;
 	uint VertexId					: SV_VertexID;
-	float3 CustomAlignmentVector	: ATTRIBUTE6;
 
 /** Optional instance ID for vertex layered rendering */
 #if FEATURE_LEVEL >= FEATURE_LEVEL_SM4 && ONEPASS_POINTLIGHT_SHADOW && USING_VERTEX_SHADER_LAYER
@@ -319,7 +316,7 @@ float3 SafeNormalize(float3 V)
 	return V * rsqrt(max(dot(V,V),0.01));
 }
 
-void GetTangents(FVertexFactoryInput Input, float3 WorldPosition, float3 OldWorldPosition, out float3 OutRight, out float3 OutUp, float3 CustomAlignmentVector)
+void GetTangents(FVertexFactoryInput Input, float3 WorldPosition, float3 OldWorldPosition, out float3 OutRight, out float3 OutUp)
 {
 	// Select camera up/right vectors.
 	float3 ResolvedViewRight = lerp(ResolvedView.ViewRight, ResolvedView.HMDViewNoRollRight, SpriteVF.RemoveHMDRoll);
@@ -331,15 +328,10 @@ void GetTangents(FVertexFactoryInput Input, float3 WorldPosition, float3 OldWorl
 	// Determine the vector from the particle to the camera and the particle's movement direction.
 	float3 CameraDirection		= SafeNormalize(ResolvedView.WorldCameraOrigin - WorldPosition.xyz);
 	float3 ParticleDirection	= SafeNormalize(WorldPosition.xyz - OldWorldPosition.xyz);
-	float3 DirectionVector	= lerp(ParticleDirection, CustomAlignmentVector, UseCustomAlignment);
 
 	// Tangent vectors for PSA_Velocity.
-	// Velocity alignment
-	float3 Right_Velocity	= SafeNormalize(cross(CameraDirection, DirectionVector));
-	float3 Up_Velocity		= -DirectionVector;
-
-
-
+	float3 Right_Velocity	= SafeNormalize(cross(CameraDirection, ParticleDirection));
+	float3 Up_Velocity		= -ParticleDirection;
 
 	// Tangent vectors for rotation locked about an axis.
 	float3 Right_AxisRot = SpriteVF.AxisLockRight.xyz;
@@ -495,7 +487,7 @@ FVertexFactoryIntermediates GetVertexFactoryIntermediates(FVertexFactoryInput In
 
 	// Tangents.
 	float3 Right,Up;
-	GetTangents(Input,ParticleWorldPosition.xyz,ParticleOldWorldPosition.xyz,Right,Up, Input.CustomAlignmentVector);
+	GetTangents(Input,ParticleWorldPosition.xyz,ParticleOldWorldPosition.xyz,Right,Up);
 	Intermediates.TangentUp = Up;
 	Intermediates.TangentRight = Right;
 
@@ -548,9 +540,8 @@ FVertexFactoryIntermediates GetVertexFactoryIntermediates(FVertexFactoryInput In
 
 #if USE_PARTICLE_VELOCITY
 	float3 ParticleVelocity = (ParticleWorldPosition.xyz - ParticleOldWorldPosition.xyz) * SpriteVF.InvDeltaSeconds;
-	float3 DirectionVector = lerp(ParticleVelocity, Input.CustomAlignmentVector, UseCustomAlignment);
-	Intermediates.ParticleVelocity.xyz = normalize(DirectionVector);
-	Intermediates.ParticleVelocity.w = length(DirectionVector);
+	Intermediates.ParticleVelocity.xyz = normalize(ParticleVelocity);
+	Intermediates.ParticleVelocity.w = length(ParticleVelocity);
 #endif
 
 #if USE_PARTICLE_TIME
diff --git a/Engine/Source/Runtime/Engine/Private/Particles/NiagaraSpriteVertexFactory.cpp b/Engine/Source/Runtime/Engine/Private/Particles/NiagaraSpriteVertexFactory.cpp
--- a/Engine/Source/Runtime/Engine/Private/Particles/NiagaraSpriteVertexFactory.cpp
+++ b/Engine/Source/Runtime/Engine/Private/Particles/NiagaraSpriteVertexFactory.cpp
@@ -136,7 +136,7 @@ public:
 	virtual void FillDeclElements(FVertexDeclarationElementList& Elements, int32& Offset)
 	{
 		uint32 InitialStride = sizeof(float) * 2;
-		uint32 PerParticleStride = (sizeof(float) * 4) * NumVertsInInstanceBuffer;
+		uint32 PerParticleStride = (sizeof(float) * 4*4 + sizeof(float)*3*2);
 
 		/** The stream to read the texture coordinates from. */
 		check( Offset == 0 );
@@ -164,6 +164,7 @@ public:
 		/** The stream to read the color from.					*/
 		Elements.Add(FVertexElement(bInstanced ? 1 : 0, Offset, VET_Float4, 3, Stride, bInstanced));
 		Offset += sizeof(float) * 4;
+
 		/** The stream to read the custom alignment vector from. */
 		Elements.Add(FVertexElement(bInstanced ? 1 : 0, Offset, VET_Float3, 6, Stride, bInstanced));
 		Offset += sizeof(float) * 3;
@@ -180,6 +181,7 @@ public:
 		Stride = sizeof(float) * 4;
 		Elements.Add(FVertexElement(bInstanced ? 2 : 1, Offset, VET_Float4, 5, Stride, bInstanced));
 		Offset += sizeof(float) * 4;
+
 	}
 
 	virtual void InitDynamicRHI()
diff --git a/Engine/Source/Runtime/Engine/Private/Particles/ParticleVertexFactory.cpp b/Engine/Source/Runtime/Engine/Private/Particles/ParticleVertexFactory.cpp
--- a/Engine/Source/Runtime/Engine/Private/Particles/ParticleVertexFactory.cpp
+++ b/Engine/Source/Runtime/Engine/Private/Particles/ParticleVertexFactory.cpp
@@ -65,14 +65,12 @@ public:
 	{
 		NumCutoutVerticesPerFrame.Bind(ParameterMap, TEXT("NumCutoutVerticesPerFrame"));
 		CutoutGeometry.Bind(ParameterMap, TEXT("CutoutGeometry"));
-		UseCustomAlignmentVector.Bind(ParameterMap, TEXT("UseCustomAlignment"));
 	}
 
 	virtual void Serialize(FArchive& Ar) override
 	{
 		Ar << NumCutoutVerticesPerFrame;
 		Ar << CutoutGeometry;
-		Ar << UseCustomAlignmentVector;
 	}
 
 	virtual void SetMesh(FRHICommandList& RHICmdList, FShader* Shader,const FVertexFactory* VertexFactory,const FSceneView& View,const FMeshBatchElement& BatchElement,uint32 DataFlags) const override
@@ -84,12 +82,10 @@ public:
 		SetShaderValue(RHICmdList, VertexShaderRHI, NumCutoutVerticesPerFrame, SpriteVF->GetNumCutoutVerticesPerFrame());
 		FShaderResourceViewRHIParamRef NullSRV = GFNullSubUVCutoutVertexBuffer.VertexBufferSRV;
 		SetSRVParameter(RHICmdList, VertexShaderRHI, CutoutGeometry, SpriteVF->GetCutoutGeometrySRV() ? SpriteVF->GetCutoutGeometrySRV() : NullSRV);
-		SetShaderValue(RHICmdList, VertexShaderRHI, UseCustomAlignmentVector, SpriteVF->GetCustomAlignment());
 	}
 
 private:
 	FShaderParameter NumCutoutVerticesPerFrame;
-	FShaderParameter UseCustomAlignmentVector;
 	FShaderResourceParameter CutoutGeometry;
 };
 
@@ -156,10 +152,7 @@ public:
 		/** The stream to read the color from.					*/
 		Elements.Add(FVertexElement(bInstanced ? 1 : 0, Offset, VET_Float4, 3, Stride, bInstanced));
 		Offset += sizeof(float) * 4;
-		/** The stream to read the custom alignment vector from. */
-		Elements.Add(FVertexElement(bInstanced ? 1 : 0, Offset, VET_Float3, 6, Stride, bInstanced));
-		Offset += sizeof(float) * 3;
-
+		
 		/** The per-particle dynamic parameter stream */
 
 		// The -V519 disables a warning from PVS-Studio's static analyzer. It noticed that offset is assigned
diff --git a/Engine/Source/Runtime/Engine/Public/ParticleHelper.h b/Engine/Source/Runtime/Engine/Public/ParticleHelper.h
--- a/Engine/Source/Runtime/Engine/Public/ParticleHelper.h
+++ b/Engine/Source/Runtime/Engine/Public/ParticleHelper.h
@@ -369,8 +369,6 @@ struct FParticleSpriteVertexNonInstanced
 	float SubImageIndex;
 	/** The color of the particle. */
 	FLinearColor Color;
-	/* Custom Alignment vector*/
-	FVector CustomAlignmentVector;
 };
 
 
