Fix shader compile error when creating volume domain materials

diff --git a/Engine/Shaders/NiagaraSpriteVertexFactory.usf b/Engine/Shaders/NiagaraSpriteVertexFactory.usf
--- a/Engine/Shaders/NiagaraSpriteVertexFactory.usf
+++ b/Engine/Shaders/NiagaraSpriteVertexFactory.usf
@@ -653,3 +653,8 @@ uint VertexFactoryGetEyeIndex(uint InstanceId)
 	return InstancedEyeIndex & bIsInstancedStereo;
 }
 #endif
+
+float4 VertexFactoryGetTranslatedPrimitiveVolumeBounds(FVertexFactoryInterpolantsVSToPS Interpolants)
+{
+	return 0;
+}
\ No newline at end of file
