Copying //UE4/Dev-Rendering to //UE4/Dev-Main (Source: //UE4/Dev-Rendering @ 3461187)

diff --git a/Engine/Shaders/ShadowPercentageCloserFiltering.usf b/Engine/Shaders/ShadowPercentageCloserFiltering.usf
new file mode 100644
--- /dev/null
+++ b/Engine/Shaders/ShadowPercentageCloserFiltering.usf
@@ -0,0 +1,500 @@
+// Copyright 1998-2017 Epic Games, Inc. All Rights Reserved.
+
+/*=============================================================================
+	ShadowPercentageCloserFiltering.usf: Contains functions to do percentage closer soft shadow sampling.
+=========================================================================*/
+
+#include "ShadowFilteringCommon.usf"
+
+
+// CONFIGURATION OF PCSS
+// -----------------------------------------------------------------------
+
+// Whether to compile some debugging utilitaries.
+#define PCSS_DEBUG_UTILITITARIES 0
+
+// Whether to debug pixel where they early return.
+#define PCSS_DEBUG_EARLY_RETURN 0
+
+// Wheather some computation should be shared across 2x2 pixel quad.
+//  0: disabled.
+//  1: share occluder search result.
+//  2: share occluder search and PCF results.
+#define PCSS_SHARE_PER_PIXEL_QUAD 2
+
+// Wheather to set a maximum depth bias.
+#define PCSS_MAX_DEPTH_BIAS 1
+
+// Idea of the experiment to turn on.
+#define PCSS_ANTI_ALIASING_METHOD 2
+
+// PCF experiment to solve translucent shadow artifacts.
+//  0: Dummy PCF.
+//  1: Cone traced occluder sums with occluder distance based weighting.
+#define PCSS_PCF_EXPERIMENT 0
+
+// Wheather to enable the sharpening filter after PCF for sharper edges than the shadow map resolution.
+#define PCSS_ENABLE_POST_PCF_SHARPENING 1
+
+
+// POISSON KERNEL
+// -----------------------------------------------------------------------
+
+struct FPCSSSamplerSettings
+{
+	Texture2D ShadowDepthTexture;
+	SamplerState ShadowDepthTextureSampler;
+
+	//XY - Pixel size of shadowmap
+	//ZW - Inverse pixel size of shadowmap
+	float4 ShadowBufferSize;
+
+	// Offset and size of the shadow tile within the shadow depth texture.
+	float4 ShadowTileOffsetAndSize;
+
+	// SceneDepth in shadow view space.
+	float SceneDepth;
+
+	// Transition scale (TODO: Could probably nuke that guy).
+	float TransitionScale;
+	
+	// Tan(0.5 * Light Source Angle) in the shadow view space.
+	float TanLightSourceAngle;
+	
+	// Maximum kernel size in the tile's UV space.
+	float MaxKernelSize;
+	
+	// Pixel's postion for random number generation.
+	float2 SvPosition;
+
+	// Thread's pixel quad message passing context.
+	FPQMPContext PQMPContext;
+
+	// Viewport UV for debuging purposes.
+	float2 DebugViewportUV;
+};
+
+
+// POISSON KERNEL
+// -----------------------------------------------------------------------
+
+#define PCSS_SEARCH_SAMPLES 16
+#define PCSS_SAMPLES 32
+
+static const float2 PCSSSamples[PCSS_SEARCH_SAMPLES] = {
+	float2(0.03522659013938079, -0.06773381793326105),
+	float2(-0.2226592407149033, 0.28712634420696953),
+	float2(0.29521592554288656, 0.3091113612849006),
+	float2(0.45445972849126726, -0.06885997394231347),
+	float2(-0.5294175874764113, -0.046014035861565994),
+	float2(-0.452648396511972, -0.47909964168393443),
+	float2(-0.24846147512522798, 0.7335295870878695),
+	float2(0.3577384490251264, 0.7225563640511622),
+	float2(0.6637172290495703, -0.5130051538790028),
+	float2(0.8247905306055152, 0.2699748069491893),
+	float2(0.2677328749047735, -0.8264198346042448),
+	float2(0.8735142147828027, -0.13873095107832628),
+	float2(-0.9200356771077463, 0.12928348617112767),
+	float2(-0.333591990232354, -0.88825657027033),
+	float2(-0.6744079723354472, 0.7098786719042596),
+	float2(-0.916141973131132, -0.35310318938523944),
+};
+
+static const float2 PCFSamples[PCSS_SAMPLES] = {
+	float2(0.05184623507064756, 0.1300833011807859),
+	float2(-0.1577127330375604, -0.04411805230289356),
+	float2(-0.3608032569836849, 0.14224022015847443),
+	float2(0.32943189033164505, -0.23116230997768047),
+	float2(-0.3248396383424306, -0.24779640789279145),
+	float2(-0.08378000616865044, -0.42320666155308073),
+	float2(0.27231365926326034, 0.4570699627554666),
+	float2(0.4665144073390388, 0.2723623292705937),
+	float2(-0.3585638845043682, 0.4184774097652295),
+	float2(-0.05651562661921321, 0.6247868960227553),
+	float2(-0.33439047389336457, -0.5347681672345762),
+	float2(0.23213281841559313, -0.5970219039183482),
+	float2(-0.6407636828693648, 0.11993935952614704),
+	float2(0.6532651709764936, 0.050993271277866264),
+	float2(-0.6024883054476216, -0.27760818629456524),
+	float2(0.48856009912681686, -0.47214983370201674),
+	float2(0.22574621537299483, 0.7150421652622163),
+	float2(0.5333102819109856, 0.5347971792924224),
+	float2(-0.8677248476742625, -0.039097016569094784),
+	float2(0.7864067837429691, 0.4116218032124058),
+	float2(0.0148798999072412, 0.8888627595628051),
+	float2(-0.51288316652954, -0.733326766790393),
+	float2(-0.6630269006484741, 0.6126292700542129),
+	float2(-0.845774354018934, 0.32435673562428946),
+	float2(0.8332104797006348, -0.3622295582682804),
+	float2(-0.2446550784225213, -0.8800972615287274),
+	float2(-0.32118393732444706, 0.8618550333128141),
+	float2(0.2393990624981109, -0.8931529275822183),
+	float2(-0.7628371490800573, -0.5306507752416643),
+	float2(0.47612860488155273, 0.8066697476711497),
+	float2(0.6204553222569549, -0.7122745692650101),
+	float2(0.9557311286580455, 0.14112604822367059),
+};
+
+
+// PCSS SPECIFIC UTILITARY FUNCTIONS
+// -----------------------------------------------------------------------
+
+float2x2 GenerateScale2x2Matrix(float2 MajorAxis, float MajorMinusMinorScale, float MinorScale)
+{
+	return float2x2(
+		MinorScale + MajorMinusMinorScale * MajorAxis.x * MajorAxis.x, MajorMinusMinorScale * MajorAxis.y * MajorAxis.x,
+		MajorMinusMinorScale * MajorAxis.x * MajorAxis.y, MinorScale + MajorMinusMinorScale * MajorAxis.y * MajorAxis.y);
+}
+
+float2x2 GenerateDirectionalScale2x2Matrix(float2 Direction, float ScaleMinusOne)
+{
+	return float2x2(
+		1 + ScaleMinusOne * Direction.x * Direction.x, ScaleMinusOne * Direction.y * Direction.x,
+		ScaleMinusOne * Direction.x * Direction.y, 1 + ScaleMinusOne * Direction.y * Direction.y);
+}
+
+
+#if PCSS_DEBUG_UTILITITARIES
+
+// Function to debug the direction of ellipse major and minor axes in the shadow map UV space.
+float PCSSDebugUVDir(float2 v)
+{
+	return frac(atan2(v.x, v.y) / (PI)+0.25);
+}
+
+#endif //PCSS_DEBUG_UTILITITARIES
+
+
+// PCSS FILTER
+// -----------------------------------------------------------------------
+
+float DirectionalPCSS(FPCSSSamplerSettings Settings, float2 ShadowPosition, float3 ShadowPositionDDX, float3 ShadowPositionDDY)
+{
+	float3 DepthBiasPlaneNormal = cross(ShadowPositionDDX, ShadowPositionDDY);
+#if PCSS_MAX_DEPTH_BIAS
+	float DepthBiasFactor = 1 / max(abs(DepthBiasPlaneNormal.z), length(DepthBiasPlaneNormal) * 0.0872665);
+#else
+	float DepthBiasFactor = 1 / abs(DepthBiasPlaneNormal.z);
+#endif
+	float2 DepthBiasDotFactors = DepthBiasPlaneNormal.xy * DepthBiasFactor;
+
+	float2 ShadowTexelSizeX = Settings.ShadowBufferSize.zw / Settings.ShadowTileOffsetAndSize.zw; //TODO: remove truncation.
+	const float MinFilterSize = max(ShadowTexelSizeX.x, ShadowTexelSizeX.y);
+	float PCSSMinFilterSize = MinFilterSize;
+	float PCFMinFilterSize = MinFilterSize;
+
+	// dot(surface's normal, light direction).
+	float3 NormalizedDepthBiasPlaneNormal = normalize(DepthBiasPlaneNormal);
+	float DotNormal = NormalizedDepthBiasPlaneNormal.z;
+
+	float RandomRotation = Rand3DPCG16(uint3(Settings.SvPosition, View.StateFrameIndexMod8)).x * (2.f * PI / 0x10000);
+	float RandomFilterScale = Rand3DPCG16(uint3(Settings.SvPosition, View.StateFrameIndexMod8 + 8)).x * (0.5f / 0x10000) + 0.5f;
+	float2x2 PerPixelRotationMatrix = float2x2(cos(RandomRotation), sin(RandomRotation), -sin(RandomRotation), cos(RandomRotation));
+	float2x2 OccluderSearchUVMatrix = PerPixelRotationMatrix;
+
+#if PCSS_ANTI_ALIASING_METHOD == 1 || PCSS_ANTI_ALIASING_METHOD == 2 // Dilate occluder search.
+	PCSSMinFilterSize *= 3;
+#endif
+
+	float SearchRadius = max(PCSSMinFilterSize, RandomFilterScale * min(Settings.SceneDepth * Settings.TanLightSourceAngle, Settings.MaxKernelSize));
+
+#if PCSS_ANTI_ALIASING_METHOD == 6 // Experience 6: anysotropy.
+	float2 ShadowPositionDx = ddx(ShadowPosition);
+	float2 ShadowPositionDy = ddy(ShadowPosition);
+	float2 ShadowAnisotropicVector = ShadowPositionDx + ShadowPositionDy;
+	float2 ShadowPositionDelta = sqrt(float2(
+		ShadowPositionDx.x * ShadowPositionDx.x + ShadowPositionDy.x * ShadowPositionDy.x,
+		ShadowPositionDx.y * ShadowPositionDx.y + ShadowPositionDy.y * ShadowPositionDy.y));
+
+	float MajorAnisotropy = max(ShadowPositionDelta.x, ShadowPositionDelta.y);
+	float MinorAnisotropy = min(ShadowPositionDelta.x, ShadowPositionDelta.y);
+
+	float MinorAnisotropyFactor = clamp(MinorAnisotropy / SearchRadius, 1, 7);
+	float MajorAnisotropyFactor = clamp(MajorAnisotropy / SearchRadius, 1, 7);
+	if (View.GeneralPurposeTweak < 1.5) {
+		MajorAnisotropyFactor = 1;
+		MinorAnisotropyFactor = 1;
+	}
+	if (View.GeneralPurposeTweak < 2.5) {
+		MajorAnisotropyFactor = clamp(length(ShadowAnisotropicVector) / SearchRadius, 1, 7);
+		MinorAnisotropyFactor = 1;
+	}
+	//return MajorAnisotropyFactor - MinorAnisotropyFactor;
+	//return (MinorAnisotropyFactor - 1);
+
+	float2 ShadowAnisotropicDir = normalize(ShadowAnisotropicVector);
+	float2x2 AnysotropticProjectionMatrix = GenerateScale2x2Matrix(
+		ShadowAnisotropicDir, MajorAnisotropyFactor - MinorAnisotropyFactor, MinorAnisotropyFactor);
+
+	PerPixelRotationMatrix = mul(AnysotropticProjectionMatrix, PerPixelRotationMatrix);
+	OccluderSearchUVMatrix = PerPixelRotationMatrix;
+#endif
+
+	// Search for an average occluder depth.
+	float DepthSum = 0;
+	float DepthSquaredSum = 0;
+	float SumWeight = 0;
+	float DiscardedSampleCount = 0;
+	float2 UVOffsetSum = float2(0, 0);
+	float2 UVOffsetSquareSum = float2(0, 0);
+	float UVCrossSum = 0;
+	for (int i = 0; i < PCSS_SEARCH_SAMPLES; i++)
+	{
+		float2 SampleUVOffset = mul(OccluderSearchUVMatrix, PCSSSamples[i]) * SearchRadius;
+		float2 SampleUV = ShadowPosition + SampleUVOffset * Settings.ShadowTileOffsetAndSize.zw;
+
+		float ShadowDepth = Texture2DSampleLevel(Settings.ShadowDepthTexture, Settings.ShadowDepthTextureSampler, SampleUV, 0).r;
+		float ShadowDepthCompare = Settings.SceneDepth - ShadowDepth;
+
+		float SampleDepthBias = max(dot(DepthBiasDotFactors, SampleUVOffset), 0);
+
+		// TODO: Cone trace to early return more often on SumWeight == 0?
+		if (ShadowDepthCompare > SampleDepthBias)
+		{
+			DepthSum += ShadowDepth;
+			DepthSquaredSum += ShadowDepth * ShadowDepth;
+			SumWeight += 1;
+			float2 UV = mul(OccluderSearchUVMatrix, PCSSSamples[i]);
+			UVOffsetSum += UV;
+			UVOffsetSquareSum += UV * UV;
+			UVCrossSum += UV.x * UV.y;
+		}
+		else if (ShadowDepthCompare > 0)
+		{
+			DiscardedSampleCount += 1;
+		}
+	}
+
+#if PCSS_SHARE_PER_PIXEL_QUAD // Pixel quad share setup.
+	const float MaxTexelShare = 1;
+	float2 MaxDerivative = max(abs(ddx(ShadowPosition)), abs(ddy(ShadowPosition)));
+	float DoPerQuad = 0.8 * max(1 - max(MaxDerivative.x, MaxDerivative.y) / (MinFilterSize * MaxTexelShare), 0);
+#else
+	float DoPerQuad = 0;
+#endif 
+
+#if PCSS_SHARE_PER_PIXEL_QUAD >= 1 // Share occluder search computation per quad.
+	DepthSum = lerp(DepthSum, 4 * PQMPAverage(Settings.PQMPContext, DepthSum), DoPerQuad);
+	DepthSquaredSum = lerp(DepthSquaredSum, 4 * PQMPAverage(Settings.PQMPContext, DepthSquaredSum), DoPerQuad);
+	SumWeight = lerp(SumWeight, 4 * PQMPAverage(Settings.PQMPContext, SumWeight), DoPerQuad);
+	UVOffsetSum = lerp(UVOffsetSum, 4 * PQMPAverage(Settings.PQMPContext, UVOffsetSum), DoPerQuad);
+	UVOffsetSquareSum = lerp(UVOffsetSquareSum, 4 * PQMPAverage(Settings.PQMPContext, UVOffsetSquareSum), DoPerQuad);
+	UVCrossSum = lerp(UVCrossSum, 4 * PQMPAverage(Settings.PQMPContext, UVCrossSum), DoPerQuad);
+#endif
+
+	if (SumWeight == 0)
+	{
+		// If no occluder found, early return as if this pixel was unshadowed.
+		return PCSS_DEBUG_EARLY_RETURN ? 0.8 : 1;
+	}
+	else if (SumWeight >= lerp(PCSS_SEARCH_SAMPLES, 4 * PCSS_SEARCH_SAMPLES - 0.5, DoPerQuad))
+	{
+		// If every occluder search sample have found an occluder, then early return as if this pixel was totaly shadowed.
+		return PCSS_DEBUG_EARLY_RETURN ? 0.2 : 0;
+	}
+
+	float NormalizeFactor = (1 / SumWeight);
+	float2 OccluderUVGradient = UVOffsetSum * NormalizeFactor;
+	float2x2 PCFUVMatrix = PerPixelRotationMatrix;
+
+#if 0 // Signal analysis:
+	// This shows the direction of the UV offset average of the occluders.
+	if (Settings.DebugScreenUV.x > 0.7 && 1) return PCSSDebugUVDir(OccluderUVGradient);
+
+	// This shows the length of the UV offset average of the occluders.
+	if (Settings.DebugScreenUV.x > 0.5 && 0) return length(OccluderUVGradient);
+
+	// This shows the number of occluders.
+	if (Settings.DebugScreenUV.x > 0.5 && 0) return SumWeight / 3;
+#endif
+
+	float DepthAvg = DepthSum * NormalizeFactor;
+	float DepthVariance = NormalizeFactor * DepthSquaredSum - DepthAvg * DepthAvg;
+	float DepthStandardDeviation = sqrt(DepthVariance);
+	float AverageOccluderDistance = (Settings.SceneDepth - DepthAvg);
+	float Penumbra = min((Settings.SceneDepth - DepthAvg) * Settings.TanLightSourceAngle, Settings.MaxKernelSize);
+
+	float RawFilterRadius = RandomFilterScale * Penumbra;
+	float FilterRadius = max(PCFMinFilterSize, RawFilterRadius);
+	float TanDirectionalLightAngle = FilterRadius / (RandomFilterScale * (Settings.SceneDepth - DepthAvg));
+
+#if PCSS_ANTI_ALIASING_METHOD == 1 // Experiement 4: do elliptical PCF kernel according to the occluder UV gradient.
+	{
+		float ElepticalFading = PCFMinFilterSize / FilterRadius;
+		float ElepticalFactor = length(OccluderUVGradient) * SumWeight * ElepticalFading;
+		float2 GrandAxisDir = normalize(float2(OccluderUVGradient.y, -OccluderUVGradient.x));
+		float2x2 ElepticalProjectionMatrix = GenerateDirectionalScale2x2Matrix(GrandAxisDir, ElepticalFactor);
+
+		PCFUVMatrix = mul(ElepticalProjectionMatrix, PCFUVMatrix);
+	}
+#elif PCSS_ANTI_ALIASING_METHOD == 2
+	{
+		// Minimal number of times the longest egeinst vector should be from the shorter one.
+		const float EigenThreshold = 1.3;
+
+		// Number of occluders requiered to consider the covriance matrix valid.
+		const float MinimumCovaranceOccluders = 2;
+
+		// Compute covariance matrix:
+		//		a b
+		//		b c
+		//
+		// With:
+		//	a == CovarianceMatrixDiagonal.x
+		//	b == CovarianceMatrixNonDiagonal
+		//	c == CovarianceMatrixDiagonal.y
+		float CovarianceMatrixNonDiagonal = SumWeight * UVCrossSum - OccluderUVGradient.x * OccluderUVGradient.y;
+		float2 CovarianceMatrixDiagonal = float2(
+			SumWeight * UVOffsetSquareSum.x - OccluderUVGradient.x * OccluderUVGradient.x,
+			SumWeight * UVOffsetSquareSum.y - OccluderUVGradient.y * OccluderUVGradient.y);
+
+		// Compute covariance matrix's Eigen values.
+		float CovarianceMatrixTrace = CovarianceMatrixDiagonal.x + CovarianceMatrixDiagonal.y;
+		float T = sqrt(CovarianceMatrixTrace * CovarianceMatrixTrace -
+			4 * (CovarianceMatrixDiagonal.x * CovarianceMatrixDiagonal.y - CovarianceMatrixNonDiagonal * CovarianceMatrixNonDiagonal));
+		float2 EigenValues = 0.5 * (CovarianceMatrixTrace + float2(T, -T));
+
+		// Choose the longest Egein vector.
+		float MaxEigenValue = max(EigenValues.x, EigenValues.y);
+		float MinEigenValue = min(EigenValues.x, EigenValues.y);
+		//float2 EigenVector = normalize(float2(CovarianceMatrixNonDiagonal, MaxEigenValue - CovarianceMatrixDiagonal.x));
+		float2 LongestEigenVector = normalize(float2(MaxEigenValue - CovarianceMatrixDiagonal.y, CovarianceMatrixNonDiagonal));
+		float LongestEigenVectorFactor = sqrt(MaxEigenValue / MinEigenValue) - EigenThreshold;
+		float CovarianceBasedEllipticalFactor = 8 * LongestEigenVectorFactor;
+		
+		// Linear interpolator between the average occluder UV offset method (0) to the covariance method (1).
+		float CovarianceFade = saturate((SumWeight >= MinimumCovaranceOccluders ? 1 : 0) * LongestEigenVectorFactor);
+		
+		// Heuristic: Force to use the average occluder UV offset method if the length of average of the offset is high.
+		CovarianceFade *= saturate(1 - 5 * (length(OccluderUVGradient) - 0.5));
+
+		// Ellipse's major axis computed from the average occluder UV offset. 
+		float2 AverageBasedMajorAxis = normalize(float2(OccluderUVGradient.y, -OccluderUVGradient.x));
+		float AverageBasedEllipticalFactor = length(OccluderUVGradient) * SumWeight;
+
+		// Ellipitical fade based on how many time the filter size.
+		float EllipticalFade = PCFMinFilterSize / FilterRadius;
+
+		// Final ellipse's major axis and scale.
+		float2 EllipseMajorAxis = normalize(lerp(AverageBasedMajorAxis, LongestEigenVector, CovarianceFade));
+		float MaxEllipticalFactor = 6 - 3 * CovarianceFade;
+		float EllipticalFactor = min(lerp(AverageBasedEllipticalFactor, CovarianceBasedEllipticalFactor, CovarianceFade), MaxEllipticalFactor) * EllipticalFade;
+
+		// Heuristic: reduce the ellipse shapse in corners of light.
+		//float ShadowCornerMultiplier = 1 - pow(saturate(30 * (length(OccluderUVGradient) - 0.45)) * abs(dot(AverageBasedMajorAxis, LongestEigenVector)), 3);
+		//EllipticalFactor *= ShadowCornerMultiplier;
+
+		// Projection matrix that transform the disk shaped kernel to an ellipse.
+		float2x2 ElepticalProjectionMatrix = GenerateDirectionalScale2x2Matrix(EllipseMajorAxis, EllipticalFactor);
+
+		// Hacks the PCF's per pixel UV rotation matrix.
+		PCFUVMatrix = mul(ElepticalProjectionMatrix, PCFUVMatrix);
+
+#if 0 // Final major axis analysis.
+		if (Settings.DebugScreenUV.x > 0.5) return PCSSDebugUVDir(EllipseMajorAxis);
+
+#elif 0 // Final elliptical factor
+		if (Settings.DebugScreenUV.x > 0.5) return EllipticalFactor * 0.1;
+
+#elif 0 // Covarance fade factor
+		if (Settings.DebugScreenUV.x > 0.5) return CovarianceFade;
+
+#elif 0 // Not enough occluder samples.
+		if (Settings.DebugScreenUV.x > 0.5) return SumWeight < 3 ? 1 : 0;
+
+#elif 0
+		if (Settings.DebugScreenUV.x > 0.5) return ShadowCornerMultiplier;
+
+#elif 0 // Egein vector 
+		if (Settings.DebugScreenUV.x > 0.5) return sqrt(MaxEigenValue / MinEigenValue) - 1;
+
+#elif 0
+		if (Settings.DebugScreenUV.x > 0.5) return (length(OccluderUVGradient) - 0.7) * 5;
+
+#elif 0 // Signal analysis.
+		if (Settings.DebugScreenUV.x > 0.7 && 1) return EllipticalFade;
+		if (Settings.DebugScreenUV.x > 0.5 && 1) return EllipticalFactor;
+		if (Settings.DebugScreenUV.x > 0.3 && 1) return CovarianceFade;
+
+#elif 0 // Major axis transition analysis.
+		if (Settings.DebugScreenUV.x > 0.7) return PCSSDebugUVDir(AverageBasedMajorAxis);
+		if (Settings.DebugScreenUV.x > 0.5) return PCSSDebugUVDir(LongestEigenVector);
+		if (Settings.DebugScreenUV.x > 0.3) return CovarianceFade;
+#endif
+	}
+#endif
+
+#if PCSS_PCF_EXPERIMENT == 1 // -> Shadow occlusion sum.
+	// UV area covered by one sample one the disk on the average depth. 
+	const float NormalizedSampleArea = 1 / float(PCSS_SAMPLES);
+	
+	// PCF loop.
+	float ShadowOcclusionSum = 0;
+	for (int j = 0; j < PCSS_SAMPLES; j++)
+	{
+		float OriginalUVLength = length(PCFSamples[j]);
+		float2 SampleUVOffset = mul(PCFUVMatrix, PCFSamples[j]) * FilterRadius;// *pow(OriginalUVLength, 0.9);
+		float2 SampleUV = ShadowPosition + SampleUVOffset * Settings.ShadowTileOffsetAndSize.zw;
+		float SampleDepth = Texture2DSampleLevel(Settings.ShadowDepthTexture, Settings.ShadowDepthTextureSampler, SampleUV, 0).r;
+		float ShadowDepthCompare = Settings.SceneDepth - SampleDepth;
+
+		float SampleDepthBias = length(SampleUVOffset) / TanLightSourceAngle;
+		//float SampleDepthBias = max(dot(DepthBiasDotFactors, SampleUVOffset), 0);
+
+		if (ShadowDepthCompare > SampleDepthBias)
+		{
+			float SampleSizeMultiplier = AverageOccluderDistance / min(ShadowDepthCompare, AverageOccluderDistance);
+			ShadowOcclusionSum += NormalizedSampleArea * SampleSizeMultiplier * SampleSizeMultiplier;
+		}
+	}
+	ShadowOcclusionSum = lerp(ShadowOcclusionSum, PQMPAverage(Settings.PQMPContext, ShadowOcclusionSum), DoPerQuad * 0.5);
+	float Visibility = max(1.0 - ShadowOcclusionSum, 0.0);
+
+#else
+	// PCF loop.
+	float VisibleLightAccumulation = 0;
+	for (int j = 0; j < PCSS_SAMPLES; j++)
+	{
+		float2 SampleUVOffset = mul(PCFUVMatrix, PCFSamples[j]) * FilterRadius;
+		float2 SampleUV = ShadowPosition + SampleUVOffset * Settings.ShadowTileOffsetAndSize.zw;
+		float SampleDepth = Texture2DSampleLevel(Settings.ShadowDepthTexture, Settings.ShadowDepthTextureSampler, SampleUV, 0).r;
+
+		float SampleDepthBias = max(dot(DepthBiasDotFactors, SampleUVOffset), 0);
+
+		VisibleLightAccumulation += saturate((SampleDepth - Settings.SceneDepth + SampleDepthBias) * Settings.TransitionScale + 1);
+	}
+
+	float Visibility = VisibleLightAccumulation / float(PCSS_SAMPLES);
+
+#endif // PCSS_PCF_EXPERIMENT == 0
+
+#if PCSS_ENABLE_POST_PCF_SHARPENING // Sharpen the PCF result.
+	{
+		// Average distance between samles in UV space.
+		const float AverageSampleDistance = sqrt(1 / (float(PCSS_SAMPLES) * PI));
+
+		// Maximum factor of sharpness to introduce.
+		const float MaxSharpnessFactor = 4;
+
+		// Filter to avoid applying the sharpening filter on high frequency shadows.
+		float SharpenessFading = saturate(1.5 * (length(OccluderUVGradient) - AverageSampleDistance));
+
+		// Factor introducing the sharpness.
+		float FinalSharpenessFactor = lerp(1, clamp(PCFMinFilterSize / RawFilterRadius, 1, MaxSharpnessFactor), SharpenessFading);
+
+		// Apply the sharpness onto the visibility.
+		Visibility = saturate(FinalSharpenessFactor * (Visibility - 0.5) + 0.5);
+
+#if 0
+		if (Settings.DebugScreenUV.x > 0.8) return SharpeningFade;
+		if (Settings.DebugScreenUV.x > 0.5) return (SharpeningFactor - 1) * 0.25;
+#endif
+	}
+#endif //PCSS_ENABLE_POST_PCF_SHARPENING
+
+#if PCSS_SHARE_PER_PIXEL_QUAD >= 2 // Share PCF computation.
+	Visibility = lerp(Visibility, PQMPAverage(Settings.PQMPContext, Visibility), DoPerQuad * 0.5);
+#endif
+
+	return Visibility;
+}
diff --git a/Engine/Shaders/ShadowProjectionPixelShader.usf b/Engine/Shaders/ShadowProjectionPixelShader.usf
--- a/Engine/Shaders/ShadowProjectionPixelShader.usf
+++ b/Engine/Shaders/ShadowProjectionPixelShader.usf
@@ -8,8 +8,6 @@
 #define FORCE_FLOATS (COMPILER_METAL)
 #endif
 
-#include "Common.usf"
-
 #ifndef USE_FADE_PLANE
 	#define USE_FADE_PLANE 0
 #endif
@@ -18,6 +16,17 @@
 	#define SHADOW_QUALITY 6
 #endif
 
+#ifndef APPLY_TRANSLUCENCY_SHADOWS
+	#define APPLY_TRANSLUCENCY_SHADOWS 0
+#endif
+
+#ifndef USE_PCSS
+	#define USE_PCSS 0
+#endif
+
+
+#include "Common.usf"
+
 #if FEATURE_LEVEL >= FEATURE_LEVEL_SM5
 	// Gather() is supported
 	#define FEATURE_GATHER4 1
@@ -28,6 +37,9 @@
 
 #include "ShadowProjectionCommon.usf"
 #include "ShadowFilteringCommon.usf"
+#if USE_PCSS
+#include "ShadowPercentageCloserFiltering.usf"
+#endif
 #include "DeferredShadingCommon.usf"
 #include "DynamicLightingCommon.usf"
 
@@ -42,6 +54,13 @@ float2 ProjectionDepthBiasParameters;
 	float InvFadePlaneLength;
 #endif
 
+#if USE_PCSS
+	// PCSS specific parameters.
+	//	- x: tan(0.5 * Directional Light Angle) in shadow projection space;
+	//  - y: Max filter size in shadow tile UV space.
+	float4 PCSSParameters;
+#endif
+
 float4 ModulatedShadowColor;
 float4 ShadowTileOffsetAndSize;
 
@@ -65,11 +84,21 @@ void Main(
 	out float4 OutColor : SV_Target0
 	)
 {
+	const FPQMPContext PQMPContext = PQMPInit(SVPos.xy);
+
 	float2 ScreenUV = float2( SVPos.xy * View.BufferSizeAndInvSize.zw );
 	float SceneW = CalcSceneDepth( ScreenUV );
-	float2 ScreenPosition = ( ScreenUV.xy - View.ScreenPositionScaleBias.wz ) / View.ScreenPositionScaleBias.xy;
-	float4 ShadowPosition = mul(float4(ScreenPosition.xy * SceneW,SceneW,1), ScreenToShadowMatrix);
+	float4 ScreenPosition = float4(((ScreenUV.xy - View.ScreenPositionScaleBias.wz ) / View.ScreenPositionScaleBias.xy) * SceneW, SceneW, 1);
+	float4 ShadowPosition = mul(ScreenPosition, ScreenToShadowMatrix);
+	
+	#if USE_PCSS
+		float3 ScreenPositionDDX = ddx_fine(ScreenPosition.xyz);
+		float3 ScreenPositionDDY = ddy_fine(ScreenPosition.xyz);
+		float4 ShadowPositionDDX = mul(float4(ScreenPositionDDX, 0), ScreenToShadowMatrix);
+		float4 ShadowPositionDDY = mul(float4(ScreenPositionDDY, 0), ScreenToShadowMatrix);
+	#endif
 
+	// TODO: Take care of ShadowPositionDDX/Y for non directional shadows.
 	ShadowPosition.xy /= ShadowPosition.w;
 
 #if MODULATED_SHADOWS
@@ -89,42 +118,60 @@ void Main(
 	float SSSTransmission = 1;
 	
 	float BlendFactor = 1;
-
-#if APPLY_TRANSLUCENCY_SHADOWS
 	
-	Shadow = CalculateTranslucencyShadowing(ShadowPosition.xy, ShadowPosition.z);
-
-#else
-
 	// For debugging
 	#define UNFILTERED_SHADOW_PROJECTION 0
-	#if UNFILTERED_SHADOW_PROJECTION
 
+	#if UNFILTERED_SHADOW_PROJECTION
+	{
 		Shadow = LightSpacePixelDepthForOpaque < Texture2DSampleLevel(ShadowDepthTexture, ShadowDepthTextureSampler, ShadowPosition.xy, 0).r;
-
-	#else
-
+	}
+	#elif APPLY_TRANSLUCENCY_SHADOWS
+	{
+		Shadow = CalculateTranslucencyShadowing(ShadowPosition.xy, ShadowPosition.z);
+	}
+	#elif USE_PCSS
+	{
+		FPCSSSamplerSettings Settings;
+		
+		Settings.ShadowDepthTexture = ShadowDepthTexture;
+		Settings.ShadowDepthTextureSampler = ShadowDepthTextureSampler;
+		Settings.ShadowBufferSize = ShadowBufferSize;
+		Settings.ShadowTileOffsetAndSize = ShadowTileOffsetAndSize;
+		Settings.SceneDepth = LightSpacePixelDepthForOpaque;
+		Settings.TransitionScale = SoftTransitionScale.z;
+		Settings.TanLightSourceAngle = PCSSParameters.x;
+		Settings.MaxKernelSize = PCSSParameters.y;
+		Settings.SvPosition = SVPos;
+		Settings.PQMPContext = PQMPContext;
+		Settings.DebugViewportUV = ScreenUV;
+		
+		Shadow = DirectionalPCSS(Settings, ShadowPosition.xy, ShadowPositionDDX.xyz, ShadowPositionDDY.xyz);
+	}
+	#else // !USE_PCSS
+	{
 		FPCFSamplerSettings Settings;
-
+		
 		Settings.ShadowDepthTexture = ShadowDepthTexture;
 		Settings.ShadowDepthTextureSampler = ShadowDepthTextureSampler;
 		Settings.ShadowBufferSize = ShadowBufferSize;
 		Settings.TransitionScale = SoftTransitionScale.z;
 		Settings.SceneDepth = LightSpacePixelDepthForOpaque;
 		Settings.bSubsurface = false;
+		Settings.bTreatMaxDepthUnshadowed = false;
 		Settings.DensityMulConstant = 0;
 		Settings.ProjectionDepthBiasParameters = 0;
 
 		Shadow = ManualPCF(ShadowPosition.xy, Settings);
-
-	#endif
+	}
+	#endif // !USE_PCSS
 		
 	#if USE_FADE_PLANE
 		// Create a blend factor which is one before and at the fade plane, and lerps to zero at the far plane.
 		BlendFactor = 1.0f - saturate((SceneW - FadePlaneOffset) * InvFadePlaneLength);
 	#endif
 
-	#if FEATURE_LEVEL >= FEATURE_LEVEL_SM4 && !FORWARD_SHADING
+	#if FEATURE_LEVEL >= FEATURE_LEVEL_SM4 && !FORWARD_SHADING && !APPLY_TRANSLUCENCY_SHADOWS
 
 		FGBufferData GBufferData = GetGBufferData(ScreenUV);
 
@@ -160,6 +207,7 @@ void Main(
 				Settings.TransitionScale = SoftTransitionScale.z;
 				Settings.SceneDepth = LightSpacePixelDepthForSSS + ProjectionDepthBiasParameters.x;
 				Settings.bSubsurface = true;
+				Settings.bTreatMaxDepthUnshadowed = false;
 				Settings.DensityMulConstant = Density * ProjectionDepthBiasParameters.y;
 				Settings.ProjectionDepthBiasParameters = ProjectionDepthBiasParameters;
 				
@@ -170,9 +218,11 @@ void Main(
 
 	#endif
 	
-#endif
-		
-	Shadow = saturate( (Shadow - 0.5) * ShadowSharpen + 0.5 );
+	#if !USE_PCSS
+		// There is no point changing the shadow sharpen on PCSS, can directly reduce the light
+		// source angle that would actually make the algorithm to run faster.
+		Shadow = saturate( (Shadow - 0.5) * ShadowSharpen + 0.5 );
+	#endif
 
 	// 0 is shadowed, 1 is unshadowed
 	// RETURN_COLOR not needed unless writing to SceneColor;
diff --git a/Engine/Shaders/TiledDeferredLightShaders.usf b/Engine/Shaders/TiledDeferredLightShaders.usf
--- a/Engine/Shaders/TiledDeferredLightShaders.usf
+++ b/Engine/Shaders/TiledDeferredLightShaders.usf
@@ -271,7 +271,7 @@ void TiledDeferredLightingMain(
 			LightData.ShadowedBits = dot(LightData.ShadowMapChannelMask, float4(1, 1, 1, 1));
 	
 			// Lights requiring light attenuation are not supported tiled for now
-			CompositedLighting += GetDynamicLighting(WorldPosition, CameraVector, ScreenSpaceData.GBuffer, ScreenSpaceData.AmbientOcclusion, InGBufferData.ShadingModelID, LightData, float4(1, 1, 1, 1), uint2(0,0));
+			CompositedLighting += GetDynamicLighting(WorldPosition, CameraVector, ScreenSpaceData.GBuffer, ScreenSpaceData.AmbientOcclusion, InGBufferData.ShadingModelID, LightData, float4(1, 1, 1, 1), 0.5, uint2(0,0));
 		}
 
 		// The simple shading model does not depend on ShadingModelID, so use it anytime the material is lit
diff --git a/Engine/Shaders/TranslucencyUpsampling.usf b/Engine/Shaders/TranslucencyUpsampling.usf
new file mode 100644
--- /dev/null
+++ b/Engine/Shaders/TranslucencyUpsampling.usf
@@ -0,0 +1,22 @@
+// Copyright 1998-2017 Epic Games, Inc. All Rights Reserved.
+
+/*=============================================================================
+	TranslucencyUpsampling.usf: PostProcessing shader to upscale
+=============================================================================*/
+
+#include "Common.usf"
+#include "SeparateTranslucency.usf"
+
+// Texture2D SceneDepthTexture
+Texture2D LowResColorTexture;
+float4 LowResColorTexelSize;
+
+void SimpleUpsamplingPS(FScreenVertexOutput Input, out float4 OutColor : SV_Target0)
+{
+	OutColor = BilinearUpsampling(Input.UV, LowResColorTexture);
+}
+
+void NearestDepthNeighborUpsamplingPS(FScreenVertexOutput Input, out float4 OutColor : SV_Target0)
+{
+	OutColor = NearestDepthNeighborUpsampling(Input.Position.xy, Input.UV, LowResColorTexture, LowResColorTexelSize.zw);
+}
diff --git a/Engine/Shaders/VolumeLightingCommon.usf b/Engine/Shaders/VolumeLightingCommon.usf
--- a/Engine/Shaders/VolumeLightingCommon.usf
+++ b/Engine/Shaders/VolumeLightingCommon.usf
@@ -72,6 +72,8 @@ float ComputeVolumeShadowing(float3 WorldPositionForLighting, bool bPointLight,
 				Settings.SceneDepth = HomogeneousShadowPosition.z;
 				Settings.TransitionScale = 40;
 				Settings.bSubsurface = false;
+				// We can sample outside of the static shadowmap, which is centered around the scene.  These 'infinite' depth values should not cause occlusion.
+				Settings.bTreatMaxDepthUnshadowed = true;
 				Settings.DensityMulConstant = 0;
 				Settings.ProjectionDepthBiasParameters = float2(0, 0);
 
diff --git a/Engine/Shaders/VolumetricFog.usf b/Engine/Shaders/VolumetricFog.usf
--- a/Engine/Shaders/VolumetricFog.usf
+++ b/Engine/Shaders/VolumetricFog.usf
@@ -176,7 +176,9 @@ void WriteToBoundingSphereVS(
 	Output.LayerIndex = LayerIndex + MinZ;
 }
 
-float3 FrameJitterOffset;
+float HistoryWeight;
+float4 FrameJitterOffsets[16]; 
+uint HistoryMissSuperSampleCount;
 float PhaseG;
 float InverseSquaredLightDistanceBiasScale;
 
@@ -194,58 +196,87 @@ void InjectShadowedLocalLightPS(
 
 	uint3 GridCoordinate = uint3(SvPosition.xy, Input.LayerIndex);
 
-	float3 CellOffset = FrameJitterOffset;
-	//float CellOffset = .5f;
-
-	float SceneDepth;
-	float3 WorldPosition = ComputeCellWorldPosition(GridCoordinate, CellOffset, SceneDepth);
-	float3 CameraVector = normalize(WorldPosition - View.WorldCameraOrigin);
-
-	FDeferredLightData LightData;
-	LightData.LightPositionAndInvRadius = float4(DeferredLightUniforms.LightPosition, DeferredLightUniforms.LightInvRadius);
-	LightData.LightColorAndFalloffExponent = float4(DeferredLightUniforms.LightColor, DeferredLightUniforms.LightFalloffExponent);
-	LightData.LightDirection = DeferredLightUniforms.NormalizedLightDirection;
-	LightData.SpotAnglesAndSourceRadius = float4(DeferredLightUniforms.SpotAngles, DeferredLightUniforms.SourceRadius, DeferredLightUniforms.SourceLength);
-	LightData.MinRoughness = DeferredLightUniforms.MinRoughness;
-	LightData.ContactShadowLength = DeferredLightUniforms.ContactShadowLength;
-	LightData.DistanceFadeMAD = DeferredLightUniforms.DistanceFadeMAD;
-	LightData.ShadowMapChannelMask = DeferredLightUniforms.ShadowMapChannelMask;
-	LightData.ShadowedBits = DeferredLightUniforms.ShadowedBits;
-
-	LightData.bInverseSquared = INVERSE_SQUARED_FALLOFF;
-	LightData.bRadialLight = true;
-	LightData.bSpotLight = LightData.SpotAnglesAndSourceRadius.x > -2.0f;
-
-	float VolumetricScatteringIntensity = DeferredLightUniforms.VolumetricScatteringIntensity;
-
-	float3 L = 0;
-	float3 ToLight = 0;
-	float NoL = 0;
-	float DistanceAttenuation = 1;
-	float LightRadiusMask = 1;
-	float SpotFalloff = 1;
-
-	float CellRadius = length(WorldPosition - ComputeCellWorldPosition(GridCoordinate + uint3(1, 1, 1), CellOffset));
-	// Bias the inverse squared light falloff based on voxel size to prevent aliasing near the light source
-	float DistanceBias = max(CellRadius * InverseSquaredLightDistanceBiasScale, 1);
-	GetLocalLightAttenuation(WorldPosition, float3(0, 0, 1), LightData, DistanceBias * DistanceBias, ToLight, L, NoL, DistanceAttenuation, LightRadiusMask, SpotFalloff);
-	float CombinedAttenuation = DistanceAttenuation * LightRadiusMask * SpotFalloff;
-
-	float ShadowFactor = 1.0f;
-	
-	if (CombinedAttenuation > 0)
+	// Somehow pixels are being rasterized outside of the viewport on a 970 GTX, perhaps due to use of a geometry shader bypassing the viewport scissor.
+	// This triggers the HistoryMissSuperSampleCount path causing significant overhead for shading off-screen pixels.
+	if (all(GridCoordinate < VolumetricFog.GridSizeInt))
 	{
-		ShadowFactor = ComputeVolumeShadowing(WorldPosition, LightData.bRadialLight && !LightData.bSpotLight, LightData.bSpotLight);
-	}
+		FDeferredLightData LightData;
+		LightData.LightPositionAndInvRadius = float4(DeferredLightUniforms.LightPosition, DeferredLightUniforms.LightInvRadius);
+		LightData.LightColorAndFalloffExponent = float4(DeferredLightUniforms.LightColor, DeferredLightUniforms.LightFalloffExponent);
+		LightData.LightDirection = DeferredLightUniforms.NormalizedLightDirection;
+		LightData.SpotAnglesAndSourceRadius = float4(DeferredLightUniforms.SpotAngles, DeferredLightUniforms.SourceRadius, DeferredLightUniforms.SourceLength);
+		LightData.MinRoughness = DeferredLightUniforms.MinRoughness;
+		LightData.ContactShadowLength = DeferredLightUniforms.ContactShadowLength;
+		LightData.DistanceFadeMAD = DeferredLightUniforms.DistanceFadeMAD;
+		LightData.ShadowMapChannelMask = DeferredLightUniforms.ShadowMapChannelMask;
+		LightData.ShadowedBits = DeferredLightUniforms.ShadowedBits;
+
+		LightData.bInverseSquared = INVERSE_SQUARED_FALLOFF;
+		LightData.bRadialLight = true;
+		LightData.bSpotLight = LightData.SpotAnglesAndSourceRadius.x > -2.0f;
+
+		float VolumetricScatteringIntensity = DeferredLightUniforms.VolumetricScatteringIntensity;
+
+		float3 L = 0;
+		float3 ToLight = 0;
+		float NoL = 0;
+		float DistanceAttenuation = 1;
+		float LightRadiusMask = 1;
+		float SpotFalloff = 1;
+
+		uint NumSuperSamples = 1;
+
+#if USE_TEMPORAL_REPROJECTION
+
+		float3 HistoryUV = ComputeVolumeUV(ComputeCellWorldPosition(GridCoordinate, .5f), UnjitteredPrevWorldToClip);
+		float HistoryAlpha = HistoryWeight;
+
+		FLATTEN
+		if (any(HistoryUV < 0) || any(HistoryUV > 1))
+		{
+			HistoryAlpha = 0;
+		}
+
+		NumSuperSamples = HistoryAlpha < .001f ? HistoryMissSuperSampleCount : 1;
+
+#endif
+
+		for (uint SampleIndex = 0; SampleIndex < NumSuperSamples; SampleIndex++)
+		{
+			float3 CellOffset = FrameJitterOffsets[SampleIndex].xyz;
+			//float CellOffset = .5f;
+
+			float SceneDepth;
+			float3 WorldPosition = ComputeCellWorldPosition(GridCoordinate, CellOffset, SceneDepth);
+			float3 CameraVector = normalize(WorldPosition - View.WorldCameraOrigin);
+
+			float CellRadius = length(WorldPosition - ComputeCellWorldPosition(GridCoordinate + uint3(1, 1, 1), CellOffset));
+			// Bias the inverse squared light falloff based on voxel size to prevent aliasing near the light source
+			float DistanceBias = max(CellRadius * InverseSquaredLightDistanceBiasScale, 1);
+			GetLocalLightAttenuation(WorldPosition, float3(0, 0, 1), LightData, DistanceBias * DistanceBias, ToLight, L, NoL, DistanceAttenuation, LightRadiusMask, SpotFalloff);
+			float CombinedAttenuation = DistanceAttenuation * LightRadiusMask * SpotFalloff;
+
+			float ShadowFactor = 1.0f;
+
+			if (CombinedAttenuation > 0)
+			{
+				ShadowFactor = ComputeVolumeShadowing(WorldPosition, LightData.bRadialLight && !LightData.bSpotLight, LightData.bSpotLight);
+			}
 
-	OutScattering.rgb = DeferredLightUniforms.LightColor
-		* (PhaseFunction(PhaseG, dot(L, -CameraVector)) * CombinedAttenuation * ShadowFactor * VolumetricScatteringIntensity);
+			OutScattering.rgb += DeferredLightUniforms.LightColor
+				* (PhaseFunction(PhaseG, dot(L, -CameraVector)) * CombinedAttenuation * ShadowFactor * VolumetricScatteringIntensity);
+				
+			// To debug culling
+			//OutScattering.rgb += DeferredLightUniforms.LightColor * .0000001f;
+		}
+
+		OutScattering.rgb /= (float)NumSuperSamples;
+	}
 }
 
 Texture3D<float4> VBufferA;
 Texture3D<float4> VBufferB;
 
-float HistoryWeight;
 Texture3D<float4> LightScatteringHistory;
 SamplerState LightScatteringHistorySampler;
 
@@ -397,6 +428,8 @@ float ComputeStaticShadowing(float3 WorldPosition)
 				Settings.SceneDepth = HomogeneousShadowPosition.z;
 				Settings.TransitionScale = 40;
 				Settings.bSubsurface = false;
+				// We can sample outside of the static shadowmap, which is centered around the lightmass importance volume.  These 'infinite' depth values should not cause occlusion.
+				Settings.bTreatMaxDepthUnshadowed = true;
 				Settings.DensityMulConstant = 0;
 				Settings.ProjectionDepthBiasParameters = float2(0, 0);
 
@@ -404,7 +437,7 @@ float ComputeStaticShadowing(float3 WorldPosition)
 			#else
 				// Sample the shadowmap depth and determine if this voxel is shadowed
 				float ShadowDepth = Texture2DSampleLevel(ForwardGlobalLightData.DirectionalLightStaticShadowmap, ForwardGlobalLightData.StaticShadowmapSampler, ShadowUVs, 0).x;
-				ShadowFactor = HomogeneousShadowPosition.z < ShadowDepth;
+				ShadowFactor = HomogeneousShadowPosition.z < ShadowDepth || ShadowDepth > .99f;
 			#endif
 		}
 	}
@@ -424,164 +457,185 @@ void LightScatteringCS(
 	uint3 GroupThreadId : SV_GroupThreadID)
 {
 	uint3 GridCoordinate = DispatchThreadId;
+	float3 LightScattering = 0;
+	uint NumSuperSamples = 1;
 
-	float3 CellOffset = FrameJitterOffset;
-	//float CellOffset = .5f;
+#if USE_TEMPORAL_REPROJECTION
+	float3 HistoryUV = ComputeVolumeUV(ComputeCellWorldPosition(GridCoordinate, .5f), UnjitteredPrevWorldToClip);
+	float HistoryAlpha = HistoryWeight;
 
-	float SceneDepth;
-	float3 WorldPosition = ComputeCellWorldPosition(GridCoordinate, CellOffset, SceneDepth);
-	float CameraVectorLength = length(WorldPosition - View.WorldCameraOrigin);
-	float3 CameraVector = (WorldPosition - View.WorldCameraOrigin) / CameraVectorLength;
-	float4 MaterialScatteringAndAbsorption = VBufferA[GridCoordinate];
+	FLATTEN
+	if (any(HistoryUV < 0) || any(HistoryUV > 1))
+	{
+		HistoryAlpha = 0;
+	}
 
-	// Shadowed point and spot lights were computed earlier
-	float3 LightScattering = LocalShadowedLightScattering[GridCoordinate].xyz;
+	// Supersample if the history was outside the camera frustum
+	// The compute shader is dispatched with extra threads, make sure those don't supersample
+	NumSuperSamples = HistoryAlpha < .001f && all(GridCoordinate < VolumetricFog.GridSizeInt) ? HistoryMissSuperSampleCount : 1;
 
-	BRANCH
-	if (ForwardGlobalLightData.HasDirectionalLight)
+#endif
+
+	for (uint SampleIndex = 0; SampleIndex < NumSuperSamples; SampleIndex++)
 	{
-		float ShadowFactor = ComputeStaticShadowing(WorldPosition);
+		float3 CellOffset = FrameJitterOffsets[SampleIndex].xyz;
+		//float CellOffset = .5f;
+
+		float SceneDepth;
+		float3 WorldPosition = ComputeCellWorldPosition(GridCoordinate, CellOffset, SceneDepth);
+		float CameraVectorLength = length(WorldPosition - View.WorldCameraOrigin);
+		float3 CameraVector = (WorldPosition - View.WorldCameraOrigin) / CameraVectorLength;
 
-		if (ForwardGlobalLightData.NumDirectionalLightCascades > 0 && UseDirectionalLightShadowing > 0)
+		BRANCH
+		if (ForwardGlobalLightData.HasDirectionalLight)
 		{
-			uint CascadeIndex = ForwardGlobalLightData.NumDirectionalLightCascades;
+			float ShadowFactor = ComputeStaticShadowing(WorldPosition);
 
-			for (uint TestCascadeIndex = 0; TestCascadeIndex < ForwardGlobalLightData.NumDirectionalLightCascades; TestCascadeIndex++)
+			if (ForwardGlobalLightData.NumDirectionalLightCascades > 0 && UseDirectionalLightShadowing > 0)
 			{
-				if (SceneDepth < ForwardGlobalLightData.CascadeEndDepths[TestCascadeIndex])
+				uint CascadeIndex = ForwardGlobalLightData.NumDirectionalLightCascades;
+
+				for (uint TestCascadeIndex = 0; TestCascadeIndex < ForwardGlobalLightData.NumDirectionalLightCascades; TestCascadeIndex++)
 				{
-					CascadeIndex = TestCascadeIndex;
-					break;
+					if (SceneDepth < ForwardGlobalLightData.CascadeEndDepths[TestCascadeIndex])
+					{
+						CascadeIndex = TestCascadeIndex;
+						break;
+					}
 				}
-			}
 
-			if (CascadeIndex < ForwardGlobalLightData.NumDirectionalLightCascades)
-			{
-				// Transform the world position into shadowmap space
-				float4 HomogeneousShadowPosition = mul(float4(WorldPosition, 1), ForwardGlobalLightData.DirectionalLightWorldToShadowMatrix[CascadeIndex]);
-				float2 ShadowUVs = HomogeneousShadowPosition.xy / HomogeneousShadowPosition.w;
-				float4 ShadowmapMinMax = ForwardGlobalLightData.DirectionalLightShadowmapMinMax[CascadeIndex];
-
-				// Treat as unshadowed if the voxel is outside of the shadow map
-				if (all(ShadowUVs >= ShadowmapMinMax.xy && ShadowUVs <= ShadowmapMinMax.zw))
+				if (CascadeIndex < ForwardGlobalLightData.NumDirectionalLightCascades)
 				{
-					// Sample the shadowmap depth and determine if this voxel is shadowed
-					float ShadowDepth = Texture2DSampleLevel(ForwardGlobalLightData.DirectionalLightShadowmapAtlas, ForwardGlobalLightData.ShadowmapSampler, ShadowUVs, 0).x;
-					ShadowFactor = HomogeneousShadowPosition.z < ShadowDepth - ForwardGlobalLightData.DirectionalLightDepthBias.x;
+					// Transform the world position into shadowmap space
+					float4 HomogeneousShadowPosition = mul(float4(WorldPosition, 1), ForwardGlobalLightData.DirectionalLightWorldToShadowMatrix[CascadeIndex]);
+					float2 ShadowUVs = HomogeneousShadowPosition.xy / HomogeneousShadowPosition.w;
+					float4 ShadowmapMinMax = ForwardGlobalLightData.DirectionalLightShadowmapMinMax[CascadeIndex];
+
+					// Treat as unshadowed if the voxel is outside of the shadow map
+					if (all(ShadowUVs >= ShadowmapMinMax.xy && ShadowUVs <= ShadowmapMinMax.zw))
+					{
+						// Sample the shadowmap depth and determine if this voxel is shadowed
+						float ShadowDepth = Texture2DSampleLevel(ForwardGlobalLightData.DirectionalLightShadowmapAtlas, ForwardGlobalLightData.ShadowmapSampler, ShadowUVs, 0).x;
+						ShadowFactor = HomogeneousShadowPosition.z < ShadowDepth - ForwardGlobalLightData.DirectionalLightDepthBias.x;
+					}
 				}
 			}
-		}
 
-		ShadowFactor *= GetLightFunction(WorldPosition);
+			ShadowFactor *= GetLightFunction(WorldPosition);
 
-		float3 DirectionalLightColor = ForwardGlobalLightData.DirectionalLightColor;
-
-		if (UseHeightFogColors > 0)
-		{
-			// Attempt to maintain intensity ratio between sky and sun
-			DirectionalLightColor = VolumetricFog.HeightFogDirectionalLightInscatteringColor * Luminance(ForwardGlobalLightData.DirectionalLightColor);
-		}
+			float3 DirectionalLightColor = ForwardGlobalLightData.DirectionalLightColor;
 
-		LightScattering += DirectionalLightColor
-			* (ShadowFactor 
-			* ForwardGlobalLightData.DirectionalLightVolumetricScatteringIntensity 
-			* PhaseFunction(PhaseG, dot(ForwardGlobalLightData.DirectionalLightDirection, -CameraVector)));
-	}
-
-	// Skylight
-	if (SkyLightVolumetricScatteringIntensity > 0)
-	{
-		FTwoBandSHVector RotatedHGZonalHarmonic;
-		RotatedHGZonalHarmonic.V = float4(1.0f, CameraVector.y, CameraVector.z, CameraVector.x) * float4(1.0f, PhaseG, PhaseG, PhaseG);
-		
-		float3 SkyLighting;
+			if (UseHeightFogColors > 0)
+			{
+				// Attempt to maintain intensity ratio between sky and sun
+				DirectionalLightColor = VolumetricFog.HeightFogDirectionalLightInscatteringColor * Luminance(ForwardGlobalLightData.DirectionalLightColor);
+			}
 
-		if (UseHeightFogColors > 0)
-		{
-			float3 HeightFogInscatteringColor = ComputeInscatteringColor(CameraVector, CameraVectorLength);
-			float ScalarFactor = SHAmbientFunction();
-			FTwoBandSHVectorRGB SkyIrradianceSH;
-			SkyIrradianceSH.R.V = float4(ScalarFactor * HeightFogInscatteringColor.r, 0, 0, 0);
-			SkyIrradianceSH.G.V = float4(ScalarFactor * HeightFogInscatteringColor.g, 0, 0, 0);
-			SkyIrradianceSH.B.V = float4(ScalarFactor * HeightFogInscatteringColor.b, 0, 0, 0);
-
-			SkyLighting = max(DotSH(SkyIrradianceSH, RotatedHGZonalHarmonic), 0);
+			LightScattering += DirectionalLightColor
+				* (ShadowFactor 
+				* ForwardGlobalLightData.DirectionalLightVolumetricScatteringIntensity 
+				* PhaseFunction(PhaseG, dot(ForwardGlobalLightData.DirectionalLightDirection, -CameraVector)));
 		}
-		else
-		{
-			FTwoBandSHVectorRGB SkyIrradianceSH;
-			SkyIrradianceSH.R.V = SkySH[0];
-			SkyIrradianceSH.G.V = SkySH[1];
-			SkyIrradianceSH.B.V = SkySH[2];
 
-			SkyLighting = View.SkyLightColor.rgb * max(DotSH(SkyIrradianceSH, RotatedHGZonalHarmonic), 0) / PI;
-		}
+		// Skylight
+		if (SkyLightVolumetricScatteringIntensity > 0)
+		{
+			FTwoBandSHVector RotatedHGZonalHarmonic;
+			RotatedHGZonalHarmonic.V = float4(1.0f, CameraVector.y, CameraVector.z, CameraVector.x) * float4(1.0f, PhaseG, PhaseG, PhaseG);
+		
+			float3 SkyLighting;
 
-		float SkyVisibility = ComputeSkyVisibility(WorldPosition);
-		LightScattering += (SkyVisibility * SkyLightVolumetricScatteringIntensity) * SkyLighting;
-	}
+			if (UseHeightFogColors > 0)
+			{
+				float3 HeightFogInscatteringColor = ComputeInscatteringColor(CameraVector, CameraVectorLength);
+				float ScalarFactor = SHAmbientFunction();
+				FTwoBandSHVectorRGB SkyIrradianceSH;
+				SkyIrradianceSH.R.V = float4(ScalarFactor * HeightFogInscatteringColor.r, 0, 0, 0);
+				SkyIrradianceSH.G.V = float4(ScalarFactor * HeightFogInscatteringColor.g, 0, 0, 0);
+				SkyIrradianceSH.B.V = float4(ScalarFactor * HeightFogInscatteringColor.b, 0, 0, 0);
+
+				SkyLighting = max(DotSH(SkyIrradianceSH, RotatedHGZonalHarmonic), 0);
+			}
+			else
+			{
+				FTwoBandSHVectorRGB SkyIrradianceSH;
+				SkyIrradianceSH.R.V = SkySH[0];
+				SkyIrradianceSH.G.V = SkySH[1];
+				SkyIrradianceSH.B.V = SkySH[2];
 
-	uint GridIndex = ComputeLightGridCellIndex(GridCoordinate.xy * VolumetricFog.FogGridToPixelXY, SceneDepth, 0);
-	const FCulledLightsGridData CulledLightsGrid = GetCulledLightsGrid(GridIndex, 0);
+				SkyLighting = View.SkyLightColor.rgb * max(DotSH(SkyIrradianceSH, RotatedHGZonalHarmonic), 0) / PI;
+			}
 
-	float CellRadius = length(WorldPosition - ComputeCellWorldPosition(GridCoordinate + uint3(1, 1, 1), CellOffset));
-	// Bias the inverse squared light falloff based on voxel size to prevent aliasing near the light source
-	float DistanceBiasSqr = max(CellRadius * InverseSquaredLightDistanceBiasScale, 1);
-	DistanceBiasSqr *= DistanceBiasSqr;
+			float SkyVisibility = ComputeSkyVisibility(WorldPosition);
+			LightScattering += (SkyVisibility * SkyLightVolumetricScatteringIntensity) * SkyLighting;
+		}
 
-	// Forward lighting of unshadowed point and spot lights
-	LOOP
-	for (uint LocalLightListIndex = 0; LocalLightListIndex < CulledLightsGrid.NumLocalLights; LocalLightListIndex++)
-	{
-		const FLocalLightData LocalLight = GetLocalLightData(CulledLightsGrid.DataStartIndex + LocalLightListIndex, 0);
+		uint GridIndex = ComputeLightGridCellIndex(GridCoordinate.xy * VolumetricFog.FogGridToPixelXY, SceneDepth, 0);
+		const FCulledLightsGridData CulledLightsGrid = GetCulledLightsGrid(GridIndex, 0);
 
-		float VolumetricScatteringIntensity = f16tof32(asuint(LocalLight.SpotAnglesAndSourceRadiusPacked.w) >> 16);
+		float CellRadius = length(WorldPosition - ComputeCellWorldPosition(GridCoordinate + uint3(1, 1, 1), CellOffset));
+		// Bias the inverse squared light falloff based on voxel size to prevent aliasing near the light source
+		float DistanceBiasSqr = max(CellRadius * InverseSquaredLightDistanceBiasScale, 1);
+		DistanceBiasSqr *= DistanceBiasSqr;
 
-		if (VolumetricScatteringIntensity > 0)
+		// Forward lighting of unshadowed point and spot lights
+		LOOP
+		for (uint LocalLightListIndex = 0; LocalLightListIndex < CulledLightsGrid.NumLocalLights; LocalLightListIndex++)
 		{
-			FDeferredLightData LightData = (FDeferredLightData)0;
-			LightData.LightPositionAndInvRadius = LocalLight.LightPositionAndInvRadius;
-			LightData.LightColorAndFalloffExponent = LocalLight.LightColorAndFalloffExponent;
-			LightData.LightDirection = LocalLight.LightDirectionAndShadowMask.xyz;
-			float SourceLength = f16tof32(asuint(LocalLight.SpotAnglesAndSourceRadiusPacked.w));
-			LightData.SpotAnglesAndSourceRadius = float4(LocalLight.SpotAnglesAndSourceRadiusPacked.xyz, SourceLength);
-			LightData.bInverseSquared = LightData.LightColorAndFalloffExponent.w == 0;
-			LightData.bRadialLight = true;
-			LightData.bSpotLight = LightData.SpotAnglesAndSourceRadius.x > -2.0f;
-
-			float3 L = 0;
-			float3 ToLight = 0;
-			float NoL = 0;
-			float DistanceAttenuation = 1;
-			float LightRadiusMask = 1;
-			float SpotFalloff = 1;
-
-			GetLocalLightAttenuation(WorldPosition, float3(0, 0, 1), LightData, DistanceBiasSqr, ToLight, L, NoL, DistanceAttenuation, LightRadiusMask, SpotFalloff);
+			const FLocalLightData LocalLight = GetLocalLightData(CulledLightsGrid.DataStartIndex + LocalLightListIndex, 0);
 
-			float CombinedAttenuation = DistanceAttenuation * LightRadiusMask * SpotFalloff;
+			float VolumetricScatteringIntensity = f16tof32(asuint(LocalLight.SpotAnglesAndSourceRadiusPacked.w) >> 16);
 
-			LightScattering += LocalLight.LightColorAndFalloffExponent.xyz
-				* (PhaseFunction(PhaseG, dot(L, -CameraVector)) * CombinedAttenuation * VolumetricScatteringIntensity);
+			if (VolumetricScatteringIntensity > 0)
+			{
+				FDeferredLightData LightData = (FDeferredLightData)0;
+				LightData.LightPositionAndInvRadius = LocalLight.LightPositionAndInvRadius;
+				LightData.LightColorAndFalloffExponent = LocalLight.LightColorAndFalloffExponent;
+				LightData.LightDirection = LocalLight.LightDirectionAndShadowMask.xyz;
+				float SourceLength = f16tof32(asuint(LocalLight.SpotAnglesAndSourceRadiusPacked.w));
+				LightData.SpotAnglesAndSourceRadius = float4(LocalLight.SpotAnglesAndSourceRadiusPacked.xyz, SourceLength);
+				LightData.bInverseSquared = LightData.LightColorAndFalloffExponent.w == 0;
+				LightData.bRadialLight = true;
+				LightData.bSpotLight = LightData.SpotAnglesAndSourceRadius.x > -2.0f;
+
+				float3 L = 0;
+				float3 ToLight = 0;
+				float NoL = 0;
+				float DistanceAttenuation = 1;
+				float LightRadiusMask = 1;
+				float SpotFalloff = 1;
+
+				GetLocalLightAttenuation(WorldPosition, float3(0, 0, 1), LightData, DistanceBiasSqr, ToLight, L, NoL, DistanceAttenuation, LightRadiusMask, SpotFalloff);
+
+				float CombinedAttenuation = DistanceAttenuation * LightRadiusMask * SpotFalloff;
+
+				LightScattering += LocalLight.LightColorAndFalloffExponent.xyz
+					* (PhaseFunction(PhaseG, dot(L, -CameraVector)) * CombinedAttenuation * VolumetricScatteringIntensity);
+
+				// To debug culling
+				//LightScattering += LocalLight.LightColorAndFalloffExponent.xyz * .0000001f;
+			}
 		}
 	}
 
+	LightScattering /= (float)NumSuperSamples;
+
+	// Shadowed point and spot lights were computed earlier
+	LightScattering += LocalShadowedLightScattering[GridCoordinate].xyz;
+
+	float4 MaterialScatteringAndAbsorption = VBufferA[GridCoordinate];
 	float Extinction = MaterialScatteringAndAbsorption.w + Luminance(MaterialScatteringAndAbsorption.xyz);
 	float3 MaterialEmissive = VBufferB[GridCoordinate].xyz;
 	float4 ScatteringAndExtinction = EncodeHDR(float4(LightScattering * MaterialScatteringAndAbsorption.xyz + MaterialEmissive, Extinction));
 
 #if USE_TEMPORAL_REPROJECTION
-	float3 HistoryUV = ComputeVolumeUV(ComputeCellWorldPosition(GridCoordinate, .5f), UnjitteredPrevWorldToClip);
-	float4 HistoryScatteringAndExtinction = Texture3DSampleLevel(LightScatteringHistory, LightScatteringHistorySampler, HistoryUV, 0);
-
-	float HistoryAlpha = HistoryWeight;
-
-	FLATTEN
-	if (any(HistoryUV < 0) || any(HistoryUV > 1))
+	BRANCH
+	if (HistoryAlpha > 0)
 	{
-		HistoryAlpha = 0;
+		float4 HistoryScatteringAndExtinction = Texture3DSampleLevel(LightScatteringHistory, LightScatteringHistorySampler, HistoryUV, 0);
+		ScatteringAndExtinction = lerp(ScatteringAndExtinction, HistoryScatteringAndExtinction, HistoryAlpha);
 	}
-
-	ScatteringAndExtinction = lerp(ScatteringAndExtinction, HistoryScatteringAndExtinction, HistoryAlpha);
+	
 #endif
 
 	if (all(GridCoordinate < VolumetricFog.GridSizeInt))
diff --git a/Engine/Source/Developer/FunctionalTesting/Classes/FunctionalTest.h b/Engine/Source/Developer/FunctionalTesting/Classes/FunctionalTest.h
--- a/Engine/Source/Developer/FunctionalTesting/Classes/FunctionalTest.h
+++ b/Engine/Source/Developer/FunctionalTesting/Classes/FunctionalTest.h
@@ -406,7 +406,7 @@ public:
 	 * @param What	A name to use in the message if the assert fails ("Expected 'What' to be {Expected} but it was {Actual} for context ''")
 	 */
 	UFUNCTION(BlueprintCallable, Category = "Asserts", DisplayName = "Assert Equal (Transform)", meta = ( HidePin = "ContextObject", DefaultToSelf = "ContextObject"))
-	bool AssertEqual_Transform(const FTransform Actual, const FTransform Expected, const FString& What, const UObject* ContextObject = nullptr);
+	bool AssertEqual_Transform(const FTransform& Actual, const FTransform& Expected, const FString& What, const UObject* ContextObject = nullptr);
 
 	/**
 	 * Assert that two floats are equal within tolerance between two floats.
@@ -441,7 +441,7 @@ public:
 	 * @param What	A name to use in the message if the assert fails ("Expected 'What' not to be {Expected} but it was {Actual} for context ''")
 	 */
 	UFUNCTION(BlueprintCallable, Category = "Asserts", DisplayName = "Assert Not Equal (Transform)", meta = ( HidePin = "ContextObject", DefaultToSelf = "ContextObject"))
-	bool AssertNotEqual_Transform(const FTransform Actual, const FTransform NotExpected, const FString& What, const UObject* ContextObject = nullptr);
+	bool AssertNotEqual_Transform(const FTransform& Actual, const FTransform& NotExpected, const FString& What, const UObject* ContextObject = nullptr);
 
 	/**
 	 * Assert that the component angles of two rotators are all equal within a small tolerance.
diff --git a/Engine/Source/Developer/FunctionalTesting/Private/FunctionalTest.cpp b/Engine/Source/Developer/FunctionalTesting/Private/FunctionalTest.cpp
--- a/Engine/Source/Developer/FunctionalTesting/Private/FunctionalTest.cpp
+++ b/Engine/Source/Developer/FunctionalTesting/Private/FunctionalTest.cpp
@@ -660,7 +660,7 @@ bool AFunctionalTest::AssertEqual_Name(const FName Actual, const FName Expected,
 }
 
 
-bool AFunctionalTest::AssertEqual_Transform(const FTransform Actual, const FTransform Expected, const FString& What, const UObject* ContextObject)
+bool AFunctionalTest::AssertEqual_Transform(const FTransform& Actual, const FTransform& Expected, const FString& What, const UObject* ContextObject)
 {
 	if ( !Expected.Equals(Actual) )
 	{
@@ -674,7 +674,7 @@ bool AFunctionalTest::AssertEqual_Transform(const FTransform Actual, const FTran
 	}
 }
 
-bool AFunctionalTest::AssertNotEqual_Transform(const FTransform Actual, const FTransform NotExpected, const FString& What, const UObject* ContextObject)
+bool AFunctionalTest::AssertNotEqual_Transform(const FTransform& Actual, const FTransform& NotExpected, const FString& What, const UObject* ContextObject)
 {
 	if ( NotExpected.Equals(Actual) )
 	{
diff --git a/Engine/Source/Developer/MeshSimplifier/Private/MeshSimplify.h b/Engine/Source/Developer/MeshSimplifier/Private/MeshSimplify.h
--- a/Engine/Source/Developer/MeshSimplifier/Private/MeshSimplify.h
+++ b/Engine/Source/Developer/MeshSimplifier/Private/MeshSimplify.h
@@ -7,10 +7,10 @@
 
 #define SIMP_CACHE	1
 
-#include "Developer/MeshSimplifier/Private/HashTable.h"
-#include "Developer/MeshSimplifier/Private/BinaryHeap.h"
-#include "Developer/MeshSimplifier/Private/MeshSimplifyElements.h"
-#include "Developer/MeshSimplifier/Private/Quadric.h"
+#include "MeshSimplifyElements.h"
+#include "Quadric.h"
+#include "Containers/HashTable.h"
+#include "Containers/BinaryHeap.h"
 //#include "Cache.h"
 
 template< typename T, uint32 NumAttributes >
diff --git a/Engine/Source/Developer/MeshUtilities/Private/LayoutUV.cpp b/Engine/Source/Developer/MeshUtilities/Private/LayoutUV.cpp
--- a/Engine/Source/Developer/MeshUtilities/Private/LayoutUV.cpp
+++ b/Engine/Source/Developer/MeshUtilities/Private/LayoutUV.cpp
@@ -205,16 +205,7 @@ void FLayoutUV::FindCharts( const TMultiMap<int32,int32>& OverlappingCorners )
 		Chart.LastTri = Tri;
 
 #if !CHART_JOINING
-		if( Chart.UVArea > 1e-4f )
-		{
-			Chart.WorldScale /= Chart.UVArea;
-		}
-		else
-		{
-			Chart.WorldScale = FVector2D::ZeroVector;
-		}
-
-		//Chart.WorldScale.Set(1,1);
+		Chart.WorldScale /= FMath::Max( Chart.UVArea, 1e-8f );
 
 		TotalUVArea += Chart.UVArea * Chart.WorldScale.X * Chart.WorldScale.Y;
 #endif
@@ -519,14 +510,7 @@ void FLayoutUV::FindCharts( const TMultiMap<int32,int32>& OverlappingCorners )
 	{
 		FMeshChart& Chart = Charts[i];
 
-		if( Chart.UVArea > 1e-4f )
-		{
-			Chart.WorldScale /= Chart.UVArea;
-		}
-		else
-		{
-			Chart.WorldScale = FVector2D::ZeroVector;
-		}
+		Chart.WorldScale /= FMath::Max( Chart.UVArea, 1e-8f );
 
 		TotalUVArea += Chart.UVArea * Chart.WorldScale.X * Chart.WorldScale.Y;
 	}
diff --git a/Engine/Source/Developer/MeshUtilities/Private/LayoutUV.h b/Engine/Source/Developer/MeshUtilities/Private/LayoutUV.h
--- a/Engine/Source/Developer/MeshUtilities/Private/LayoutUV.h
+++ b/Engine/Source/Developer/MeshUtilities/Private/LayoutUV.h
@@ -84,7 +84,7 @@ private:
 };
 
 
-#define THRESH_UVS_ARE_SAME (1.0f / 1024.0f)
+#define THRESH_UVS_ARE_SAME THRESH_POINTS_ARE_SAME
 
 inline bool FLayoutUV::PositionsMatch( uint32 a, uint32 b ) const
 {
diff --git a/Engine/Source/Developer/PackageDependencyInfo/Private/PackageDependencyInfo.cpp b/Engine/Source/Developer/PackageDependencyInfo/Private/PackageDependencyInfo.cpp
--- a/Engine/Source/Developer/PackageDependencyInfo/Private/PackageDependencyInfo.cpp
+++ b/Engine/Source/Developer/PackageDependencyInfo/Private/PackageDependencyInfo.cpp
@@ -823,9 +823,6 @@ void FPackageDependencyInfo::DetermineShaderSourceTimeStamp()
 {
 	ShaderSourceTimeStamp = FDateTime::MinValue();
 
-	// Get all the shader source files
-	FString ShaderSourceDirectory = FPlatformProcess::ShaderDir();
-
 	// use the timestamp grabbing visitor (include directories)
 	TArray<FString> DirectoriesWildcards;
 	TArray<FString> DirectoriesToIgnore;
@@ -834,7 +831,10 @@ void FPackageDependencyInfo::DetermineShaderSourceTimeStamp()
 	FileWildcards.Add(TEXT(".usf"));
 
 	FPackageDependencyTimestampVisitor TimeStampVisitor(FPlatformFileManager::Get().GetPlatformFile(), DirectoriesWildcards, DirectoriesToIgnore, DirectoriesToNotRecurse, FileWildcards, false, true);
-	TimeStampVisitor.Visit(*ShaderSourceDirectory, true);
+	for (FString ShaderDir : FPlatformProcess::AllShaderDirs())
+	{
+		TimeStampVisitor.Visit(*ShaderDir, true);
+	}
 
 	FMD5 AllShaderHash;
 	TimeStampVisitor.FileTimes.KeySort(TLess<FString>());
diff --git a/Engine/Source/Developer/ShaderCompilerCommon/Private/HlslUtils.h b/Engine/Source/Developer/ShaderCompilerCommon/Private/HlslUtils.h
--- a/Engine/Source/Developer/ShaderCompilerCommon/Private/HlslUtils.h
+++ b/Engine/Source/Developer/ShaderCompilerCommon/Private/HlslUtils.h
@@ -136,7 +136,7 @@ namespace CrossCompiler
 				{
 					// Allocate memory from the stack.
 					Data = (ElementType*)LinearAllocator->Alloc(NumElements * NumBytesPerElement,
-						FMath::Max((uint32)sizeof(void*), (uint32)ALIGNOF(ElementType))
+						FMath::Max((uint32)sizeof(void*), (uint32)alignof(ElementType))
 						);
 
 					// If the container previously held elements, copy them into the new allocation.
diff --git a/Engine/Source/Developer/ShaderPreprocessor/Private/ShaderPreprocessor.cpp b/Engine/Source/Developer/ShaderPreprocessor/Private/ShaderPreprocessor.cpp
--- a/Engine/Source/Developer/ShaderPreprocessor/Private/ShaderPreprocessor.cpp
+++ b/Engine/Source/Developer/ShaderPreprocessor/Private/ShaderPreprocessor.cpp
@@ -32,19 +32,23 @@ public:
 	explicit FMcppFileLoader(const FShaderCompilerInput& InShaderInput)
 		: ShaderInput(InShaderInput)
 	{
-		FString ShaderDir = FPlatformProcess::ShaderDir();
-
-		InputShaderFile = ShaderDir / FPaths::GetCleanFilename(ShaderInput.SourceFilename);
+		// SourceFilename is expected to be relative to the engine shader folder 
+		InputShaderFile = ShaderInput.SourceFilename;
+		
+		// If there is no extension then we can end up with an extra version of this file
 		if (FPaths::GetExtension(InputShaderFile) != TEXT("usf"))
 		{
 			InputShaderFile += TEXT(".usf");
 		}
 
+		// Attempt to keep filename reference and map add logic the same as the file contents callback
+		FString Filename = GetRelativeShaderFilename(InputShaderFile);
 		FString InputShaderSource;
-		if (LoadShaderSourceFile(*ShaderInput.SourceFilename,InputShaderSource))
+		
+		if (LoadShaderSourceFile(*Filename, InputShaderSource))
 		{
 			InputShaderSource = FString::Printf(TEXT("%s\n#line 1\n%s"), *ShaderInput.SourceFilePrefix, *InputShaderSource);
-			CachedFileContents.Add(GetRelativeShaderFilename(InputShaderFile),StringToArray<ANSICHAR>(*InputShaderSource, InputShaderSource.Len()));
+			CachedFileContents.Add(Filename, StringToArray<ANSICHAR>(*InputShaderSource, InputShaderSource.Len()));
 		}
 	}
 
diff --git a/Engine/Source/Developer/ShaderPreprocessor/Public/ShaderPreprocessor.h b/Engine/Source/Developer/ShaderPreprocessor/Public/ShaderPreprocessor.h
--- a/Engine/Source/Developer/ShaderPreprocessor/Public/ShaderPreprocessor.h
+++ b/Engine/Source/Developer/ShaderPreprocessor/Public/ShaderPreprocessor.h
@@ -19,16 +19,3 @@ extern SHADERPREPROCESSOR_API bool PreprocessShader(
 	const FShaderCompilerInput& ShaderInput,
 	const FShaderCompilerDefinitions& AdditionalDefines
 	);
-
-/**
- * Preprocess a shader file, for non-usf files.
- * @param OutPreprocessedShader - Upon return contains the preprocessed source code.
- * @param ShaderOutput - ShaderOutput to which errors can be added.
- * @param ShaderInput - The shader compiler input.
- * @returns true if the shader is preprocessed without error.
- */
-extern SHADERPREPROCESSOR_API bool PreprocessShaderFile(
-	FString& OutPreprocessedShader,
-	TArray<FShaderCompilerError>& OutShaderErrors,
-	const FString& InShaderFile
-	);
diff --git a/Engine/Source/Developer/SlateFileDialogs/Private/SlateFileDlgWindow.cpp b/Engine/Source/Developer/SlateFileDialogs/Private/SlateFileDlgWindow.cpp
--- a/Engine/Source/Developer/SlateFileDialogs/Private/SlateFileDlgWindow.cpp
+++ b/Engine/Source/Developer/SlateFileDialogs/Private/SlateFileDlgWindow.cpp
@@ -34,11 +34,19 @@ class FSlateFileDialogVisitor : public IPlatformFile::FDirectoryVisitor
 {
 public:
 	FSlateFileDialogVisitor(TArray<TSharedPtr<FFileEntry>> &InFileList,
-						TArray<TSharedPtr<FFileEntry>> &InFolderList, TSharedPtr<FString> InFilterList)
+						TArray<TSharedPtr<FFileEntry>> &InFolderList, const FString& InFilterList)
 		: FileList(InFileList),
-		FolderList(InFolderList),
-		FilterList(InFilterList)
-	{}
+		FolderList(InFolderList)
+	{
+		// Process the filters once rather than once for each file encountered
+		InFilterList.ParseIntoArray(FilterList, TEXT(";"), true);
+		// Remove cruft from the extension list
+		for (int32 Index = 0; Index < FilterList.Num(); Index++)
+		{
+			FilterList[Index].ReplaceInline(TEXT(")"), TEXT(""));
+			FilterList[Index] = FilterList[Index].TrimQuotes().TrimTrailing().Trim();
+		}
+	}
 	
 
 	virtual bool Visit(const TCHAR* FilenameOrDirectory, bool bIsDirectory) override
@@ -49,7 +57,9 @@ public:
 		for (i = FCString::Strlen(FilenameOrDirectory) - 1; i >= 0; i--)
 		{
 			if (FilenameOrDirectory[i] == TCHAR('/'))
+			{
 				break;
+			}
 		}
 
 #if HIDE_HIDDEN_FILES
@@ -119,77 +129,29 @@ public:
 		return true;
 	}
 
-	bool PassesFilterTest(const TCHAR *Filename)
+	bool PassesFilterTest(const TCHAR* Filename)
 	{		
-		if (!FilterList.IsValid())
+		if (FilterList.Num() == 0)
 		{
 			return true; // no filters. everything passes.
 		}
 
-		int32 i;
-		const TCHAR *Extension = NULL;
-
-		// get extension
-		for (i = FCString::Strlen(Filename) - 1; i >= 0; i--)
+		FString Extension = FPaths::GetExtension(FString(Filename), true);
+		// See if it matches any of the extensions
+		for (const FString& FilterExt : FilterList)
 		{
-			if (Filename[i] == TCHAR('.'))
+			if (FilterExt == TEXT("*") || FilterExt == TEXT(".*") || FilterExt == TEXT("*.*") || FilterExt.EndsWith(Extension))
 			{
-				Extension = &Filename[i];
-				break;
+				return true;
 			}
 		}
-
-		if (Extension == NULL)
-		{
-			return false; // file has no extension. it fails filter test.
-		}
-
-		TCHAR Temp[MAX_FILTER_LENGTH];
-		FCString::Strcpy(Temp, FilterList->Len(), *(*FilterList.Get()));
-
-		// break path into tokens
-		TCHAR *ContextStr = nullptr;
-
-		TCHAR *FilterExt = FCString::Strtok(Temp, TEXT(";"), &ContextStr);
-		bool RC = false;
-
-		while (FilterExt)
-		{
-			// strip leading spaces and '*' from beginning.
-			FilterExt = FCString::Strchr(FilterExt, '.');
-
-			// strip any trailing spaces
-			for (i = FCString::Strlen(FilterExt) - 1; i >= 0 && FilterExt[i] == TCHAR(' '); i--)
-			{
-				FilterExt[i] = 0;
-			}
-
-			if (FCString::Strcmp(FilterExt, TEXT(".*")) == 0)
-			{
-				// *.* matches all
-				RC = true;
-				break;
-			}
-
-			if (FCString::Stricmp(FilterExt, Extension) == 0)
-			{
-				// positive hit.
-				RC = true;
-				break;
-			}
-
-			// next filter entry
-			FilterExt = FCString::Strtok(nullptr, TEXT(";"), &ContextStr);
-		}
-
-		// cleanup and return failed.
-		return RC;
+		return false;
 	}
 
 private:
 	TArray<TSharedPtr<FFileEntry>>& FileList;
 	TArray<TSharedPtr<FFileEntry>>& FolderList;
-	TSharedPtr<FString> FilterList;
+	TArray<FString> FilterList;
 };
 
 
@@ -1250,7 +1212,7 @@ void SSlateFileOpenDlg::ReadDir(bool bIsRefresh)
 
 	FilesArray.Empty();
 	FoldersArray.Empty();
-	TSharedPtr<FString> FilterList = nullptr;
+	FString FilterList;
 
 	if (FilterListArray.Num() > 0 && FilterIndex >= 0)
 	{
@@ -1454,7 +1416,8 @@ void SSlateFileOpenDlg::OnFileNameCommitted(const FText& InText, ETextCommit::Ty
 		if (!bDirectoriesOnly && GetFilterExtension(Extension))
 		{
 			// append extension to filename if user left it off
-			if (!SaveFilename.EndsWith(Extension, ESearchCase::CaseSensitive) && Extension.Compare(".*") != 0)
+			if (!SaveFilename.EndsWith(Extension, ESearchCase::CaseSensitive) &&
+				!IsWildcardExtension(Extension))
 			{
 				SaveFilename = SaveFilename + Extension;
 			}
@@ -1525,7 +1488,7 @@ void SSlateFileOpenDlg::ParseFilters()
 					FilterList = FCString::Strtok(nullptr, TEXT("|"), &ContextStr);
 
 					FilterNameArray.Add(MakeShareable(new FString(FilterDescription)));
-					FilterListArray.Add(MakeShareable(new FString(FilterList)));
+					FilterListArray.Add(FString(FilterList));
 
 					// next filter entry
 					FilterDescription = FCString::Strtok(nullptr, TEXT("|"), &ContextStr);
@@ -1545,10 +1508,11 @@ void SSlateFileOpenDlg::ParseFilters()
 
 bool SSlateFileOpenDlg::GetFilterExtension(FString &OutString)
 {
+	OutString.Empty();
+
 	// check to see if filters were given
 	if (Filters.Len() == 0)
 	{
-		OutString = "";
 		return false;
 	}
 
@@ -1558,24 +1522,34 @@ bool SSlateFileOpenDlg::GetFilterExtension(FString &OutString)
 
 	// find start of extension
 	TCHAR *FilterExt = FCString::Strchr(Temp, '.');
-
-	// strip any trailing junk
-	int32 i;
-	for (i = 0; i < FCString::Strlen(FilterExt); i++)
+	if (FilterExt != nullptr)
 	{
-		if (FilterExt[i] == ' ' || FilterExt[i] == ')' || FilterExt[i] == ';')
+		// strip any trailing junk
+		int32 i;
+		for (i = 0; i < FCString::Strlen(FilterExt); i++)
 		{
-			FilterExt[i] = 0;
-			break;
+			if (FilterExt[i] == ' ' || FilterExt[i] == ')' || FilterExt[i] == ';')
+			{
+				FilterExt[i] = 0;
+				break;
+			}
 		}
-	}
 
-	// store result and clean up
-	OutString = FString(FilterExt);
-
-	return (i > 0);
+		// store result and clean up
+		OutString = FilterExt;
+	}
+	else if (Temp[0] == TEXT('*'))
+	{
+		OutString = Temp;
+	}
+	return !OutString.IsEmpty();
 }
 
+bool SSlateFileOpenDlg::IsWildcardExtension(const FString& Extension)
+{
+	return (Extension.Find(TEXT(".*")) >= 0) ||
+			(Extension.Find(TEXT("*")) >= 0);
+}
 
 void SSlateFileOpenDlg::OnNewDirectoryCommitted(const FText & InText, ETextCommit::Type InCommitType)
 {
diff --git a/Engine/Source/Developer/SlateFileDialogs/Private/SlateFileDlgWindow.h b/Engine/Source/Developer/SlateFileDialogs/Private/SlateFileDlgWindow.h
--- a/Engine/Source/Developer/SlateFileDialogs/Private/SlateFileDlgWindow.h
+++ b/Engine/Source/Developer/SlateFileDialogs/Private/SlateFileDlgWindow.h
@@ -182,6 +182,8 @@ private:
 
 	bool GetFilterExtension(FString &OutString);
 	void ParseFilters();
+	/** @return true if the extension filter contains a wildcard or not */
+	bool IsWildcardExtension(const FString& Extension);
 
 	TArray< FDirNode > DirectoryNodesArray;
 	TArray<TSharedPtr<FFileEntry>> FoldersArray;
@@ -202,7 +204,7 @@ private:
 	TSharedPtr<STextBlock> DirErrorMsg;
 
 	TArray<TSharedPtr<FString>> FilterNameArray;
-	TArray<TSharedPtr<FString>> FilterListArray;
+	TArray<FString> FilterListArray;
 
 	int32 FilterIndex;
 
diff --git a/Engine/Source/Developer/Windows/ShaderFormatD3D/Private/D3D11ShaderCompiler.cpp b/Engine/Source/Developer/Windows/ShaderFormatD3D/Private/D3D11ShaderCompiler.cpp
--- a/Engine/Source/Developer/Windows/ShaderFormatD3D/Private/D3D11ShaderCompiler.cpp
+++ b/Engine/Source/Developer/Windows/ShaderFormatD3D/Private/D3D11ShaderCompiler.cpp
@@ -167,7 +167,7 @@ static FString D3D11CreateShaderCompileCommandLine(
 	)
 {
 	// fxc is our command line compiler
-	FString FXCCommandline = FString(TEXT("\"%DXSDK_DIR%\\Utilities\\bin\\x86\\fxc\" ")) + ShaderPath;
+	FString FXCCommandline = FString(TEXT("%FXC% ")) + ShaderPath;
 
 	// add the entry point reference
 	FXCCommandline += FString(TEXT(" /E ")) + EntryFunction;
@@ -258,7 +258,29 @@ static FString D3D11CreateShaderCompileCommandLine(
 
 	// add a pause on a newline
 	FXCCommandline += FString(TEXT(" \r\n pause"));
-	return FXCCommandline;
+
+	// Batch file header:
+	/*
+	@ECHO OFF
+		SET FXC="C:\Program Files (x86)\Windows Kits\8.1\bin\x64\fxc.exe"
+		IF EXIST %FXC% (
+			REM
+			) ELSE (
+				ECHO Couldn't find Windows 8.1 SDK, falling back to DXSDK...
+				SET FXC="%DXSDK_DIR%\Utilities\bin\x86\fxc.exe"
+				IF EXIST %FXC% (
+					REM
+					) ELSE (
+						ECHO Couldn't find DXSDK! Exiting...
+						GOTO END
+					)
+			)
+	*/
+	FString BatchFileHeader = TEXT("@ECHO OFF\nSET FXC=\"C:\\Program Files (x86)\\Windows Kits\\8.1\\bin\\x64\\fxc.exe\"\n"\
+		"IF EXIST %FXC% (\nREM\n) ELSE (\nECHO Couldn't find Windows 8.1 SDK, falling back to DXSDK...\n"\
+		"SET FXC=\"%DXSDK_DIR%\\Utilities\\bin\\x86\\fxc.exe\"\nIF EXIST %FXC% (\nREM\n) ELSE (\nECHO Couldn't find DXSDK! Exiting...\n"\
+		"GOTO END\n)\n)\n");
+	return BatchFileHeader + FXCCommandline + TEXT("\n:END\nREM\n");
 }
 
 /** Creates a batch file string to call the AMD shader analyzer. */
@@ -284,8 +306,37 @@ static FString CreateAMDCodeXLCommandLine(
 	return Commandline;
 }
 
-// @return pointer to the D3DCompile function
-pD3DCompile GetD3DCompileFunc(const FString& NewCompilerPath)
+// D3Dcompiler.h has function pointer typedefs for some functions, but not all
+typedef HRESULT(WINAPI *pD3DReflect)
+	(__in_bcount(SrcDataSize) LPCVOID pSrcData,
+	 __in SIZE_T  SrcDataSize,
+	 __in  REFIID pInterface,
+	 __out void** ppReflector);
+
+typedef HRESULT(WINAPI *pD3DStripShader)
+	(__in_bcount(BytecodeLength) LPCVOID pShaderBytecode,
+	 __in SIZE_T     BytecodeLength,
+	 __in UINT       uStripFlags,
+	__out ID3DBlob** ppStrippedBlob);
+
+#define DEFINE_GUID_FOR_CURRENT_COMPILER(name, l, w1, w2, b1, b2, b3, b4, b5, b6, b7, b8) \
+	static const GUID name = { l, w1, w2, { b1, b2,  b3,  b4,  b5,  b6,  b7,  b8 } }
+
+// ShaderReflection IIDs may change between SDK versions if the reflection API changes.
+// Define a GUID below that matches the desired IID for the DLL in CompilerPath. For example,
+// look for IID_ID3D11ShaderReflection in d3d11shader.h for the SDK matching the compiler DLL.
+DEFINE_GUID_FOR_CURRENT_COMPILER(IID_ID3D11ShaderReflectionForCurrentCompiler, 0x8d536ca1, 0x0cca, 0x4956, 0xa8, 0x37, 0x78, 0x69, 0x63, 0x75, 0x55, 0x84);
+
+/**
+ * GetD3DCompilerFuncs - gets function pointers from the dll at NewCompilerPath
+ * @param OutD3DCompile - function pointer for D3DCompile (0 if not found)
+ * @param OutD3DReflect - function pointer for D3DReflect (0 if not found)
+ * @param OutD3DDisassemble - function pointer for D3DDisassemble (0 if not found)
+ * @param OutD3DStripShader - function pointer for D3DStripShader (0 if not found)
+ * @return bool - true if functions were retrieved from NewCompilerPath
+ */
+static bool GetD3DCompilerFuncs(const FString& NewCompilerPath, pD3DCompile* OutD3DCompile,
+	pD3DReflect* OutD3DReflect, pD3DDisassemble* OutD3DDisassemble, pD3DStripShader* OutD3DStripShader)
 {
 	static FString CurrentCompiler;
 	static HMODULE CompilerDLL = 0;
@@ -308,21 +359,33 @@ pD3DCompile GetD3DCompileFunc(const FString& NewCompilerPath)
 		if(!CompilerDLL && NewCompilerPath.Len())
 		{
 			// Couldn't find HLSL compiler in specified path. We fail the first compile.
-			return 0;
-		}	
+			*OutD3DCompile = 0;
+			*OutD3DReflect = 0;
+			*OutD3DDisassemble = 0;
+			*OutD3DStripShader = 0;
+			return false;
+		}
 	}
 
 	if(CompilerDLL)
 	{
 		// from custom folder e.g. "C:/DXWin8/D3DCompiler_44.dll"
-		return (pD3DCompile)(void*)GetProcAddress(CompilerDLL, "D3DCompile");
+		*OutD3DCompile = (pD3DCompile)(void*)GetProcAddress(CompilerDLL, "D3DCompile");
+		*OutD3DReflect = (pD3DReflect)(void*)GetProcAddress(CompilerDLL, "D3DReflect");
+		*OutD3DDisassemble = (pD3DDisassemble)(void*)GetProcAddress(CompilerDLL, "D3DDisassemble");
+		*OutD3DStripShader = (pD3DStripShader)(void*)GetProcAddress(CompilerDLL, "D3DStripShader");
+		return true;
 	}
 
 	// D3D SDK we compiled with (usually D3DCompiler_43.dll from windows folder)
-	return &D3DCompile;
+	*OutD3DCompile = &D3DCompile;
+	*OutD3DReflect = &D3DReflect;
+	*OutD3DDisassemble = &D3DDisassemble;
+	*OutD3DStripShader = &D3DStripShader;
+	return false;
 }
 
-HRESULT D3DCompileWrapper(
+static HRESULT D3DCompileWrapper(
 	pD3DCompile				D3DCompileFunc,
 	bool&					bException,
 	LPCVOID					pSrcData,
@@ -405,7 +468,11 @@ static bool CompileAndProcessD3DShader(FString& PreprocessedShaderSource, const
 	TRefCountPtr<ID3DBlob> Errors;
 
 	HRESULT Result;
-	pD3DCompile D3DCompileFunc = GetD3DCompileFunc(CompilerPath);
+	pD3DCompile D3DCompileFunc;
+	pD3DReflect D3DReflectFunc;
+	pD3DDisassemble D3DDisassembleFunc;
+	pD3DStripShader D3DStripShaderFunc;
+	const bool bCompilerPathFunctionsUsed = GetD3DCompilerFuncs(CompilerPath, &D3DCompileFunc, &D3DReflectFunc, &D3DDisassembleFunc, &D3DStripShaderFunc);
 
 	if (D3DCompileFunc)
 	{
@@ -448,10 +515,10 @@ static bool CompileAndProcessD3DShader(FString& PreprocessedShaderSource, const
 	// Fail the compilation if double operations are being used, since those are not supported on all D3D11 cards
 	if (SUCCEEDED(Result))
 	{
-		if (GD3DCheckForDoubles || GD3DDumpD3DAsmFile)
+		if (D3DDisassembleFunc && (GD3DCheckForDoubles || GD3DDumpD3DAsmFile))
 		{
 			TRefCountPtr<ID3DBlob> Dissasembly;
-			if (SUCCEEDED(D3DDisassemble(Shader->GetBufferPointer(), Shader->GetBufferSize(), 0, "", Dissasembly.GetInitReference())))
+			if (SUCCEEDED(D3DDisassembleFunc(Shader->GetBufferPointer(), Shader->GetBufferSize(), 0, "", Dissasembly.GetInitReference())))
 			{
 				ANSICHAR* DissasemblyString = new ANSICHAR[Dissasembly->GetBufferSize() + 1];
 				FMemory::Memcpy(DissasemblyString, Dissasembly->GetBufferPointer(), Dissasembly->GetBufferSize());
@@ -474,17 +541,23 @@ static bool CompileAndProcessD3DShader(FString& PreprocessedShaderSource, const
 				}
 			}
 		}
+	}
 
-		// Gather reflection information
-		int32 NumInterpolants = 0;
-		TIndirectArray<FString> InterpolantNames;
-		TArray<FString> ShaderInputs;
-		if (SUCCEEDED(Result))
+	// Gather reflection information
+	int32 NumInterpolants = 0;
+	TIndirectArray<FString> InterpolantNames;
+	TArray<FString> ShaderInputs;
+	if (SUCCEEDED(Result))
+	{
+		if (D3DReflectFunc)
 		{
 			Output.bSucceeded = true;
-
 			ID3D11ShaderReflection* Reflector = NULL;
-			Result = D3DReflect(Shader->GetBufferPointer(), Shader->GetBufferSize(), IID_ID3D11ShaderReflection, (void**)&Reflector);
+
+			// IID_ID3D11ShaderReflectionForCurrentCompiler is defined in this file and needs to match the IID from the dll in CompilerPath
+			// if the function pointers from that dll are being used
+			const IID ShaderReflectionInterfaceID = bCompilerPathFunctionsUsed ? IID_ID3D11ShaderReflectionForCurrentCompiler : IID_ID3D11ShaderReflection;
+			Result = D3DReflectFunc(Shader->GetBufferPointer(), Shader->GetBufferSize(), ShaderReflectionInterfaceID, (void**)&Reflector);
 			if (FAILED(Result))
 			{
 				UE_LOG(LogD3D11ShaderCompiler, Fatal, TEXT("D3DReflect failed: Result=%08x"), Result);
@@ -738,13 +811,13 @@ static bool CompileAndProcessD3DShader(FString& PreprocessedShaderSource, const
 			{
 				CompressedData = Shader;
 			}
-			else
+			else if (D3DStripShaderFunc)
 			{
 				// Strip shader reflection and debug info
 				D3D_SHADER_DATA ShaderData;
 				ShaderData.pBytecode = Shader->GetBufferPointer();
 				ShaderData.BytecodeLength = Shader->GetBufferSize();
-				Result = D3DStripShader(Shader->GetBufferPointer(),
+				Result = D3DStripShaderFunc(Shader->GetBufferPointer(),
 					Shader->GetBufferSize(),
 					D3DCOMPILER_STRIP_REFLECTION_DATA | D3DCOMPILER_STRIP_DEBUG_INFO | D3DCOMPILER_STRIP_TEST_BLOBS,
 					CompressedData.GetInitReference());
@@ -754,6 +827,12 @@ static bool CompileAndProcessD3DShader(FString& PreprocessedShaderSource, const
 					UE_LOG(LogD3D11ShaderCompiler, Fatal, TEXT("D3DStripShader failed: Result=%08x"), Result);
 				}
 			}
+			else
+			{
+				// D3DStripShader is not guaranteed to exist
+				// e.g. the open-source DXIL shader compiler does not currently implement it
+				CompressedData = Shader;
+			}
 
 			// Build the SRT for this shader.
 			FD3D11ShaderResourceTable SRT;
@@ -833,7 +912,16 @@ static bool CompileAndProcessD3DShader(FString& PreprocessedShaderSource, const
 			// Pass the target through to the output.
 			Output.Target = Input.Target;
 		}
+		else
+		{
+			FilteredErrors.Add(FString::Printf(TEXT("Couldn't find shader reflection function in %s"), *CompilerPath));
+			Result = E_FAIL;
+			Output.bSucceeded = false;
+		}
+	}
 
+	if (SUCCEEDED(Result))
+	{
 		if (Input.Target.Platform == SP_PCD3D_ES2)
 		{
 			if (Output.NumTextureSamplers > 8)
@@ -858,7 +946,8 @@ static bool CompileAndProcessD3DShader(FString& PreprocessedShaderSource, const
 			}
 		}
 	}
-	else
+
+	if (FAILED(Result))
 	{
 		++GBreakpoint;
 	}
diff --git a/Engine/Source/Editor/AnimGraph/Private/AnimGraphNode_AnimDynamics.cpp b/Engine/Source/Editor/AnimGraph/Private/AnimGraphNode_AnimDynamics.cpp
--- a/Engine/Source/Editor/AnimGraph/Private/AnimGraphNode_AnimDynamics.cpp
+++ b/Engine/Source/Editor/AnimGraph/Private/AnimGraphNode_AnimDynamics.cpp
@@ -131,7 +131,7 @@ void UAnimGraphNode_AnimDynamics::Draw(FPrimitiveDrawInterface* PDI, USkeletalMe
 						SphereTransform *= PreviewSkelMeshComp->GetComponentSpaceTransforms()[DrivingBoneIdx];
 					}
 
-					DrawSphere(PDI, SphereTransform.GetLocation(), FVector(SphericalLimit.LimitRadius), 24, 6, GEngine->ConstraintLimitMaterialY->GetRenderProxy(false), SDPG_World);
+					DrawSphere(PDI, SphereTransform.GetLocation(), FRotator::ZeroRotator, FVector(SphericalLimit.LimitRadius), 24, 6, GEngine->ConstraintLimitMaterialY->GetRenderProxy(false), SDPG_World);
 					DrawWireSphere(PDI, SphereTransform, FLinearColor::Black, SphericalLimit.LimitRadius, 24, SDPG_World);
 				}
 			}
diff --git a/Engine/Source/Editor/AnimGraph/Private/EditModes/SplineIKEditMode.cpp b/Engine/Source/Editor/AnimGraph/Private/EditModes/SplineIKEditMode.cpp
--- a/Engine/Source/Editor/AnimGraph/Private/EditModes/SplineIKEditMode.cpp
+++ b/Engine/Source/Editor/AnimGraph/Private/EditModes/SplineIKEditMode.cpp
@@ -69,7 +69,7 @@ void FSplineIKEditMode::Render(const FSceneView* View, FViewport* Viewport, FPri
 		PDI->SetHitProxy(new HSplineHandleHitProxy(SplineHandleIndex));
 		FTransform StartTransform = SplineIKRuntimeNode->GetTransformedSplinePoint(SplineHandleIndex);
 		const float Scale = View->WorldToScreen(StartTransform.GetLocation()).W * (4.0f / View->UnscaledViewRect.Width() / View->ViewMatrices.GetProjectionMatrix().M[0][0]);
-		DrawSphere(PDI, StartTransform.GetLocation(), FVector(4.0f) * Scale, 64, 64, GEngine->ArrowMaterial->GetRenderProxy(SelectedSplinePoint == SplineHandleIndex), SDPG_Foreground);
+		DrawSphere(PDI, StartTransform.GetLocation(), FRotator::ZeroRotator, FVector(4.0f) * Scale, 64, 64, GEngine->ArrowMaterial->GetRenderProxy(SelectedSplinePoint == SplineHandleIndex), SDPG_Foreground);
 		DrawCoordinateSystem(PDI, StartTransform.GetLocation(), StartTransform.GetRotation().Rotator(), 30.0f * Scale, SDPG_Foreground);
 	}
 
diff --git a/Engine/Source/Editor/Cascade/Private/Cascade.cpp b/Engine/Source/Editor/Cascade/Private/Cascade.cpp
--- a/Engine/Source/Editor/Cascade/Private/Cascade.cpp
+++ b/Engine/Source/Editor/Cascade/Private/Cascade.cpp
@@ -3445,6 +3445,9 @@ void FCascade::AddLOD(bool bBeforeCurrent)
 			}
 		}
 
+		//This should probably have fixed size and behave like LODDistances but for now just avoid the crash.
+		ParticleSystem->LODSettings.SetNumZeroed(FMath::Max(CurrentLODIndex, ParticleSystem->LODSettings.Num()));
+
 		ParticleSystem->LODDistances.InsertZeroed(CurrentLODIndex, 1);
 		if (CurrentLODIndex == 0)
 		{
diff --git a/Engine/Source/Editor/ContentBrowser/Private/SAssetDialog.cpp b/Engine/Source/Editor/ContentBrowser/Private/SAssetDialog.cpp
--- a/Engine/Source/Editor/ContentBrowser/Private/SAssetDialog.cpp
+++ b/Engine/Source/Editor/ContentBrowser/Private/SAssetDialog.cpp
@@ -124,6 +124,7 @@ void SAssetDialog::Construct(const FArguments& InArgs, const FSharedAssetDialogC
 	PathPicker = StaticCastSharedRef<SPathPicker>(FContentBrowserSingleton::Get().CreatePathPicker(PathPickerConfig));
 	AssetPicker = StaticCastSharedRef<SAssetPicker>(FContentBrowserSingleton::Get().CreateAssetPicker(AssetPickerConfig));
 
+	FContentBrowserCommands::Register();
 	BindCommands();
 
 	// The root widget in this dialog.
diff --git a/Engine/Source/Editor/LevelEditor/Private/LevelEditor.cpp b/Engine/Source/Editor/LevelEditor/Private/LevelEditor.cpp
--- a/Engine/Source/Editor/LevelEditor/Private/LevelEditor.cpp
+++ b/Engine/Source/Editor/LevelEditor/Private/LevelEditor.cpp
@@ -619,7 +619,8 @@ void FLevelEditorModule::BindGlobalLevelEditorCommands()
 		FCanExecuteAction::CreateStatic( &FLevelEditorActionCallbacks::ExportSelected_CanExecute ) );
 
 	ActionList.MapAction( Commands.Build,
-		FExecuteAction::CreateStatic( &FLevelEditorActionCallbacks::Build_Execute ) );
+		FExecuteAction::CreateStatic( &FLevelEditorActionCallbacks::Build_Execute ),
+		FCanExecuteAction::CreateStatic( &FLevelEditorActionCallbacks::Build_CanExecute ) );
 
 	ActionList.MapAction(
 		Commands.ConnectToSourceControl,
diff --git a/Engine/Source/Editor/LevelEditor/Private/LevelEditorActions.cpp b/Engine/Source/Editor/LevelEditor/Private/LevelEditorActions.cpp
--- a/Engine/Source/Editor/LevelEditor/Private/LevelEditorActions.cpp
+++ b/Engine/Source/Editor/LevelEditor/Private/LevelEditorActions.cpp
@@ -631,6 +631,10 @@ void FLevelEditorActionCallbacks::Build_Execute()
 	FEditorBuildUtils::EditorBuild( GetWorld(), FBuildOptions::BuildAll );
 }
 
+bool FLevelEditorActionCallbacks::Build_CanExecute()
+{
+	return !(GEditor->PlayWorld || GUnrealEd->bIsSimulatingInEditor);
+}
 
 void FLevelEditorActionCallbacks::BuildAndSubmitToSourceControl_Execute()
 {
@@ -653,7 +657,7 @@ bool FLevelEditorActionCallbacks::BuildLighting_CanExecute()
 {
 	static const auto AllowStaticLightingVar = IConsoleManager::Get().FindTConsoleVariableDataInt(TEXT("r.AllowStaticLighting"));
 	const bool bAllowStaticLighting = (!AllowStaticLightingVar || AllowStaticLightingVar->GetValueOnGameThread() != 0);
-	return bAllowStaticLighting;
+	return bAllowStaticLighting && !(GEditor->PlayWorld || GUnrealEd->bIsSimulatingInEditor);
 }
 
 void FLevelEditorActionCallbacks::BuildReflectionCapturesOnly_Execute()
diff --git a/Engine/Source/Editor/LevelEditor/Public/LevelEditorActions.h b/Engine/Source/Editor/LevelEditor/Public/LevelEditorActions.h
--- a/Engine/Source/Editor/LevelEditor/Public/LevelEditorActions.h
+++ b/Engine/Source/Editor/LevelEditor/Public/LevelEditorActions.h
@@ -727,6 +727,7 @@ public:
 	 * Build callbacks
 	 */
 	static void Build_Execute();
+	static bool Build_CanExecute();
 	static void BuildAndSubmitToSourceControl_Execute();
 	static void BuildLightingOnly_Execute();
 	static bool BuildLighting_CanExecute();
